ModelingTools: add Edge Bevel operation in PolyEdit Tool. New UPolyEditBevelEdgeActivity, and operations to initialize it in PolyEdit. This bevel operation currently uses an explicit bevel-distance rather than a live viewport interaction, and as a result required adding a separate 'apply' button inside the PolyEditTool panel (via new UEditMeshPolygonsToolAcceptCancelAction). The UX here is terrible but it's not clear we can do better in 5.0 timeframe.

GeometryProcessing: add support for beveling the border ring(s) of a triangle or polygroup-face selection. Improve polygroup assignment around bevel loops, previously the entire loop received a single polygroup which is not what the user expects in polygroup-modeling. Now a new polygroup is assigned for each unique group-edge-pair. Terminator vertices now also generate a new polygroup if there are 2 or more adjacent groups, previously behavior (assigning to arbitrary nbr) was not useful during polygroup-edting.

#rnx

diff --git a/Engine/Plugins/Runtime/GeometryProcessing/Source/DynamicMesh/Private/Operations/MeshBevel.cpp b/Engine/Plugins/Runtime/GeometryProcessing/Source/DynamicMesh/Private/Operations/MeshBevel.cpp
--- a/Engine/Plugins/Runtime/GeometryProcessing/Source/DynamicMesh/Private/Operations/MeshBevel.cpp
+++ b/Engine/Plugins/Runtime/GeometryProcessing/Source/DynamicMesh/Private/Operations/MeshBevel.cpp
@@ -9,6 +9,7 @@
 #include "CompGeom/PolygonTriangulation.h"
 #include "DynamicMesh/MeshNormals.h"
 #include "DynamicMesh/MeshIndexUtil.h"
+#include "MeshRegionBoundaryLoops.h"
 #include "Operations/PolyEditingEdgeUtil.h"
 #include "Operations/PolyEditingUVUtil.h"
 #include "Algo/Count.h"
@@ -16,9 +17,6 @@
 
 using namespace UE::Geometry;
 
-
-
-
 namespace UELocal
 {
 	static void QuadsToTris(const FDynamicMesh3& Mesh, const TArray<FIndex2i>& Quads, TArray<int32>& TrisOut, bool bReset)
@@ -76,7 +74,7 @@ void FMeshBevel::InitializeFromGroupTopology(const FDynamicMesh3& Mesh, const FG
 }
 
 
-void FMeshBevel::InitializeFromGroupTopology(const FDynamicMesh3& Mesh, const FGroupTopology& Topology, const TArray<int32>& GroupEdges)
+void FMeshBevel::InitializeFromGroupTopologyEdges(const FDynamicMesh3& Mesh, const FGroupTopology& Topology, const TArray<int32>& GroupEdges)
 {
 	ResultInfo = FGeometryResult(EGeometryResultType::InProgress);
 
@@ -105,6 +103,53 @@ void FMeshBevel::InitializeFromGroupTopology(const FDynamicMesh3& Mesh, const FG
 }
 
 
+bool FMeshBevel::InitializeFromGroupTopologyFaces(const FDynamicMesh3& Mesh, const FGroupTopology& Topology, const TArray<int32>& GroupFaces)
+{
+	FGroupTopologySelection Selection;
+	Selection.SelectedGroupIDs.Append(GroupFaces);
+	TArray<int32> Triangles;
+	Topology.GetSelectedTriangles(Selection, Triangles);
+	return InitializeFromTriangleSet(Mesh, Triangles);
+}
+
+bool FMeshBevel::InitializeFromTriangleSet(const FDynamicMesh3& Mesh, const TArray<int32>& Triangles)
+{
+	ResultInfo = FGeometryResult(EGeometryResultType::InProgress);
+
+	FMeshRegionBoundaryLoops RegionLoops(&Mesh, Triangles, true);
+	if (RegionLoops.bFailed)
+	{
+		ResultInfo.SetFailed();
+		return false;
+	}
+
+	// cannot bevel a selection-bowtie vertex, so we have to check for those and fail here
+	TSet<int32> AllVertices;
+	for (const FEdgeLoop& Loop : RegionLoops.Loops)
+	{
+		for (int32 vid : Loop.Vertices)
+		{
+			if (AllVertices.Contains(vid))
+			{
+				return false;
+			}
+			AllVertices.Add(vid);
+		}
+	}
+
+	for (const FEdgeLoop& Loop : RegionLoops.Loops)
+	{
+		AddBevelEdgeLoop(Mesh, Loop);
+	}
+
+	// precompute topological information necessary to apply bevel to vertices/edges/loops
+	BuildVertexSets(Mesh);
+
+	return true;
+}
+
+
+
 bool FMeshBevel::Apply(FDynamicMesh3& Mesh, FDynamicMeshChangeTracker* ChangeTracker)
 {
 	// disconnect along bevel graph edges/vertices and save necessary info
@@ -473,7 +518,7 @@ void FMeshBevel::BuildTerminatorVertex(FBevelVertex& Vertex, const FDynamicMesh3
 			if (Mesh.GetTriangleGroup(OtherGroupTris[k-1]) != Mesh.GetTriangleGroup(OtherGroupTris[k]))
 			{
 				RingSplitEdgeID = FindSharedEdgeInTriangles(Mesh, OtherGroupTris[k-1], OtherGroupTris[k]);
-				Vertex.NewGroupID = Mesh.GetTriangleGroup(OtherGroupTris[k-1]);		// this seems very arbitrary
+				Vertex.NewGroupID = -1;		// allocate a new group for this triangle, this is usually what one would want
 				break;
 			}
 		}
@@ -1160,7 +1205,24 @@ void FMeshBevel::AppendLoopQuads(FDynamicMesh3& Mesh, FBevelLoop& Loop)
 		return;
 	}
 
-	Loop.NewGroupID = Mesh.AllocateTriangleGroup();
+	auto GetGroupKey = [&Mesh, &Loop](int32 k)
+	{
+		FIndex2i EdgeTris = Loop.MeshEdgeTris[k];
+		int32 Group0 = Mesh.GetTriangleGroup(EdgeTris.A);
+		int32 Group1 = Mesh.IsTriangle(EdgeTris.B) ? Mesh.GetTriangleGroup(EdgeTris.B) : -1;
+		return FIndex2i(FMath::Max(Group0, Group1), FMath::Min(Group0, Group1));
+	};
+
+	TMap<FIndex2i, int> NewGroupIDs;
+	for (int32 k = 0; k < NumEdges; ++k)
+	{
+		FIndex2i GroupKey = GetGroupKey(k);
+		if (NewGroupIDs.Contains(GroupKey) == false)
+		{
+			NewGroupIDs.Add(GroupKey, Mesh.AllocateTriangleGroup());
+			Loop.NewGroupIDs.Add(NewGroupIDs[GroupKey]);
+		}
+	}
 
 	// At this point each edge-span should be fully disconnected into a set of paired edges, 
 	// so we can trivially join each edge pair with a quad.
@@ -1174,10 +1236,13 @@ void FMeshBevel::AppendLoopQuads(FDynamicMesh3& Mesh, FBevelLoop& Loop)
 		FIndex2i QuadTris(IndexConstants::InvalidID, IndexConstants::InvalidID);
 		if (EdgeID0 != EdgeID1 && Mesh.IsEdge(EdgeID1))
 		{
+			FIndex2i GroupKey = GetGroupKey(k);
+			int32 NewGroupID = NewGroupIDs[GroupKey];
+
 			FIndex2i EdgeV0 = Mesh.GetOrientedBoundaryEdgeV(EdgeID0);
 			FIndex2i EdgeV1 = Mesh.GetOrientedBoundaryEdgeV(EdgeID1);
-			QuadTris.A = Mesh.AppendTriangle(EdgeV0.B, EdgeV0.A, EdgeV1.B, Loop.NewGroupID);
-			QuadTris.B = Mesh.AppendTriangle(EdgeV1.B, EdgeV1.A, EdgeV0.B, Loop.NewGroupID);
+			QuadTris.A = Mesh.AppendTriangle(EdgeV0.B, EdgeV0.A, EdgeV1.B, NewGroupID);
+			QuadTris.B = Mesh.AppendTriangle(EdgeV1.B, EdgeV1.A, EdgeV0.B, NewGroupID);
 		}
 		Loop.StripQuads.Add(QuadTris);
 	}
@@ -1451,4 +1516,4 @@ void FMeshBevel::ComputeMaterialIDs(FDynamicMesh3& Mesh)
 		}
 	}
 
-}
\ No newline at end of file
+}
diff --git a/Engine/Plugins/Runtime/GeometryProcessing/Source/DynamicMesh/Public/Operations/MeshBevel.h b/Engine/Plugins/Runtime/GeometryProcessing/Source/DynamicMesh/Public/Operations/MeshBevel.h
--- a/Engine/Plugins/Runtime/GeometryProcessing/Source/DynamicMesh/Public/Operations/MeshBevel.h
+++ b/Engine/Plugins/Runtime/GeometryProcessing/Source/DynamicMesh/Public/Operations/MeshBevel.h
@@ -88,7 +88,19 @@ public:
 	/**
 	 * Initialize the bevel with the specified edges of a GroupTopology
 	 */
-	void InitializeFromGroupTopology(const FDynamicMesh3& Mesh, const FGroupTopology& Topology, const TArray<int32>& GroupEdges);
+	void InitializeFromGroupTopologyEdges(const FDynamicMesh3& Mesh, const FGroupTopology& Topology, const TArray<int32>& GroupEdges);
+
+	/**
+	* Initialize the bevel with the specified faces of a GroupTopology
+	* @return false if any selection-bowtie vertices were found, in this case we cannot compute the bevel
+	*/
+	bool InitializeFromGroupTopologyFaces(const FDynamicMesh3& Mesh, const FGroupTopology& Topology, const TArray<int32>& GroupFaces);
+
+	/**
+	* Initialize the bevel with border loops of the selected triangles. 
+	* @return false if any selection-bowtie vertices were found, in this case we cannot compute the bevel
+	*/
+	bool InitializeFromTriangleSet(const FDynamicMesh3& Mesh, const TArray<int32>& Triangles);
 
 	/**
 	 * Apply the bevel operation to the mesh, and optionally track changes
@@ -132,7 +144,7 @@ public:
 		TArray<FVector3d> NewPositions1;	// new positions for NewMeshVertices list
 
 		// new geometry computed during mesh step
-		int32 NewGroupID;
+		TArray<int32> NewGroupIDs;
 		TArray<FIndex2i> StripQuads;		// triangle-ID-pairs for each new quad added along edge, 1-1 with MeshEdges
 	};
 
diff --git a/Engine/Plugins/Runtime/MeshModelingToolset/Source/MeshModelingTools/Private/EditMeshPolygonsTool.cpp b/Engine/Plugins/Runtime/MeshModelingToolset/Source/MeshModelingTools/Private/EditMeshPolygonsTool.cpp
--- a/Engine/Plugins/Runtime/MeshModelingToolset/Source/MeshModelingTools/Private/EditMeshPolygonsTool.cpp
+++ b/Engine/Plugins/Runtime/MeshModelingToolset/Source/MeshModelingTools/Private/EditMeshPolygonsTool.cpp
@@ -42,6 +42,7 @@
 #include "ToolActivities/PolyEditInsetOutsetActivity.h"
 #include "ToolActivities/PolyEditCutFacesActivity.h"
 #include "ToolActivities/PolyEditPlanarProjectionUVActivity.h"
+#include "ToolActivities/PolyEditBevelEdgeActivity.h"
 #include "ToolContextInterfaces.h" // FToolBuilderState
 #include "ToolSetupUtil.h"
 #include "ToolTargetManager.h"
@@ -223,6 +224,12 @@ void UEditMeshPolygonsTool::Setup()
 	AddToolPropertySource(CancelAction);
 	SetToolPropertySourceEnabled(CancelAction, false);
 
+	AcceptCancelAction = NewObject<UEditMeshPolygonsToolAcceptCancelAction>();
+	AcceptCancelAction->Initialize(this);
+	AddToolPropertySource(AcceptCancelAction);
+	SetToolPropertySourceEnabled(AcceptCancelAction, false);
+
+
 	// Initialize the common properties but don't add them yet, because we want them to be under the activity-specific ones.
 	CommonProps = NewObject<UPolyEditCommonProperties>(this);
 	CommonProps->RestoreProperties(this);
@@ -359,6 +366,9 @@ void UEditMeshPolygonsTool::Setup()
 	InsertEdgeActivity = NewObject<UPolyEditInsertEdgeActivity>();
 	InsertEdgeActivity->Setup(this);
 
+	BevelEdgeActivity = NewObject<UPolyEditBevelEdgeActivity>();
+	BevelEdgeActivity->Setup(this);
+
 	// Now that we've initialized the activities, add in the selection settings and 
 	// CommonProps so that they are at the bottom.
 	AddToolPropertySource(SelectionMechanic->Properties);
@@ -449,6 +459,7 @@ void UEditMeshPolygonsTool::Shutdown(EToolShutdownType ShutdownType)
 	PlanarProjectionUVActivity->Shutdown(ShutdownType);
 	InsertEdgeActivity->Shutdown(ShutdownType);
 	InsertEdgeLoopActivity->Shutdown(ShutdownType);
+	BevelEdgeActivity->Shutdown(ShutdownType);
 
 	GetToolManager()->GetPairedGizmoManager()->DestroyAllGizmosByOwner(this);
 
@@ -521,6 +532,7 @@ void UEditMeshPolygonsTool::Shutdown(EToolShutdownType ShutdownType)
 	EditEdgeActions_Triangles = nullptr;
 	EditUVActions = nullptr;
 	CancelAction = nullptr;
+	AcceptCancelAction = nullptr;
 
 	ExtrudeActivity = nullptr;
 	InsetOutsetActivity = nullptr;
@@ -528,6 +540,7 @@ void UEditMeshPolygonsTool::Shutdown(EToolShutdownType ShutdownType)
 	PlanarProjectionUVActivity = nullptr;
 	InsertEdgeActivity = nullptr;
 	InsertEdgeLoopActivity = nullptr;
+	BevelEdgeActivity = nullptr;
 
 	SelectionMechanic = nullptr;
 	DragAlignmentMechanic = nullptr;
@@ -893,9 +906,21 @@ void UEditMeshPolygonsTool::OnTick(float DeltaTime)
 			break;
 		}
 
+		case EEditMeshPolygonsToolActions::BevelFaces:
+		case EEditMeshPolygonsToolActions::BevelEdges:
+		{
+			StartActivity(BevelEdgeActivity);
+			break;
+		}
+
 		case EEditMeshPolygonsToolActions::CancelCurrent:
 		{
-			EndCurrentActivity();
+			EndCurrentActivity(EToolShutdownType::Cancel);
+			break;
+		}
+		case EEditMeshPolygonsToolActions::AcceptCurrent:
+		{
+			EndCurrentActivity(EToolShutdownType::Accept);
 			break;
 		}
 
@@ -969,7 +994,14 @@ void UEditMeshPolygonsTool::StartActivity(TObjectPtr<UInteractiveToolActivity> A
 		TransformGizmo->SetVisibility(false);
 		SelectionMechanic->SetIsEnabled(false);
 		CurrentActivity = Activity;
-		SetToolPropertySourceEnabled(CancelAction, true);
+		if (CurrentActivity == BevelEdgeActivity)
+		{
+			SetToolPropertySourceEnabled(AcceptCancelAction, true);
+		}
+		else
+		{
+			SetToolPropertySourceEnabled(CancelAction, true);
+		}
 		SetActionButtonPanelsVisible(false);
 	}
 }
@@ -991,6 +1023,7 @@ void UEditMeshPolygonsTool::EndCurrentActivity(EToolShutdownType ShutdownType)
 		++ActivityTimestamp;
 
 		SetToolPropertySourceEnabled(CancelAction, false);
+		SetToolPropertySourceEnabled(AcceptCancelAction, false);
 		SetActionButtonPanelsVisible(true);
 		SelectionMechanic->SetIsEnabled(true);
 		UpdateGizmoVisibility();
diff --git a/Engine/Plugins/Runtime/MeshModelingToolset/Source/MeshModelingTools/Private/ToolActivities/PolyEditBevelEdgeActivity.cpp b/Engine/Plugins/Runtime/MeshModelingToolset/Source/MeshModelingTools/Private/ToolActivities/PolyEditBevelEdgeActivity.cpp
new file mode 100644
--- /dev/null
+++ b/Engine/Plugins/Runtime/MeshModelingToolset/Source/MeshModelingTools/Private/ToolActivities/PolyEditBevelEdgeActivity.cpp
@@ -0,0 +1,305 @@
+// Copyright Epic Games, Inc. All Rights Reserved.
+
+#include "ToolActivities/PolyEditBevelEdgeActivity.h"
+
+#include "InteractiveToolManager.h"
+#include "ContextObjectStore.h"
+#include "EditMeshPolygonsTool.h"
+#include "MeshOpPreviewHelpers.h"
+#include "Operations/MeshBevel.h"
+#include "Selection/PolygonSelectionMechanic.h"
+#include "ToolActivities/PolyEditActivityContext.h"
+
+#define LOCTEXT_NAMESPACE "UPolyEditBevelEdgeActivity"
+
+using namespace UE::Geometry;
+
+
+
+class FBevelOp : public FDynamicMeshOperator
+{
+public:
+	virtual ~FBevelOp() {}
+
+	// Inputs:
+	TSharedPtr<const FDynamicMesh3, ESPMode::ThreadSafe> OriginalMesh;
+	TSharedPtr<UE::Geometry::FGroupTopology, ESPMode::ThreadSafe> MeshTopology;
+	TArray<int32> BevelGroupEdges;
+	TArray<int32> BevelGroupFaces;
+	double BevelDistance = 1.0;
+
+	void SetTransform(const FTransform& Transform)
+	{
+		ResultTransform = (UE::Geometry::FTransform3d)Transform;
+	}
+
+	// FDynamicMeshOperator implementation 
+	virtual void CalculateResult(FProgressCancel* Progress) override
+	{
+		if (FMath::Abs(BevelDistance) < KINDA_SMALL_NUMBER || (Progress && Progress->Cancelled()))
+		{
+			return;
+		}
+		ResultMesh->Copy(*OriginalMesh, true, true, true, true);
+
+		UE::Geometry::FMeshBevel Bevel;
+		Bevel.InsetDistance = BevelDistance;
+		Bevel.MaterialIDMode = FMeshBevel::EMaterialIDMode::InferMaterialID;
+		Bevel.SetConstantMaterialID = 0;
+		bool bBevelSetupValid = false;
+		if (BevelGroupEdges.Num() > 0)
+		{
+			Bevel.InitializeFromGroupTopologyEdges(*ResultMesh, *MeshTopology, BevelGroupEdges);
+			bBevelSetupValid = true;
+		}
+		else if (BevelGroupFaces.Num() > 0)
+		{
+			bBevelSetupValid = Bevel.InitializeFromGroupTopologyFaces(*ResultMesh, *MeshTopology, BevelGroupFaces);
+		}
+		else
+		{
+			Bevel.InitializeFromGroupTopology(*ResultMesh, *MeshTopology);
+			bBevelSetupValid = true;
+		}
+		if (bBevelSetupValid)
+		{
+			Bevel.Apply(*ResultMesh, nullptr);
+			SetResultInfo(FGeometryResult());
+		}
+		else
+		{
+			SetResultInfo(FGeometryResult::Failed());
+		}
+	}
+
+};
+
+
+TUniquePtr<FDynamicMeshOperator> UPolyEditBevelEdgeActivity::MakeNewOperator()
+{
+	// create operator here
+	TUniquePtr<FBevelOp> Op = MakeUnique<FBevelOp>();
+	Op->OriginalMesh = ActivityContext->CurrentMesh;
+	Op->MeshTopology = ActivityContext->CurrentTopology;
+	Op->BevelDistance = BevelProperties->BevelDistance;
+
+	// copy edge selection
+	if (ActiveSelection.SelectedEdgeIDs.Num() > 0)
+	{
+		for (int32 eid : ActiveSelection.SelectedEdgeIDs)
+		{
+			Op->BevelGroupEdges.Add(eid);
+		}
+	}
+	else if (ActiveSelection.SelectedGroupIDs.Num() > 0)
+	{
+		for (int32 gid : ActiveSelection.SelectedGroupIDs)
+		{
+			Op->BevelGroupFaces.Add(gid);
+		}
+	}
+	else
+	{
+		UE_LOG(LogGeometry, Warning, TEXT("UPolyEditBevelEdgeActivity::MakeNewOperator : empty selection"));
+	}
+
+	UE::Geometry::FTransform3d WorldTransform(ActivityContext->Preview->PreviewMesh->GetTransform());
+	Op->SetResultTransform(WorldTransform);
+	
+	return Op;
+}
+
+void UPolyEditBevelEdgeActivity::Setup(UInteractiveTool* ParentToolIn)
+{
+	Super::Setup(ParentToolIn);
+
+	ActivityContext = ParentTool->GetToolManager()->GetContextObjectStore()->FindContext<UPolyEditActivityContext>();
+
+	BevelProperties = NewObject<UPolyEditBevelEdgeProperties>();
+	BevelProperties->RestoreProperties(ParentTool.Get());
+	AddToolPropertySource(BevelProperties);
+	SetToolPropertySourceEnabled(BevelProperties, false);
+	BevelProperties->WatchProperty(BevelProperties->BevelDistance, [this](float) {  
+		if (bIsRunning) {
+			ActivityContext->Preview->InvalidateResult();
+		}
+	});
+}
+
+void UPolyEditBevelEdgeActivity::Shutdown(EToolShutdownType ShutdownType)
+{
+	if (bIsRunning)
+	{
+		End(ShutdownType);
+	}
+
+	BevelProperties->SaveProperties(ParentTool.Get());
+
+	BevelProperties = nullptr;
+	ParentTool = nullptr;
+	ActivityContext = nullptr;
+}
+
+bool UPolyEditBevelEdgeActivity::CanStart() const
+{
+	if (ActivityContext)
+	{
+		const FGroupTopologySelection& Selection = ActivityContext->SelectionMechanic->GetActiveSelection();
+		return (Selection.SelectedEdgeIDs.Num() > 0) ||  (Selection.SelectedGroupIDs.Num() > 0);
+	}
+	return false;
+}
+
+EToolActivityStartResult UPolyEditBevelEdgeActivity::Start()
+{
+	if (!CanStart())
+	{
+		ParentTool->GetToolManager()->DisplayMessage( LOCTEXT("OnBevelFailedMesssage", "Action requires edge or face selection."), EToolMessageLevel::UserWarning);
+		return EToolActivityStartResult::FailedStart;
+	}
+
+	// Change the op we use in the preview to an Bevel op.
+	ActivityContext->Preview->ChangeOpFactory(this);
+	ActivityContext->Preview->OnOpCompleted.AddWeakLambda(this,
+		[this](const UE::Geometry::FDynamicMeshOperator* UncastOp) {
+			const FBevelOp* Op = static_cast<const FBevelOp*>(UncastOp);
+			FGeometryResult Result = Op->GetResultInfo();
+
+			// If bevel failed, print a message and kill the activity
+			// Warning: if FBevelOp is modified to allow failure based on the bevel parameters, this could incorrectly terminate the operation, and will require more complex handling
+			if (Result.HasFailed())
+			{
+				ParentTool->GetToolManager()->DisplayMessage( LOCTEXT("OnBevelFailedMessage", "Bevel failed for current selection"), EToolMessageLevel::UserWarning);
+				EndInternal();
+				Cast<IToolActivityHost>(ParentTool)->NotifyActivitySelfEnded(this);
+			}
+		});
+
+	SetToolPropertySourceEnabled(BevelProperties, true); 
+
+	BeginBevel();
+
+	bIsRunning = true;
+
+	ActivityContext->EmitActivityStart(LOCTEXT("BeginBevelActivity", "Begin Bevel"));
+	return EToolActivityStartResult::Running;
+}
+
+void UPolyEditBevelEdgeActivity::BeginBevel()
+{
+	ActiveSelection = ActivityContext->SelectionMechanic->GetActiveSelection();
+
+	// Temporarily clear the selection to avoid the highlighting on the preview
+	ActivityContext->SelectionMechanic->SetSelection(FGroupTopologySelection(), false);
+
+	// force initial compute?
+	ActivityContext->Preview->InvalidateResult();
+}
+
+bool UPolyEditBevelEdgeActivity::CanAccept() const
+{
+	return true;
+}
+
+EToolActivityEndResult UPolyEditBevelEdgeActivity::End(EToolShutdownType ShutdownType)
+{
+	if (!ensure(bIsRunning))
+	{
+		EndInternal();
+		return EToolActivityEndResult::ErrorDuringEnd;
+	}
+
+	if (ShutdownType == EToolShutdownType::Cancel)
+	{
+		EndInternal();
+
+		// Reset the preview
+		ActivityContext->Preview->PreviewMesh->UpdatePreview(ActivityContext->CurrentMesh.Get());
+
+		return EToolActivityEndResult::Cancelled;
+	}
+	else
+	{
+		// Stop the current compute if there is one.
+		ActivityContext->Preview->CancelCompute();
+
+		// Apply whatever we happen to have at the moment.
+		ApplyBevel();
+
+		EndInternal();
+		return EToolActivityEndResult::Completed;
+	}
+}
+
+// Does whatever kind of cleanup we need to do to end the activity.
+void UPolyEditBevelEdgeActivity::EndInternal()
+{
+	if (ActivityContext->SelectionMechanic->GetActiveSelection() != ActiveSelection)
+	{
+		// We haven't reset the selection back from when we cleared it to avoid the highlighting.
+		ActivityContext->SelectionMechanic->SetSelection(ActiveSelection, false);
+	}
+
+	ActivityContext->Preview->ClearOpFactory();
+	ActivityContext->Preview->OnOpCompleted.RemoveAll(this);
+	SetToolPropertySourceEnabled(BevelProperties, false);
+	bIsRunning = false;
+}
+
+void UPolyEditBevelEdgeActivity::ApplyBevel()
+{
+	const FDynamicMesh3* ResultMesh = ActivityContext->Preview->PreviewMesh->GetMesh();
+
+	if (ResultMesh->TriangleCount() == 0)
+	{
+		ParentTool->GetToolManager()->DisplayMessage( LOCTEXT("OnBevelFailedMessage", "Bevel created empty mesh, ignoring"), EToolMessageLevel::UserWarning);
+		// Reset the preview
+		ActivityContext->Preview->PreviewMesh->UpdatePreview(ActivityContext->CurrentMesh.Get());
+		return;
+	}
+
+	// Prep for undo.
+	FDynamicMeshChangeTracker ChangeTracker(ActivityContext->CurrentMesh.Get());
+	ChangeTracker.BeginChange();
+
+	// todo figure out what we need to save, instead of everything
+
+	TArray<int32> AllTids;
+	for (int32 Tid : ActivityContext->CurrentMesh->TriangleIndicesItr())
+	{
+		AllTids.Add(Tid);
+	}
+	ChangeTracker.SaveTriangles(AllTids, true /*bSaveVertices*/);
+
+	// Update current mesh
+	ActivityContext->CurrentMesh->Copy(*ResultMesh, true, true, true, true);
+
+	// discard selection for now
+	FGroupTopologySelection NewSelection;
+	//NewSelection.SelectedEdgeIDs.Append( ... );
+
+	// We need to reset the old selection back before we give the new one, so
+	// that undo reverts back to the correct selection state.
+	if (ActivityContext->SelectionMechanic->GetActiveSelection() != ActiveSelection)
+	{
+		ActivityContext->SelectionMechanic->SetSelection(ActiveSelection, false);
+	}
+
+	// Emit undo  (also updates relevant structures)
+	ActivityContext->EmitCurrentMeshChangeAndUpdate(LOCTEXT("PolyMeshBevelChange", "Bevel"),
+		ChangeTracker.EndChange(), NewSelection, true);
+
+	ActiveSelection = NewSelection;
+}
+
+
+void UPolyEditBevelEdgeActivity::Render(IToolsContextRenderAPI* RenderAPI)
+{
+}
+
+void UPolyEditBevelEdgeActivity::Tick(float DeltaTime)
+{
+}
+
+
+#undef LOCTEXT_NAMESPACE
diff --git a/Engine/Plugins/Runtime/MeshModelingToolset/Source/MeshModelingTools/Private/ToolActivities/PolyEditBevelEdgeActivity.h b/Engine/Plugins/Runtime/MeshModelingToolset/Source/MeshModelingTools/Private/ToolActivities/PolyEditBevelEdgeActivity.h
new file mode 100644
--- /dev/null
+++ b/Engine/Plugins/Runtime/MeshModelingToolset/Source/MeshModelingTools/Private/ToolActivities/PolyEditBevelEdgeActivity.h
@@ -0,0 +1,74 @@
+// Copyright Epic Games, Inc. All Rights Reserved.
+
+#pragma once
+
+#include "CoreMinimal.h"
+
+#include "BaseBehaviors/BehaviorTargetInterfaces.h"
+#include "GeometryBase.h"
+#include "GroupTopology.h" // FGroupTopologySelection
+#include "InteractiveToolActivity.h"
+#include "FrameTypes.h"
+#include "ModelingOperators.h"
+
+#include "PolyEditBevelEdgeActivity.generated.h"
+
+class UPolyEditActivityContext;
+class UPolyEditPreviewMesh;
+PREDECLARE_GEOMETRY(class FDynamicMesh3);
+
+
+
+UCLASS()
+class MESHMODELINGTOOLS_API UPolyEditBevelEdgeProperties : public UInteractiveToolPropertySet
+{
+	GENERATED_BODY()
+
+public:
+	UPROPERTY(EditAnywhere, Category = Bevel)
+	double BevelDistance = 4.0;
+};
+
+
+/**
+ * 
+ */
+UCLASS()
+class MESHMODELINGTOOLS_API UPolyEditBevelEdgeActivity : public UInteractiveToolActivity,
+	public UE::Geometry::IDynamicMeshOperatorFactory
+
+{
+	GENERATED_BODY()
+
+public:
+
+	// IInteractiveToolActivity
+	virtual void Setup(UInteractiveTool* ParentTool) override;
+	virtual void Shutdown(EToolShutdownType ShutdownType) override;
+	virtual bool CanStart() const override;
+	virtual EToolActivityStartResult Start() override;
+	virtual bool IsRunning() const override { return bIsRunning; }
+	virtual bool CanAccept() const override;
+	virtual EToolActivityEndResult End(EToolShutdownType) override;
+	virtual void Render(IToolsContextRenderAPI* RenderAPI) override;
+	virtual void Tick(float DeltaTime) override;
+
+	// IDynamicMeshOperatorFactory
+	virtual TUniquePtr<UE::Geometry::FDynamicMeshOperator> MakeNewOperator() override;
+
+	UPROPERTY()
+	TObjectPtr<UPolyEditBevelEdgeProperties> BevelProperties = nullptr;
+
+protected:
+
+	virtual void BeginBevel();
+	virtual void ApplyBevel();
+	virtual void EndInternal();
+
+	UPROPERTY()
+	TObjectPtr<UPolyEditActivityContext> ActivityContext;
+
+	bool bIsRunning = false;
+
+	UE::Geometry::FGroupTopologySelection ActiveSelection;
+};
diff --git a/Engine/Plugins/Runtime/MeshModelingToolset/Source/MeshModelingTools/Public/EditMeshPolygonsTool.h b/Engine/Plugins/Runtime/MeshModelingToolset/Source/MeshModelingTools/Public/EditMeshPolygonsTool.h
--- a/Engine/Plugins/Runtime/MeshModelingToolset/Source/MeshModelingTools/Public/EditMeshPolygonsTool.h
+++ b/Engine/Plugins/Runtime/MeshModelingToolset/Source/MeshModelingTools/Public/EditMeshPolygonsTool.h
@@ -32,6 +32,7 @@ class UPolyEditExtrudeActivity;
 class UPolyEditInsetOutsetActivity;
 class UPolyEditCutFacesActivity;
 class UPolyEditPlanarProjectionUVActivity;
+class UPolyEditBevelEdgeActivity;
 class UPolygonSelectionMechanic;
 class UTransformProxy;
 
@@ -93,12 +94,14 @@ UENUM()
 enum class EEditMeshPolygonsToolActions
 {
 	NoAction,
+	AcceptCurrent,
 	CancelCurrent,
 	Extrude,
 	PushPull,
 	Offset,
 	Inset,
 	Outset,
+	BevelFaces,
 	InsertEdge,
 	InsertEdgeLoop,
 	Complete,
@@ -118,6 +121,7 @@ enum class EEditMeshPolygonsToolActions
 	WeldEdges,
 	StraightenEdge,
 	FillHole,
+	BevelEdges,
 
 	PlanarProjectionUV,
 
@@ -206,6 +210,10 @@ public:
 	//~ TODO: Make the Merge, Delete, and Flip comments visible as tooltips. Currently we can't due to a bug that
 	//~ limits our total tooltip text allotment: UE-124608
 
+	//~ Bevel the edge loops around the selected faces 
+	UFUNCTION(CallInEditor, Category = FaceEdits, meta = (DisplayName = "Bevel", DisplayPriority = 4))
+	void Bevel() { PostAction(EEditMeshPolygonsToolActions::BevelFaces); }
+
 	//~ Merge the current set of selected faces into a single face.
 	UFUNCTION(CallInEditor, Category = FaceEdits, meta = (DisplayName = "Merge", DisplayPriority = 4))
 	void Merge() { PostAction(EEditMeshPolygonsToolActions::Merge);	}
@@ -340,7 +348,9 @@ public:
 	UFUNCTION(CallInEditor, Category = EdgeEdits, meta = (DisplayName = "Fill Hole", DisplayPriority = 5))
 	void FillHole()	{ PostAction(EEditMeshPolygonsToolActions::FillHole); }
 
-	
+	UFUNCTION(CallInEditor, Category = EdgeEdits, meta = (DisplayName = "Bevel", DisplayPriority = 6))
+	void Bevel() { PostAction(EEditMeshPolygonsToolActions::BevelEdges); }
+
 };
 
 
@@ -382,6 +392,24 @@ public:
 };
 
 
+/**
+ * TODO: This is currently a separate action set so that we can show/hide it depending on whether
+ * we have an activity running. We should have a cleaner alternative.
+ */
+UCLASS()
+class MESHMODELINGTOOLS_API UEditMeshPolygonsToolAcceptCancelAction : public UEditMeshPolygonsToolActionPropertySet
+{
+	GENERATED_BODY()
+public:
+	UFUNCTION(CallInEditor, Category = CurrentOperation, meta = (DisplayName = "Apply", DisplayPriority = 1))
+	void Apply() { PostAction(EEditMeshPolygonsToolActions::AcceptCurrent); }
+
+	UFUNCTION(CallInEditor, Category = CurrentOperation, meta = (DisplayName = "Cancel", DisplayPriority = 2))
+	void Cancel() { PostAction(EEditMeshPolygonsToolActions::CancelCurrent); }
+};
+
+
+
 
 /**
  *
@@ -467,6 +495,10 @@ protected:
 	UPROPERTY()
 	TObjectPtr<UEditMeshPolygonsToolCancelAction> CancelAction = nullptr;
 
+	UPROPERTY()
+	TObjectPtr<UEditMeshPolygonsToolAcceptCancelAction> AcceptCancelAction = nullptr;
+
+
 	/**
 	 * Activity objects that handle multi-interaction operations
 	 */
@@ -482,6 +514,8 @@ protected:
 	TObjectPtr<UPolyEditInsertEdgeActivity> InsertEdgeActivity = nullptr;
 	UPROPERTY()
 	TObjectPtr<UPolyEditInsertEdgeLoopActivity> InsertEdgeLoopActivity = nullptr;
+	UPROPERTY()
+	TObjectPtr<UPolyEditBevelEdgeActivity> BevelEdgeActivity = nullptr;
 
 	/**
 	 * Points to one of the activities when it is active
