fix crash on launching trim tool and trying to use mshbool tool's properties

diff --git a/Engine/Plugins/Runtime/MeshModelingToolset/Source/MeshModelingTools/Private/CSGMeshesTool.cpp b/Engine/Plugins/Runtime/MeshModelingToolset/Source/MeshModelingTools/Private/CSGMeshesTool.cpp
--- a/Engine/Plugins/Runtime/MeshModelingToolset/Source/MeshModelingTools/Private/CSGMeshesTool.cpp
+++ b/Engine/Plugins/Runtime/MeshModelingToolset/Source/MeshModelingTools/Private/CSGMeshesTool.cpp
@@ -291,9 +291,9 @@ TUniquePtr<FDynamicMeshOperator> UCSGMeshesTool::MakeNewOperator()
 	TUniquePtr<FBooleanMeshesOp> BooleanOp = MakeUnique<FBooleanMeshesOp>();
 	
 	BooleanOp->bTrimMode = bTrimMode;
-	BooleanOp->WindingThreshold = CSGProperties->WindingThreshold;
 	if (bTrimMode)
 	{
+		BooleanOp->WindingThreshold = TrimProperties->WindingThreshold;
 		BooleanOp->TrimOperation = TrimProperties->WhichMesh;
 		BooleanOp->TrimSide = TrimProperties->TrimSide;
 		BooleanOp->bAttemptFixHoles = false;
@@ -301,6 +301,7 @@ TUniquePtr<FDynamicMeshOperator> UCSGMeshesTool::MakeNewOperator()
 	}
 	else
 	{
+		BooleanOp->WindingThreshold = CSGProperties->WindingThreshold;
 		BooleanOp->CSGOperation = CSGProperties->Operation;
 		BooleanOp->bAttemptFixHoles = CSGProperties->bTryFixHoles;
 		BooleanOp->bTryCollapseExtraEdges = CSGProperties->bTryCollapseEdges;
diff --git a/Engine/Plugins/Runtime/MeshModelingToolset/Source/MeshModelingTools/Public/CSGMeshesTool.h b/Engine/Plugins/Runtime/MeshModelingToolset/Source/MeshModelingTools/Public/CSGMeshesTool.h
--- a/Engine/Plugins/Runtime/MeshModelingToolset/Source/MeshModelingTools/Public/CSGMeshesTool.h
+++ b/Engine/Plugins/Runtime/MeshModelingToolset/Source/MeshModelingTools/Public/CSGMeshesTool.h
@@ -78,6 +78,10 @@ public:
 	UPROPERTY(EditAnywhere, Category = Operation)
 	ETrimSide TrimSide = ETrimSide::RemoveInside;
 
+	/** Threshold to determine whether a triangle in one mesh is inside or outside of the other */
+	UPROPERTY(EditAnywhere, Category = Operation, AdvancedDisplay, meta = (UIMin = "0", UIMax = "1"))
+	float WindingThreshold = 0.5;
+
 	/** Whether to show a translucent version of the trimming mesh, to help visualize what is being cut */
 	UPROPERTY(EditAnywhere, Category = Display)
 	bool bShowTrimmingMesh = true;
