ModelingTools: Make TargetWorld a proper property on SingleSelectionMeshEditingTool.

#rnx

diff --git a/Engine/Plugins/Runtime/MeshModelingToolset/Source/ModelingComponents/Public/BaseTools/SingleSelectionMeshEditingTool.h b/Engine/Plugins/Runtime/MeshModelingToolset/Source/ModelingComponents/Public/BaseTools/SingleSelectionMeshEditingTool.h
--- a/Engine/Plugins/Runtime/MeshModelingToolset/Source/ModelingComponents/Public/BaseTools/SingleSelectionMeshEditingTool.h
+++ b/Engine/Plugins/Runtime/MeshModelingToolset/Source/ModelingComponents/Public/BaseTools/SingleSelectionMeshEditingTool.h
@@ -54,7 +54,8 @@ public:
 	virtual UWorld* GetTargetWorld() { return TargetWorld; }
 
 protected:
-	UWorld* TargetWorld = nullptr;
+	UPROPERTY()
+	TObjectPtr<UWorld> TargetWorld = nullptr;
 
 	//
 	// Mesh Selection support
