InteractiveToolFramework: Remove unnecessary tool target interface requirements for MeshSurfacePointTool.

ModelingTools: Add MeshSurfacePointMeshEditingToolBuilder with common mesh editing requirements.
Update MeshSurfacePointTool modeling tools to use MeshSurfacePointMeshEditingToolBuilder

#rnx

diff --git a/Engine/Plugins/Experimental/HairModelingToolset/Source/HairModelingToolset/Public/GroomCardsEditorTool.h b/Engine/Plugins/Experimental/HairModelingToolset/Source/HairModelingToolset/Public/GroomCardsEditorTool.h
--- a/Engine/Plugins/Experimental/HairModelingToolset/Source/HairModelingToolset/Public/GroomCardsEditorTool.h
+++ b/Engine/Plugins/Experimental/HairModelingToolset/Source/HairModelingToolset/Public/GroomCardsEditorTool.h
@@ -4,7 +4,7 @@
 
 #include "CoreMinimal.h"
 #include "Templates/PimplPtr.h"
-#include "BaseTools/MeshSurfacePointTool.h"
+#include "BaseTools/MeshSurfacePointMeshEditingTool.h"
 #include "InteractiveToolBuilder.h"
 #include "PreviewMesh.h"
 #include "Drawing/PreviewGeometryActor.h"
@@ -25,11 +25,13 @@ class FEditableGroomCardSet;
 class FGroomCardEdit;
 class FMeshVertexChangeBuilder;
 
+class UGroomCardsEditorTool;
+
 /**
  *
  */
 UCLASS()
-class HAIRMODELINGTOOLSET_API UGroomCardsEditorToolBuilder : public UMeshSurfacePointToolBuilder
+class HAIRMODELINGTOOLSET_API UGroomCardsEditorToolBuilder : public UMeshSurfacePointMeshEditingToolBuilder
 {
 	GENERATED_BODY()
 public:
diff --git a/Engine/Plugins/Experimental/MeshModelingToolsetExp/Source/MeshModelingToolsEditorOnly/Public/ShapeSprayTool.h b/Engine/Plugins/Experimental/MeshModelingToolsetExp/Source/MeshModelingToolsEditorOnly/Public/ShapeSprayTool.h
--- a/Engine/Plugins/Experimental/MeshModelingToolsetExp/Source/MeshModelingToolsEditorOnly/Public/ShapeSprayTool.h
+++ b/Engine/Plugins/Experimental/MeshModelingToolsetExp/Source/MeshModelingToolsEditorOnly/Public/ShapeSprayTool.h
@@ -5,13 +5,14 @@
 #include "CoreMinimal.h"
 #include "InteractiveToolBuilder.h"
 #include "DynamicMeshBrushTool.h"
+#include "BaseTools/MeshSurfacePointMeshEditingTool.h"
 #include "ShapeSprayTool.generated.h"
 
 /**
- * UMeshSurfacePointToolBuilder override for UShapeSprayTool
+ * UMeshSurfacePointMeshEditingToolBuilder override for UShapeSprayTool
  */
 UCLASS(Transient)
-class MESHMODELINGTOOLSEDITORONLY_API UShapeSprayToolBuilder : public UMeshSurfacePointToolBuilder
+class MESHMODELINGTOOLSEDITORONLY_API UShapeSprayToolBuilder : public UMeshSurfacePointMeshEditingToolBuilder
 {
 	GENERATED_BODY()
 
diff --git a/Engine/Plugins/Experimental/MeshModelingToolsetExp/Source/MeshModelingToolsExp/Public/DeformMeshPolygonsTool.h b/Engine/Plugins/Experimental/MeshModelingToolsetExp/Source/MeshModelingToolsExp/Public/DeformMeshPolygonsTool.h
--- a/Engine/Plugins/Experimental/MeshModelingToolsetExp/Source/MeshModelingToolsExp/Public/DeformMeshPolygonsTool.h
+++ b/Engine/Plugins/Experimental/MeshModelingToolsetExp/Source/MeshModelingToolsExp/Public/DeformMeshPolygonsTool.h
@@ -5,7 +5,7 @@
 #include "CoreMinimal.h"
 #include "Templates/PimplPtr.h"
 #include "UObject/NoExportTypes.h"
-#include "BaseTools/MeshSurfacePointTool.h"
+#include "BaseTools/MeshSurfacePointMeshEditingTool.h"
 #include "Components/DynamicMeshComponent.h"
 #include "DynamicMesh/DynamicMeshAABBTree3.h"
 #include "ToolDataVisualizer.h"
@@ -27,7 +27,7 @@ struct FDeformerVertexConstraintData;
  * ToolBuilder
  */
 UCLASS()
-class MESHMODELINGTOOLSEXP_API UDeformMeshPolygonsToolBuilder : public UMeshSurfacePointToolBuilder
+class MESHMODELINGTOOLSEXP_API UDeformMeshPolygonsToolBuilder : public UMeshSurfacePointMeshEditingToolBuilder
 {
 	GENERATED_BODY()
 
diff --git a/Engine/Plugins/Experimental/MeshModelingToolsetExp/Source/MeshModelingToolsExp/Public/DynamicMeshSculptTool.h b/Engine/Plugins/Experimental/MeshModelingToolsetExp/Source/MeshModelingToolsExp/Public/DynamicMeshSculptTool.h
--- a/Engine/Plugins/Experimental/MeshModelingToolsetExp/Source/MeshModelingToolsExp/Public/DynamicMeshSculptTool.h
+++ b/Engine/Plugins/Experimental/MeshModelingToolsetExp/Source/MeshModelingToolsExp/Public/DynamicMeshSculptTool.h
@@ -6,7 +6,7 @@
 #include "HAL/Platform.h"
 #include "SubRegionRemesher.h"
 #include "UObject/NoExportTypes.h"
-#include "BaseTools/MeshSurfacePointTool.h"
+#include "BaseTools/MeshSurfacePointMeshEditingTool.h"
 #include "Components/OctreeDynamicMeshComponent.h"
 #include "DynamicMesh/DynamicMeshAABBTree3.h"
 #include "DynamicMesh/DynamicMeshOctree3.h"
@@ -95,7 +95,7 @@ enum class EDynamicMeshSculptBrushType : uint8
  * Tool Builder
  */
 UCLASS()
-class MESHMODELINGTOOLSEXP_API UDynamicMeshSculptToolBuilder : public UMeshSurfacePointToolBuilder
+class MESHMODELINGTOOLSEXP_API UDynamicMeshSculptToolBuilder : public UMeshSurfacePointMeshEditingToolBuilder
 {
 	GENERATED_BODY()
 
diff --git a/Engine/Plugins/Experimental/MeshModelingToolsetExp/Source/MeshModelingToolsExp/Public/EditUVIslandsTool.h b/Engine/Plugins/Experimental/MeshModelingToolsetExp/Source/MeshModelingToolsExp/Public/EditUVIslandsTool.h
--- a/Engine/Plugins/Experimental/MeshModelingToolsetExp/Source/MeshModelingToolsExp/Public/EditUVIslandsTool.h
+++ b/Engine/Plugins/Experimental/MeshModelingToolsetExp/Source/MeshModelingToolsExp/Public/EditUVIslandsTool.h
@@ -3,7 +3,7 @@
 #pragma once
 
 #include "CoreMinimal.h"
-#include "BaseTools/MeshSurfacePointTool.h"
+#include "BaseTools/MeshSurfacePointMeshEditingTool.h"
 #include "Properties/MeshMaterialProperties.h"
 #include "Components/DynamicMeshComponent.h"
 #include "DynamicMesh/DynamicMeshAABBTree3.h"
@@ -23,7 +23,7 @@ using UE::Geometry::FDynamicMeshUVOverlay;
  * ToolBuilder
  */
 UCLASS()
-class MESHMODELINGTOOLSEXP_API UEditUVIslandsToolBuilder : public UMeshSurfacePointToolBuilder
+class MESHMODELINGTOOLSEXP_API UEditUVIslandsToolBuilder : public UMeshSurfacePointMeshEditingToolBuilder
 {
 	GENERATED_BODY()
 public:
diff --git a/Engine/Plugins/Experimental/MeshModelingToolsetExp/Source/MeshModelingToolsExp/Public/MeshAttributePaintTool.h b/Engine/Plugins/Experimental/MeshModelingToolsetExp/Source/MeshModelingToolsExp/Public/MeshAttributePaintTool.h
--- a/Engine/Plugins/Experimental/MeshModelingToolsetExp/Source/MeshModelingToolsExp/Public/MeshAttributePaintTool.h
+++ b/Engine/Plugins/Experimental/MeshModelingToolsetExp/Source/MeshModelingToolsExp/Public/MeshAttributePaintTool.h
@@ -3,6 +3,7 @@
 #pragma once
 
 #include "DynamicMeshBrushTool.h"
+#include "BaseTools/MeshSurfacePointMeshEditingTool.h"
 #include "Changes/ValueWatcher.h"
 #include "Changes/IndexedAttributeChange.h"
 #include "DynamicMesh/DynamicVerticesOctree3.h"
@@ -11,6 +12,7 @@
 
 
 struct FMeshDescription;
+class UMeshAttributePaintTool;
 
 
 /**
@@ -82,7 +84,7 @@ public:
  * Tool Builder for Attribute Paint Tool
  */
 UCLASS()
-class MESHMODELINGTOOLSEXP_API UMeshAttributePaintToolBuilder : public UMeshSurfacePointToolBuilder
+class MESHMODELINGTOOLSEXP_API UMeshAttributePaintToolBuilder : public UMeshSurfacePointMeshEditingToolBuilder
 {
 	GENERATED_BODY()
 public:
diff --git a/Engine/Plugins/Experimental/MeshModelingToolsetExp/Source/MeshModelingToolsExp/Public/MeshGroupPaintTool.h b/Engine/Plugins/Experimental/MeshModelingToolsetExp/Source/MeshModelingToolsExp/Public/MeshGroupPaintTool.h
--- a/Engine/Plugins/Experimental/MeshModelingToolsetExp/Source/MeshModelingToolsExp/Public/MeshGroupPaintTool.h
+++ b/Engine/Plugins/Experimental/MeshModelingToolsetExp/Source/MeshModelingToolsExp/Public/MeshGroupPaintTool.h
@@ -6,6 +6,7 @@
 #include "HAL/Platform.h"
 #include "UObject/NoExportTypes.h"
 #include "BaseTools/BaseBrushTool.h"
+#include "BaseTools/MeshSurfacePointMeshEditingTool.h"
 #include "Components/DynamicMeshComponent.h"
 #include "PropertySets/PolygroupLayersProperties.h"
 #include "Mechanics/PolyLassoMarqueeMechanic.h"
@@ -47,7 +48,7 @@ DECLARE_CYCLE_STAT(TEXT("GroupPaintTool_Normals_Compute"), GroupPaintTool_Normal
  * Tool Builder
  */
 UCLASS()
-class MESHMODELINGTOOLSEXP_API UMeshGroupPaintToolBuilder : public UMeshSurfacePointToolBuilder
+class MESHMODELINGTOOLSEXP_API UMeshGroupPaintToolBuilder : public UMeshSurfacePointMeshEditingToolBuilder
 {
 	GENERATED_BODY()
 
diff --git a/Engine/Plugins/Experimental/MeshModelingToolsetExp/Source/MeshModelingToolsExp/Public/MeshSelectionTool.h b/Engine/Plugins/Experimental/MeshModelingToolsetExp/Source/MeshModelingToolsExp/Public/MeshSelectionTool.h
--- a/Engine/Plugins/Experimental/MeshModelingToolsetExp/Source/MeshModelingToolsExp/Public/MeshSelectionTool.h
+++ b/Engine/Plugins/Experimental/MeshModelingToolsetExp/Source/MeshModelingToolsExp/Public/MeshSelectionTool.h
@@ -4,6 +4,7 @@
 
 #include "CoreMinimal.h"
 #include "DynamicMeshBrushTool.h"
+#include "BaseTools/MeshSurfacePointMeshEditingTool.h"
 #include "InteractiveToolQueryInterfaces.h"
 #include "SelectionSet.h"
 #include "Changes/MeshSelectionChange.h"
@@ -15,12 +16,13 @@ class UMeshStatisticsProperties;
 class UMeshElementsVisualizer;
 class UMeshUVChannelProperties;
 class UPolygroupLayersProperties;
+class UMeshSelectionTool;
 
 /**
  *
  */
 UCLASS()
-class MESHMODELINGTOOLSEXP_API UMeshSelectionToolBuilder : public UMeshSurfacePointToolBuilder
+class MESHMODELINGTOOLSEXP_API UMeshSelectionToolBuilder : public UMeshSurfacePointMeshEditingToolBuilder
 {
 	GENERATED_BODY()
 
diff --git a/Engine/Plugins/Experimental/MeshModelingToolsetExp/Source/MeshModelingToolsExp/Public/MeshVertexSculptTool.h b/Engine/Plugins/Experimental/MeshModelingToolsetExp/Source/MeshModelingToolsExp/Public/MeshVertexSculptTool.h
--- a/Engine/Plugins/Experimental/MeshModelingToolsetExp/Source/MeshModelingToolsExp/Public/MeshVertexSculptTool.h
+++ b/Engine/Plugins/Experimental/MeshModelingToolsetExp/Source/MeshModelingToolsExp/Public/MeshVertexSculptTool.h
@@ -10,6 +10,7 @@
 #include "DynamicMesh/DynamicMeshOctree3.h"
 #include "DynamicMesh/MeshNormals.h"
 #include "BaseTools/BaseBrushTool.h"
+#include "BaseTools/MeshSurfacePointMeshEditingTool.h"
 #include "TransformTypes.h"
 #include "Sculpting/MeshSculptToolBase.h"
 #include "Sculpting/MeshBrushOpBase.h"
@@ -29,7 +30,7 @@ class UPreviewMesh;
  * Tool Builder
  */
 UCLASS()
-class MESHMODELINGTOOLSEXP_API UMeshVertexSculptToolBuilder : public UMeshSurfacePointToolBuilder
+class MESHMODELINGTOOLSEXP_API UMeshVertexSculptToolBuilder : public UMeshSurfacePointMeshEditingToolBuilder
 {
 	GENERATED_BODY()
 
diff --git a/Engine/Plugins/Experimental/MeshModelingToolsetExp/Source/MeshModelingToolsExp/Public/SeamSculptTool.h b/Engine/Plugins/Experimental/MeshModelingToolsetExp/Source/MeshModelingToolsExp/Public/SeamSculptTool.h
--- a/Engine/Plugins/Experimental/MeshModelingToolsetExp/Source/MeshModelingToolsExp/Public/SeamSculptTool.h
+++ b/Engine/Plugins/Experimental/MeshModelingToolsetExp/Source/MeshModelingToolsExp/Public/SeamSculptTool.h
@@ -4,6 +4,7 @@
 
 #include "CoreMinimal.h"
 #include "DynamicMeshBrushTool.h"
+#include "BaseTools/MeshSurfacePointMeshEditingTool.h"
 #include "SeamSculptTool.generated.h"
 
 class UExistingMeshMaterialProperties;
@@ -14,7 +15,7 @@ PREDECLARE_USE_GEOMETRY_CLASS(FDynamicMesh3);
  *
  */
 UCLASS()
-class MESHMODELINGTOOLSEXP_API USeamSculptToolBuilder : public UMeshSurfacePointToolBuilder
+class MESHMODELINGTOOLSEXP_API USeamSculptToolBuilder : public UMeshSurfacePointMeshEditingToolBuilder
 {
 	GENERATED_BODY()
 
diff --git a/Engine/Plugins/Experimental/MeshModelingToolsetExp/Source/MeshModelingToolsExp/Public/SkinWeightsPaintTool.h b/Engine/Plugins/Experimental/MeshModelingToolsetExp/Source/MeshModelingToolsExp/Public/SkinWeightsPaintTool.h
--- a/Engine/Plugins/Experimental/MeshModelingToolsetExp/Source/MeshModelingToolsExp/Public/SkinWeightsPaintTool.h
+++ b/Engine/Plugins/Experimental/MeshModelingToolsetExp/Source/MeshModelingToolsExp/Public/SkinWeightsPaintTool.h
@@ -3,6 +3,7 @@
 #pragma once
 
 #include "DynamicMeshBrushTool.h"
+#include "BaseTools/MeshSurfacePointMeshEditingTool.h"
 #include "Changes/ValueWatcher.h"
 #include "MeshDescription.h"
 #include "DynamicMesh/DynamicVerticesOctree3.h"
@@ -51,7 +52,7 @@ private:
  *
  */
 UCLASS()
-class MESHMODELINGTOOLSEXP_API USkinWeightsPaintToolBuilder : public UMeshSurfacePointToolBuilder
+class MESHMODELINGTOOLSEXP_API USkinWeightsPaintToolBuilder : public UMeshSurfacePointMeshEditingToolBuilder
 {
 	GENERATED_BODY()
 
diff --git a/Engine/Plugins/Runtime/MeshModelingToolset/Source/ModelingComponents/Private/BaseTools/MeshSurfacePointMeshEditingTool.cpp b/Engine/Plugins/Runtime/MeshModelingToolset/Source/ModelingComponents/Private/BaseTools/MeshSurfacePointMeshEditingTool.cpp
new file mode 100644
--- /dev/null
+++ b/Engine/Plugins/Runtime/MeshModelingToolset/Source/ModelingComponents/Private/BaseTools/MeshSurfacePointMeshEditingTool.cpp
@@ -0,0 +1,34 @@
+// Copyright Epic Games, Inc. All Rights Reserved.
+
+#include "BaseTools/MeshSurfacePointMeshEditingTool.h"
+#include "InteractiveToolManager.h"
+#include "ToolBuilderUtil.h"
+
+#include "TargetInterfaces/PrimitiveComponentBackedTarget.h"
+#include "TargetInterfaces/MeshDescriptionProvider.h"
+#include "TargetInterfaces/MeshDescriptionCommitter.h"
+#include "TargetInterfaces/MaterialProvider.h"
+#include "ToolTargetManager.h"
+
+
+/*
+ * ToolBuilder
+ */
+
+const FToolTargetTypeRequirements& UMeshSurfacePointMeshEditingToolBuilder::GetTargetRequirements() const
+{
+	static FToolTargetTypeRequirements TypeRequirements({
+		UMaterialProvider::StaticClass(),
+		UMeshDescriptionProvider::StaticClass(),
+		UMeshDescriptionCommitter::StaticClass(),
+		UPrimitiveComponentBackedTarget::StaticClass()
+		});
+	return TypeRequirements;
+}
+
+UMeshSurfacePointTool* UMeshSurfacePointMeshEditingToolBuilder::CreateNewTool(const FToolBuilderState& SceneState) const
+{
+	return NewObject<UMeshSurfacePointTool>(SceneState.ToolManager);
+}
+
+
diff --git a/Engine/Plugins/Runtime/MeshModelingToolset/Source/ModelingComponents/Public/BaseTools/MeshSurfacePointMeshEditingTool.h b/Engine/Plugins/Runtime/MeshModelingToolset/Source/ModelingComponents/Public/BaseTools/MeshSurfacePointMeshEditingTool.h
new file mode 100644
--- /dev/null
+++ b/Engine/Plugins/Runtime/MeshModelingToolset/Source/ModelingComponents/Public/BaseTools/MeshSurfacePointMeshEditingTool.h
@@ -0,0 +1,22 @@
+// Copyright Epic Games, Inc. All Rights Reserved.
+
+#pragma once
+
+#include "BaseTools/MeshSurfacePointTool.h"
+#include "MeshSurfacePointMeshEditingTool.generated.h"
+
+/**
+ * Base tool builder class for UMeshSurfacePointTools with mesh editing requirements.
+ */
+UCLASS()
+class MODELINGCOMPONENTS_API UMeshSurfacePointMeshEditingToolBuilder : public UMeshSurfacePointToolBuilder
+{
+	GENERATED_BODY()
+public:
+	/** @return new Tool instance. Override this in subclasses to build a different Tool class type */
+	virtual UMeshSurfacePointTool* CreateNewTool(const FToolBuilderState& SceneState) const;
+
+protected:
+	virtual const FToolTargetTypeRequirements& GetTargetRequirements() const override;
+};
+
diff --git a/Engine/Source/Runtime/InteractiveToolsFramework/Private/BaseTools/MeshSurfacePointTool.cpp b/Engine/Source/Runtime/InteractiveToolsFramework/Private/BaseTools/MeshSurfacePointTool.cpp
--- a/Engine/Source/Runtime/InteractiveToolsFramework/Private/BaseTools/MeshSurfacePointTool.cpp
+++ b/Engine/Source/Runtime/InteractiveToolsFramework/Private/BaseTools/MeshSurfacePointTool.cpp
@@ -5,11 +5,7 @@
 #include "BaseBehaviors/ClickDragBehavior.h"
 #include "InteractiveToolManager.h"
 #include "ToolBuilderUtil.h"
-#include "Components/PrimitiveComponent.h"
 
-#include "TargetInterfaces/MaterialProvider.h"
-#include "TargetInterfaces/MeshDescriptionCommitter.h"
-#include "TargetInterfaces/MeshDescriptionProvider.h"
 #include "TargetInterfaces/PrimitiveComponentBackedTarget.h"
 #include "ToolTargetManager.h"
 
@@ -22,12 +18,7 @@
 
 const FToolTargetTypeRequirements& UMeshSurfacePointToolBuilder::GetTargetRequirements() const
 {
-	static FToolTargetTypeRequirements TypeRequirements({
-		UMaterialProvider::StaticClass(),
-		UMeshDescriptionCommitter::StaticClass(),
-		UMeshDescriptionProvider::StaticClass(),
-		UPrimitiveComponentBackedTarget::StaticClass()
-		});
+	static FToolTargetTypeRequirements TypeRequirements(UPrimitiveComponentBackedTarget::StaticClass());
 	return TypeRequirements;
 }
 
