make to/from static mesh conversion not crash if it winds up with a null meshdescription (but still ensure, because something has gone quite wrong in this case)

#rnx

diff --git a/Engine/Plugins/Runtime/MeshModelingToolset/Source/ModelingComponents/Private/ModelingToolTargetUtil.cpp b/Engine/Plugins/Runtime/MeshModelingToolset/Source/ModelingComponents/Private/ModelingToolTargetUtil.cpp
--- a/Engine/Plugins/Runtime/MeshModelingToolset/Source/ModelingComponents/Private/ModelingToolTargetUtil.cpp
+++ b/Engine/Plugins/Runtime/MeshModelingToolset/Source/ModelingComponents/Private/ModelingToolTargetUtil.cpp
@@ -220,6 +220,11 @@ FDynamicMesh3 UE::ToolTarget::GetDynamicMeshCopy(UToolTarget* Target, bool bWant
 
 UE::ToolTarget::EDynamicMeshUpdateResult UE::ToolTarget::CommitMeshDescriptionUpdate(UToolTarget* Target, const FMeshDescription* UpdatedMesh, const FComponentMaterialSet* UpdatedMaterials)
 {
+	if (!ensure(UpdatedMesh != nullptr))
+	{
+		return EDynamicMeshUpdateResult::Failed;
+	}
+
 	IMeshDescriptionCommitter* MeshDescriptionCommitter = Cast<IMeshDescriptionCommitter>(Target);
 	if (!ensure(MeshDescriptionCommitter))
 	{
diff --git a/Engine/Plugins/Runtime/MeshModelingToolset/Source/ModelingComponentsEditorOnly/Private/ToolTargets/StaticMeshComponentToolTarget.cpp b/Engine/Plugins/Runtime/MeshModelingToolset/Source/ModelingComponentsEditorOnly/Private/ToolTargets/StaticMeshComponentToolTarget.cpp
--- a/Engine/Plugins/Runtime/MeshModelingToolset/Source/ModelingComponentsEditorOnly/Private/ToolTargets/StaticMeshComponentToolTarget.cpp
+++ b/Engine/Plugins/Runtime/MeshModelingToolset/Source/ModelingComponentsEditorOnly/Private/ToolTargets/StaticMeshComponentToolTarget.cpp
@@ -161,6 +161,11 @@ void UStaticMeshComponentToolTarget::CommitMeshDescription(const FCommitter& Com
 	FMeshDescription* UpdateMeshDescription = (EditingLOD == EStaticMeshEditingLOD::HiResSource) ?
 		StaticMesh->GetHiResMeshDescription() : StaticMesh->GetMeshDescription((int32)EditingLOD);
 
+	if (!ensure(UpdateMeshDescription != nullptr))
+	{
+		return; // invalid target; cannot commit
+	}
+
 	// unregister the component while we update its static mesh
 	FComponentReregisterContext ComponentReregisterContext(Component);
 
diff --git a/Engine/Source/Runtime/MeshConversion/Private/MeshDescriptionToDynamicMesh.cpp b/Engine/Source/Runtime/MeshConversion/Private/MeshDescriptionToDynamicMesh.cpp
--- a/Engine/Source/Runtime/MeshConversion/Private/MeshDescriptionToDynamicMesh.cpp
+++ b/Engine/Source/Runtime/MeshConversion/Private/MeshDescriptionToDynamicMesh.cpp
@@ -173,6 +173,11 @@ void FMeshDescriptionToDynamicMesh::Convert(const FMeshDescription* MeshIn, FDyn
 	TriIDMap.Reset();
 	VertIDMap.Reset();
 
+	if (!ensure(MeshIn != nullptr))
+	{
+		return; // nothing to convert
+	}
+
 	// allocate the VertIDMap.  Unfortunately the array will need to grow more if MeshIn has non-manifold edges that need to be split
 	VertIDMap.SetNumUninitialized(MeshIn->Vertices().Num());
 
