InteractiveToolsFramework: Add UInteractiveToolWithToolTargetsBuilder intermediary class to define the common GetTargetRequirements() interface used across tools that use tool targets.

UVEditor: Remove unused GetTargetRequirements interface on UVEditor tools.

#rnx

diff --git a/Engine/Plugins/Experimental/MeshModelingToolsetExp/Source/MeshModelingToolsExp/Public/CubeGridTool.h b/Engine/Plugins/Experimental/MeshModelingToolsetExp/Source/MeshModelingToolsExp/Public/CubeGridTool.h
--- a/Engine/Plugins/Experimental/MeshModelingToolsetExp/Source/MeshModelingToolsExp/Public/CubeGridTool.h
+++ b/Engine/Plugins/Experimental/MeshModelingToolsetExp/Source/MeshModelingToolsExp/Public/CubeGridTool.h
@@ -10,7 +10,7 @@
 #include "FrameTypes.h"
 #include "GeometryBase.h"
 #include "InteractiveTool.h"
-#include "InteractiveToolBuilder.h"
+#include "InteractiveToolWithToolTargetsBuilder.h"
 #include "InteractiveToolQueryInterfaces.h" // IInteractiveToolNestedAcceptCancelAPI
 #include "Mechanics/CubeGrid.h"
 #include "ModelingOperators.h"
@@ -42,7 +42,7 @@ class UToolTarget;
 
 
 UCLASS()
-class MESHMODELINGTOOLSEXP_API UCubeGridToolBuilder : public UInteractiveToolBuilder
+class MESHMODELINGTOOLSEXP_API UCubeGridToolBuilder : public UInteractiveToolWithToolTargetsBuilder
 {
 	GENERATED_BODY()
 
@@ -54,7 +54,7 @@ public:
 
 protected:
 
-	virtual const FToolTargetTypeRequirements& GetTargetRequirements() const;
+	virtual const FToolTargetTypeRequirements& GetTargetRequirements() const override;
 };
 
 
diff --git a/Engine/Plugins/Experimental/UVEditor/Source/UVEditorTools/Private/UVEditorLayerEditTool.cpp b/Engine/Plugins/Experimental/UVEditor/Source/UVEditorTools/Private/UVEditorLayerEditTool.cpp
--- a/Engine/Plugins/Experimental/UVEditor/Source/UVEditorTools/Private/UVEditorLayerEditTool.cpp
+++ b/Engine/Plugins/Experimental/UVEditor/Source/UVEditorTools/Private/UVEditorLayerEditTool.cpp
@@ -261,12 +261,6 @@ namespace UUVEditorChannelEditLocals
 }
 
 
-const FToolTargetTypeRequirements& UUVEditorChannelEditToolBuilder::GetTargetRequirements() const
-{
-	static FToolTargetTypeRequirements TypeRequirements(UUVEditorToolMeshInput::StaticClass());
-	return TypeRequirements;
-}
-
 bool UUVEditorChannelEditToolBuilder::CanBuildTool(const FToolBuilderState& SceneState) const
 {
 	return Targets && Targets->Num() > 0;
diff --git a/Engine/Plugins/Experimental/UVEditor/Source/UVEditorTools/Private/UVEditorLayoutTool.cpp b/Engine/Plugins/Experimental/UVEditor/Source/UVEditorTools/Private/UVEditorLayoutTool.cpp
--- a/Engine/Plugins/Experimental/UVEditor/Source/UVEditorTools/Private/UVEditorLayoutTool.cpp
+++ b/Engine/Plugins/Experimental/UVEditor/Source/UVEditorTools/Private/UVEditorLayoutTool.cpp
@@ -19,12 +19,6 @@ using namespace UE::Geometry;
 // Tool builder
 // TODO: Could consider sharing some of the tool builder boilerplate for UV editor tools in a common base class.
 
-const FToolTargetTypeRequirements& UUVEditorLayoutToolBuilder::GetTargetRequirements() const
-{
-	static FToolTargetTypeRequirements TypeRequirements(UUVEditorToolMeshInput::StaticClass());
-	return TypeRequirements;
-}
-
 bool UUVEditorLayoutToolBuilder::CanBuildTool(const FToolBuilderState& SceneState) const
 {
 	return Targets && Targets->Num() > 0;
diff --git a/Engine/Plugins/Experimental/UVEditor/Source/UVEditorTools/Private/UVEditorRecomputeUVsTool.cpp b/Engine/Plugins/Experimental/UVEditor/Source/UVEditorTools/Private/UVEditorRecomputeUVsTool.cpp
--- a/Engine/Plugins/Experimental/UVEditor/Source/UVEditorTools/Private/UVEditorRecomputeUVsTool.cpp
+++ b/Engine/Plugins/Experimental/UVEditor/Source/UVEditorTools/Private/UVEditorRecomputeUVsTool.cpp
@@ -23,12 +23,6 @@ using namespace UE::Geometry;
  * ToolBuilder
  */
 
-const FToolTargetTypeRequirements& UUVEditorRecomputeUVsToolBuilder::GetTargetRequirements() const
-{
-	static FToolTargetTypeRequirements TypeRequirements(UUVEditorToolMeshInput::StaticClass());
-	return TypeRequirements;
-}
-
 bool UUVEditorRecomputeUVsToolBuilder::CanBuildTool(const FToolBuilderState& SceneState) const
 {
 	return Targets && Targets->Num() > 0;
diff --git a/Engine/Plugins/Experimental/UVEditor/Source/UVEditorTools/Private/UVSelectTool.cpp b/Engine/Plugins/Experimental/UVEditor/Source/UVEditorTools/Private/UVSelectTool.cpp
--- a/Engine/Plugins/Experimental/UVEditor/Source/UVEditorTools/Private/UVSelectTool.cpp
+++ b/Engine/Plugins/Experimental/UVEditor/Source/UVEditorTools/Private/UVSelectTool.cpp
@@ -260,12 +260,6 @@ namespace UVSelectToolLocals
  * ToolBuilder
  */
 
-const FToolTargetTypeRequirements& UUVSelectToolBuilder::GetTargetRequirements() const
-{
-	static FToolTargetTypeRequirements TypeRequirements(UUVEditorToolMeshInput::StaticClass());
-	return TypeRequirements;
-}
-
 bool UUVSelectToolBuilder::CanBuildTool(const FToolBuilderState& SceneState) const
 {
 	return Targets && Targets->Num() > 0;
diff --git a/Engine/Plugins/Experimental/UVEditor/Source/UVEditorTools/Public/UVEditorLayerEditTool.h b/Engine/Plugins/Experimental/UVEditor/Source/UVEditorTools/Public/UVEditorLayerEditTool.h
--- a/Engine/Plugins/Experimental/UVEditor/Source/UVEditorTools/Public/UVEditorLayerEditTool.h
+++ b/Engine/Plugins/Experimental/UVEditor/Source/UVEditorTools/Public/UVEditorLayerEditTool.h
@@ -25,9 +25,6 @@ public:
 	// This is a pointer so that it can be updated under the builder without
 	// having to set it in the mode after initializing targets.
 	const TArray<TObjectPtr<UUVEditorToolMeshInput>>* Targets = nullptr;
-
-protected:
-	virtual const FToolTargetTypeRequirements& GetTargetRequirements() const;
 };
 
 UENUM()
diff --git a/Engine/Plugins/Experimental/UVEditor/Source/UVEditorTools/Public/UVEditorLayoutTool.h b/Engine/Plugins/Experimental/UVEditor/Source/UVEditorTools/Public/UVEditorLayoutTool.h
--- a/Engine/Plugins/Experimental/UVEditor/Source/UVEditorTools/Public/UVEditorLayoutTool.h
+++ b/Engine/Plugins/Experimental/UVEditor/Source/UVEditorTools/Public/UVEditorLayoutTool.h
@@ -25,9 +25,6 @@ public:
 	// This is a pointer so that it can be updated under the builder without
 	// having to set it in the mode after initializing targets.
 	const TArray<TObjectPtr<UUVEditorToolMeshInput>>* Targets = nullptr;
-
-protected:
-	virtual const FToolTargetTypeRequirements& GetTargetRequirements() const;
 };
 
 /**
diff --git a/Engine/Plugins/Experimental/UVEditor/Source/UVEditorTools/Public/UVEditorRecomputeUVsTool.h b/Engine/Plugins/Experimental/UVEditor/Source/UVEditorTools/Public/UVEditorRecomputeUVsTool.h
--- a/Engine/Plugins/Experimental/UVEditor/Source/UVEditorTools/Public/UVEditorRecomputeUVsTool.h
+++ b/Engine/Plugins/Experimental/UVEditor/Source/UVEditorTools/Public/UVEditorRecomputeUVsTool.h
@@ -36,9 +36,6 @@ public:
 	// This is a pointer so that it can be updated under the builder without
 	// having to set it in the mode after initializing targets.
 	const TArray<TObjectPtr<UUVEditorToolMeshInput>>* Targets = nullptr;
-
-protected:
-	virtual const FToolTargetTypeRequirements& GetTargetRequirements() const;
 };
 
 
diff --git a/Engine/Plugins/Experimental/UVEditor/Source/UVEditorTools/Public/UVSelectTool.h b/Engine/Plugins/Experimental/UVEditor/Source/UVEditorTools/Public/UVSelectTool.h
--- a/Engine/Plugins/Experimental/UVEditor/Source/UVEditorTools/Public/UVSelectTool.h
+++ b/Engine/Plugins/Experimental/UVEditor/Source/UVEditorTools/Public/UVSelectTool.h
@@ -45,9 +45,6 @@ public:
 	// This is a pointer so that it can be updated under the builder without
 	// having to set it in the mode after initializing targets.
 	const TArray<TObjectPtr<UUVEditorToolMeshInput>>* Targets = nullptr;
-
-protected:
-	virtual const FToolTargetTypeRequirements& GetTargetRequirements() const;
 };
 
 UENUM()
diff --git a/Engine/Plugins/Experimental/UVEditor/Source/UVEditorToolsEditorOnly/Private/UVEditorParameterizeMeshTool.cpp b/Engine/Plugins/Experimental/UVEditor/Source/UVEditorToolsEditorOnly/Private/UVEditorParameterizeMeshTool.cpp
--- a/Engine/Plugins/Experimental/UVEditor/Source/UVEditorToolsEditorOnly/Private/UVEditorParameterizeMeshTool.cpp
+++ b/Engine/Plugins/Experimental/UVEditor/Source/UVEditorToolsEditorOnly/Private/UVEditorParameterizeMeshTool.cpp
@@ -21,12 +21,6 @@ using namespace UE::Geometry;
 // Tool builder
 // TODO: Could consider sharing some of the tool builder boilerplate for UV editor tools in a common base class.
 
-const FToolTargetTypeRequirements& UUVEditorParameterizeMeshToolBuilder::GetTargetRequirements() const
-{
-	static FToolTargetTypeRequirements TypeRequirements(UUVEditorToolMeshInput::StaticClass());
-	return TypeRequirements;
-}
-
 bool UUVEditorParameterizeMeshToolBuilder::CanBuildTool(const FToolBuilderState& SceneState) const
 {
 	return Targets && Targets->Num() >= 1;
diff --git a/Engine/Plugins/Experimental/UVEditor/Source/UVEditorToolsEditorOnly/Public/UVEditorParameterizeMeshTool.h b/Engine/Plugins/Experimental/UVEditor/Source/UVEditorToolsEditorOnly/Public/UVEditorParameterizeMeshTool.h
--- a/Engine/Plugins/Experimental/UVEditor/Source/UVEditorToolsEditorOnly/Public/UVEditorParameterizeMeshTool.h
+++ b/Engine/Plugins/Experimental/UVEditor/Source/UVEditorToolsEditorOnly/Public/UVEditorParameterizeMeshTool.h
@@ -36,9 +36,6 @@ public:
 	// This is a pointer so that it can be updated under the builder without
 	// having to set it in the mode after initializing targets.
 	const TArray<TObjectPtr<UUVEditorToolMeshInput>>* Targets = nullptr;
-
-protected:
-	virtual const FToolTargetTypeRequirements& GetTargetRequirements() const;
 };
 
 
diff --git a/Engine/Plugins/Runtime/MeshModelingToolset/Source/ModelingComponents/Public/BaseTools/BaseMeshProcessingTool.h b/Engine/Plugins/Runtime/MeshModelingToolset/Source/ModelingComponents/Public/BaseTools/BaseMeshProcessingTool.h
--- a/Engine/Plugins/Runtime/MeshModelingToolset/Source/ModelingComponents/Public/BaseTools/BaseMeshProcessingTool.h
+++ b/Engine/Plugins/Runtime/MeshModelingToolset/Source/ModelingComponents/Public/BaseTools/BaseMeshProcessingTool.h
@@ -4,7 +4,7 @@
 
 #include "CoreMinimal.h"
 #include "SingleSelectionTool.h"
-#include "InteractiveToolBuilder.h"
+#include "InteractiveToolWithToolTargetsBuilder.h"
 #include "DynamicMesh/DynamicMesh3.h"
 #include "WeightMapTypes.h"
 #include "MeshOpPreviewHelpers.h"
@@ -15,6 +15,7 @@
 struct FMeshDescription;
 class UDynamicMeshComponent;
 class UPreviewMesh;
+class UBaseMeshProcessingTool;
 PREDECLARE_USE_GEOMETRY_CLASS(FMeshBoundaryLoops);
 PREDECLARE_USE_GEOMETRY_CLASS(FMeshNormals);
 using UE::Geometry::FDynamicMesh3;
@@ -23,7 +24,7 @@ using UE::Geometry::FIndexedWeightMap1f;
  * ToolBuilder for UBaseMeshProcessingTool
  */
 UCLASS()
-class MODELINGCOMPONENTS_API UBaseMeshProcessingToolBuilder : public UInteractiveToolBuilder
+class MODELINGCOMPONENTS_API UBaseMeshProcessingToolBuilder : public UInteractiveToolWithToolTargetsBuilder
 {
 	GENERATED_BODY()
 
@@ -36,7 +37,7 @@ public:
 	virtual UBaseMeshProcessingTool* MakeNewToolInstance(UObject* Outer) const { check(false); return nullptr; }
 
 protected:
-	virtual const FToolTargetTypeRequirements& GetTargetRequirements() const;
+	virtual const FToolTargetTypeRequirements& GetTargetRequirements() const override;
 };
 
 
diff --git a/Engine/Plugins/Runtime/MeshModelingToolset/Source/ModelingComponents/Public/BaseTools/MultiSelectionMeshEditingTool.h b/Engine/Plugins/Runtime/MeshModelingToolset/Source/ModelingComponents/Public/BaseTools/MultiSelectionMeshEditingTool.h
--- a/Engine/Plugins/Runtime/MeshModelingToolset/Source/ModelingComponents/Public/BaseTools/MultiSelectionMeshEditingTool.h
+++ b/Engine/Plugins/Runtime/MeshModelingToolset/Source/ModelingComponents/Public/BaseTools/MultiSelectionMeshEditingTool.h
@@ -3,7 +3,7 @@
 #pragma once
 
 #include "MultiSelectionTool.h"
-#include "InteractiveToolBuilder.h"
+#include "InteractiveToolWithToolTargetsBuilder.h"
 #include "MultiSelectionMeshEditingTool.generated.h"
 
 class UMultiSelectionMeshEditingTool;
@@ -14,7 +14,7 @@ class UMultiSelectionMeshEditingTool;
  * for editing meshes.
  */
 UCLASS(Transient, Abstract)
-class MODELINGCOMPONENTS_API UMultiSelectionMeshEditingToolBuilder : public UInteractiveToolBuilder
+class MODELINGCOMPONENTS_API UMultiSelectionMeshEditingToolBuilder : public UInteractiveToolWithToolTargetsBuilder
 {
 	GENERATED_BODY()
 
@@ -32,7 +32,7 @@ public:
 	virtual void InitializeNewTool(UMultiSelectionMeshEditingTool* Tool, const FToolBuilderState& SceneState) const;
 
 protected:
-	virtual const FToolTargetTypeRequirements& GetTargetRequirements() const;
+	virtual const FToolTargetTypeRequirements& GetTargetRequirements() const override;
 };
 
 
diff --git a/Engine/Plugins/Runtime/MeshModelingToolset/Source/ModelingComponents/Public/BaseTools/SingleSelectionMeshEditingTool.h b/Engine/Plugins/Runtime/MeshModelingToolset/Source/ModelingComponents/Public/BaseTools/SingleSelectionMeshEditingTool.h
--- a/Engine/Plugins/Runtime/MeshModelingToolset/Source/ModelingComponents/Public/BaseTools/SingleSelectionMeshEditingTool.h
+++ b/Engine/Plugins/Runtime/MeshModelingToolset/Source/ModelingComponents/Public/BaseTools/SingleSelectionMeshEditingTool.h
@@ -3,7 +3,7 @@
 #pragma once
 
 #include "SingleSelectionTool.h"
-#include "InteractiveToolBuilder.h"
+#include "InteractiveToolWithToolTargetsBuilder.h"
 #include "SingleSelectionMeshEditingTool.generated.h"
 
 class USingleSelectionMeshEditingTool;
@@ -15,7 +15,7 @@ class UPersistentMeshSelection;
  * for editing meshes.
  */
 UCLASS(Transient, Abstract)
-class MODELINGCOMPONENTS_API USingleSelectionMeshEditingToolBuilder : public UInteractiveToolBuilder
+class MODELINGCOMPONENTS_API USingleSelectionMeshEditingToolBuilder : public UInteractiveToolWithToolTargetsBuilder
 {
 	GENERATED_BODY()
 
@@ -36,7 +36,7 @@ public:
 	virtual bool WantsInputSelectionIfAvailable() const { return false; }
 
 protected:
-	virtual const FToolTargetTypeRequirements& GetTargetRequirements() const;
+	virtual const FToolTargetTypeRequirements& GetTargetRequirements() const override;
 };
 
 
diff --git a/Engine/Source/Runtime/InteractiveToolsFramework/Public/BaseTools/MeshSurfacePointTool.h b/Engine/Source/Runtime/InteractiveToolsFramework/Public/BaseTools/MeshSurfacePointTool.h
--- a/Engine/Source/Runtime/InteractiveToolsFramework/Public/BaseTools/MeshSurfacePointTool.h
+++ b/Engine/Source/Runtime/InteractiveToolsFramework/Public/BaseTools/MeshSurfacePointTool.h
@@ -4,7 +4,7 @@
 
 #include "CoreMinimal.h"
 #include "SingleSelectionTool.h"
-#include "InteractiveToolBuilder.h"
+#include "InteractiveToolWithToolTargetsBuilder.h"
 #include "InputBehaviorSet.h"
 #include "BaseBehaviors/BehaviorTargetInterfaces.h"
 #include "BaseBehaviors/AnyButtonInputBehavior.h"
@@ -26,7 +26,7 @@ public:
  * 
  */
 UCLASS()
-class INTERACTIVETOOLSFRAMEWORK_API UMeshSurfacePointToolBuilder : public UInteractiveToolBuilder
+class INTERACTIVETOOLSFRAMEWORK_API UMeshSurfacePointToolBuilder : public UInteractiveToolWithToolTargetsBuilder
 {
 	GENERATED_BODY()
 
@@ -46,7 +46,7 @@ public:
 	virtual void InitializeNewTool(UMeshSurfacePointTool* Tool, const FToolBuilderState& SceneState) const;
 
 protected:
-	virtual const FToolTargetTypeRequirements& GetTargetRequirements() const;
+	virtual const FToolTargetTypeRequirements& GetTargetRequirements() const override;
 };
 
 
diff --git a/Engine/Source/Runtime/InteractiveToolsFramework/Public/InteractiveToolWithToolTargetsBuilder.h b/Engine/Source/Runtime/InteractiveToolsFramework/Public/InteractiveToolWithToolTargetsBuilder.h
new file mode 100644
--- /dev/null
+++ b/Engine/Source/Runtime/InteractiveToolsFramework/Public/InteractiveToolWithToolTargetsBuilder.h
@@ -0,0 +1,34 @@
+// Copyright Epic Games, Inc. All Rights Reserved.
+
+#pragma once
+
+#include "CoreMinimal.h"
+#include "InteractiveToolBuilder.h"
+#include "ToolTargets/ToolTarget.h"
+
+#include "InteractiveToolWithToolTargetsBuilder.generated.h"
+
+
+
+/**
+ * A UInteractiveToolWithToolTargetsBuilder creates a new instance of an InteractiveTool that uses tool targets.
+ * See ToolTarget.h for more information on tool targets and their usage. This class defines the common
+ * interface(s) for defining the tool target requirements of the tool it builds.
+ * This is an abstract base class, you must subclass it in order to create your particular Tool instance.
+ */
+UCLASS(Transient, Abstract)
+class INTERACTIVETOOLSFRAMEWORK_API UInteractiveToolWithToolTargetsBuilder : public UInteractiveToolBuilder
+{
+	GENERATED_BODY()
+
+protected:
+	/**
+	 * Gives the target requirements of the associated tool. Usually, it is the tool builder
+	 * will use this function in CanBuildTool and BuildTool to find and create any necessary targets.
+	 */
+	virtual const FToolTargetTypeRequirements& GetTargetRequirements() const
+	{
+		static FToolTargetTypeRequirements TypeRequirements; // Default initialized to no requirements.
+		return TypeRequirements;
+	}
+};
