ModelingTools: Fix CIS break with BaseCreateFromSelectedTool.

#rnx

diff --git a/Engine/Plugins/Runtime/MeshModelingToolset/Source/ModelingComponents/Private/BaseTools/BaseCreateFromSelectedTool.cpp b/Engine/Plugins/Runtime/MeshModelingToolset/Source/ModelingComponents/Private/BaseTools/BaseCreateFromSelectedTool.cpp
--- a/Engine/Plugins/Runtime/MeshModelingToolset/Source/ModelingComponents/Private/BaseTools/BaseCreateFromSelectedTool.cpp
+++ b/Engine/Plugins/Runtime/MeshModelingToolset/Source/ModelingComponents/Private/BaseTools/BaseCreateFromSelectedTool.cpp
@@ -35,6 +35,12 @@ bool UBaseCreateFromSelectedToolBuilder::CanBuildTool(const FToolBuilderState& S
 	return ComponentCount >= MinComponentsSupported() && (!MaxComponentsSupported().IsSet() || ComponentCount <= MaxComponentsSupported().GetValue());
 }
 
+UMultiSelectionMeshEditingTool* UBaseCreateFromSelectedToolBuilder::CreateNewTool(const FToolBuilderState& SceneState) const
+{
+	// Subclasses must override this method.
+	check(false);
+	return nullptr;
+}
 
 
 /*
diff --git a/Engine/Plugins/Runtime/MeshModelingToolset/Source/ModelingComponents/Public/BaseTools/BaseCreateFromSelectedTool.h b/Engine/Plugins/Runtime/MeshModelingToolset/Source/ModelingComponents/Public/BaseTools/BaseCreateFromSelectedTool.h
--- a/Engine/Plugins/Runtime/MeshModelingToolset/Source/ModelingComponents/Public/BaseTools/BaseCreateFromSelectedTool.h
+++ b/Engine/Plugins/Runtime/MeshModelingToolset/Source/ModelingComponents/Public/BaseTools/BaseCreateFromSelectedTool.h
@@ -26,6 +26,8 @@ class MODELINGCOMPONENTS_API UBaseCreateFromSelectedToolBuilder : public UMultiS
 public:
 	virtual bool CanBuildTool(const FToolBuilderState& SceneState) const override;
 
+	virtual UMultiSelectionMeshEditingTool* CreateNewTool(const FToolBuilderState& SceneState) const override;
+
 public:
 	virtual TOptional<int32> MaxComponentsSupported() const { return TOptional<int32>(); }
 	virtual int32 MinComponentsSupported() const { return 1; }
diff --git a/Engine/Plugins/Runtime/MeshModelingToolset/Source/ModelingComponents/Public/BaseTools/MultiSelectionMeshEditingTool.h b/Engine/Plugins/Runtime/MeshModelingToolset/Source/ModelingComponents/Public/BaseTools/MultiSelectionMeshEditingTool.h
--- a/Engine/Plugins/Runtime/MeshModelingToolset/Source/ModelingComponents/Public/BaseTools/MultiSelectionMeshEditingTool.h
+++ b/Engine/Plugins/Runtime/MeshModelingToolset/Source/ModelingComponents/Public/BaseTools/MultiSelectionMeshEditingTool.h
@@ -26,7 +26,7 @@ public:
 	virtual UInteractiveTool* BuildTool(const FToolBuilderState& SceneState) const override;
 
 	/** @return new Tool instance. Override this in subclasses to build a different Tool class type */
-	virtual UMultiSelectionMeshEditingTool* CreateNewTool(const FToolBuilderState& SceneState) const PURE_VIRTUAL(USingleSelectionMeshEditingToolBuilder::CreateNewTool, return nullptr; );
+	virtual UMultiSelectionMeshEditingTool* CreateNewTool(const FToolBuilderState& SceneState) const PURE_VIRTUAL(UMultiSelectionMeshEditingToolBuilder::CreateNewTool, return nullptr; );
 
 	/** Called by BuildTool to configure the Tool with the input mesh source(s) based on the SceneState */
 	virtual void InitializeNewTool(UMultiSelectionMeshEditingTool* Tool, const FToolBuilderState& SceneState) const;
