ModelingTools: Add read only skeletal mesh tool target.
Enable read-only skeletal mesh tool targets in modeling mode.

#rnx

diff --git a/Engine/Plugins/Editor/ModelingToolsEditorMode/Source/ModelingToolsEditorMode/Private/ModelingToolsEditorMode.cpp b/Engine/Plugins/Editor/ModelingToolsEditorMode/Source/ModelingToolsEditorMode/Private/ModelingToolsEditorMode.cpp
--- a/Engine/Plugins/Editor/ModelingToolsEditorMode/Source/ModelingToolsEditorMode/Private/ModelingToolsEditorMode.cpp
+++ b/Engine/Plugins/Editor/ModelingToolsEditorMode/Source/ModelingToolsEditorMode/Private/ModelingToolsEditorMode.cpp
@@ -320,10 +320,9 @@ void UModelingToolsEditorMode::Enter()
 	GetInteractiveToolsContext()->TargetManager->AddTargetFactory(NewObject<UVolumeMeshDescriptionToolTargetFactory>(GetToolManager()));
 	GetInteractiveToolsContext()->TargetManager->AddTargetFactory(NewObject<UDynamicMeshComponentToolTargetFactory>(GetToolManager()));
 
-	// TODO: Re-enable skeletal mesh target support when the tool target system has a better way to filter/exclude targets from certain tools.
-	// Skeletal mesh targets are disabled since many modeling tools would break skin weights. In the interim, tools
-	// that require skeletal mesh support should manually handle it at the tool level.
-	//GetInteractiveToolsContext()->TargetManager->AddTargetFactory(NewObject<USkeletalMeshComponentToolTargetFactory>(GetToolManager()));
+	// Register read-only skeletal mesh tool targets. Currently tools that write to meshes risk breaking
+	// skin weights.
+	GetInteractiveToolsContext()->TargetManager->AddTargetFactory(NewObject<USkeletalMeshComponentReadOnlyToolTargetFactory>(GetToolManager()));
 
 	// register stylus event handler
 	StylusStateTracker = MakeUnique<FStylusStateTracker>();
diff --git a/Engine/Plugins/Runtime/MeshModelingToolset/Source/ModelingComponentsEditorOnly/Private/ToolTargets/SkeletalMeshComponentToolTarget.cpp b/Engine/Plugins/Runtime/MeshModelingToolset/Source/ModelingComponentsEditorOnly/Private/ToolTargets/SkeletalMeshComponentToolTarget.cpp
--- a/Engine/Plugins/Runtime/MeshModelingToolset/Source/ModelingComponentsEditorOnly/Private/ToolTargets/SkeletalMeshComponentToolTarget.cpp
+++ b/Engine/Plugins/Runtime/MeshModelingToolset/Source/ModelingComponentsEditorOnly/Private/ToolTargets/SkeletalMeshComponentToolTarget.cpp
@@ -17,17 +17,22 @@ namespace USkeletalMeshComponentToolTargetLocals
 }
 
 
-int32 USkeletalMeshComponentToolTarget::GetNumMaterials() const
+//
+// USkeletalMeshComponentReadOnlyToolTarget
+//
+
+
+int32 USkeletalMeshComponentReadOnlyToolTarget::GetNumMaterials() const
 {
 	return ensure(IsValid()) ? Component->GetNumMaterials() : 0;
 }
 
-UMaterialInterface* USkeletalMeshComponentToolTarget::GetMaterial(int32 MaterialIndex) const
+UMaterialInterface* USkeletalMeshComponentReadOnlyToolTarget::GetMaterial(int32 MaterialIndex) const
 {
 	return ensure(IsValid()) ? Component->GetMaterial(MaterialIndex) : nullptr;
 }
 
-void USkeletalMeshComponentToolTarget::GetMaterialSet(FComponentMaterialSet& MaterialSetOut, bool bPreferAssetMaterials) const
+void USkeletalMeshComponentReadOnlyToolTarget::GetMaterialSet(FComponentMaterialSet& MaterialSetOut, bool bPreferAssetMaterials) const
 {
 	if (!ensure(IsValid())) return;
 
@@ -47,7 +52,7 @@ void USkeletalMeshComponentToolTarget::GetMaterialSet(FComponentMaterialSet& Mat
 	}
 }
 
-bool USkeletalMeshComponentToolTarget::CommitMaterialSetUpdate(const FComponentMaterialSet& MaterialSet, bool bApplyToAsset)
+bool USkeletalMeshComponentReadOnlyToolTarget::CommitMaterialSetUpdate(const FComponentMaterialSet& MaterialSet, bool bApplyToAsset)
 {
 	if (!ensure(IsValid())) return false;
 
@@ -67,7 +72,7 @@ bool USkeletalMeshComponentToolTarget::CommitMaterialSetUpdate(const FComponentM
 
 		// We wrote the below code to support a mismatch in the number of materials.
 		// However, it is not yet clear whether this might be desirable, and we don't
-		// want to inadvertantly hide bugs in the meantime. So, we keep this ensure here
+		// want to inadvertently hide bugs in the meantime. So, we keep this ensure here
 		// for now, and we can remove it if we decide that we want the ability.
 		ensure(NumMaterialsNeeded == NumMaterialsGiven);
 
@@ -83,7 +88,7 @@ bool USkeletalMeshComponentToolTarget::CommitMaterialSetUpdate(const FComponentM
 	return true;
 }
 
-const FMeshDescription* USkeletalMeshComponentToolTarget::GetMeshDescription()
+const FMeshDescription* USkeletalMeshComponentReadOnlyToolTarget::GetMeshDescription()
 {
 	if (!ensure(IsValid()))
 	{
@@ -100,6 +105,24 @@ const FMeshDescription* USkeletalMeshComponentToolTarget::GetMeshDescription()
 	return CachedMeshDescription.Get();
 }
 
+
+TSharedPtr<FDynamicMesh3, ESPMode::ThreadSafe> USkeletalMeshComponentReadOnlyToolTarget::GetDynamicMesh()
+{
+	return GetDynamicMeshViaMeshDescription(*this);
+}
+
+
+USkeletalMesh* USkeletalMeshComponentReadOnlyToolTarget::GetSkeletalMesh() const
+{
+	return IsValid() ? Cast<USkinnedMeshComponent>(Component)->SkeletalMesh : nullptr;
+}
+
+
+//
+// USkeletalMeshComponentToolTarget
+//
+
+
 void USkeletalMeshComponentToolTarget::CommitMeshDescription(const FCommitter& Committer)
 {
 	if (ensure(IsValid()) == false) return;
@@ -122,23 +145,37 @@ void USkeletalMeshComponentToolTarget::CommitMeshDescription(const FCommitter& C
 	CachedMeshDescription.Reset();
 }
 
-TSharedPtr<FDynamicMesh3, ESPMode::ThreadSafe> USkeletalMeshComponentToolTarget::GetDynamicMesh()
-{
-	return GetDynamicMeshViaMeshDescription(*this);
-}
 
 void USkeletalMeshComponentToolTarget::CommitDynamicMesh(const FDynamicMesh3& Mesh, const FDynamicMeshCommitInfo& CommitInfo)
 {
 	CommitDynamicMeshViaMeshDescription(*this, Mesh, CommitInfo);
 }
 
-USkeletalMesh* USkeletalMeshComponentToolTarget::GetSkeletalMesh() const
+
+// Factory
+
+bool USkeletalMeshComponentReadOnlyToolTargetFactory::CanBuildTarget(UObject* SourceObject, const FToolTargetTypeRequirements& Requirements) const
 {
-	return IsValid() ? Cast<USkinnedMeshComponent>(Component)->SkeletalMesh : nullptr;
+	// We are using an exact cast here to prevent subclasses, which might not meet all
+	// requirements for functionality such as the deprecated DestructibleMeshComponent, from 
+	// being caught up as valid targets.
+	// If you want to make the tool target work with some subclass of USkeletalMeshComponent,
+	// just add another factory that allows that class specifically(but make sure that
+	// GetMeshDescription and such work properly)
+
+	return Cast<USkinnedMeshComponent>(SourceObject) && Cast<USkinnedMeshComponent>(SourceObject)->SkeletalMesh &&
+		ExactCast<USkeletalMesh>(Cast<USkinnedMeshComponent>(SourceObject)->SkeletalMesh) &&
+		Requirements.AreSatisfiedBy(USkeletalMeshComponentReadOnlyToolTarget::StaticClass());
 }
 
+UToolTarget* USkeletalMeshComponentReadOnlyToolTargetFactory::BuildTarget(UObject* SourceObject, const FToolTargetTypeRequirements& Requirements)
+{
+	USkeletalMeshComponentReadOnlyToolTarget* Target = NewObject<USkeletalMeshComponentReadOnlyToolTarget>();
+	Target->Component = Cast<USkinnedMeshComponent>(SourceObject);
+	check(Target->Component && Requirements.AreSatisfiedBy(Target));
 
-// Factory
+	return Target;
+}
 
 bool USkeletalMeshComponentToolTargetFactory::CanBuildTarget(UObject* SourceObject, const FToolTargetTypeRequirements& Requirements) const
 {
diff --git a/Engine/Plugins/Runtime/MeshModelingToolset/Source/ModelingComponentsEditorOnly/Private/ToolTargets/SkeletalMeshToolTarget.cpp b/Engine/Plugins/Runtime/MeshModelingToolset/Source/ModelingComponentsEditorOnly/Private/ToolTargets/SkeletalMeshToolTarget.cpp
--- a/Engine/Plugins/Runtime/MeshModelingToolset/Source/ModelingComponentsEditorOnly/Private/ToolTargets/SkeletalMeshToolTarget.cpp
+++ b/Engine/Plugins/Runtime/MeshModelingToolset/Source/ModelingComponentsEditorOnly/Private/ToolTargets/SkeletalMeshToolTarget.cpp
@@ -16,29 +16,35 @@ namespace USkeletalMeshToolTargetLocals
 	int32 LODIndex = 0;
 }
 
-bool USkeletalMeshToolTarget::IsValid() const
+
+//
+// USkeletalMeshReadOnlyToolTarget
+//
+
+
+bool USkeletalMeshReadOnlyToolTarget::IsValid() const
 {
 	return SkeletalMesh && IsValidChecked(SkeletalMesh) && !SkeletalMesh->IsUnreachable() && SkeletalMesh->IsValidLowLevel();
 }
 
-int32 USkeletalMeshToolTarget::GetNumMaterials() const
+int32 USkeletalMeshReadOnlyToolTarget::GetNumMaterials() const
 {
 	return ensure(IsValid()) ? SkeletalMesh->GetMaterials().Num() : 0;
 }
 
-UMaterialInterface* USkeletalMeshToolTarget::GetMaterial(int32 MaterialIndex) const
+UMaterialInterface* USkeletalMeshReadOnlyToolTarget::GetMaterial(int32 MaterialIndex) const
 {
 	return ensure(IsValid() && MaterialIndex < SkeletalMesh->GetMaterials().Num()) ? 
 		SkeletalMesh->GetMaterials()[MaterialIndex].MaterialInterface : nullptr;
 }
 
-void USkeletalMeshToolTarget::GetMaterialSet(FComponentMaterialSet& MaterialSetOut, bool bPreferAssetMaterials) const
+void USkeletalMeshReadOnlyToolTarget::GetMaterialSet(FComponentMaterialSet& MaterialSetOut, bool bPreferAssetMaterials) const
 {
 	if (!ensure(IsValid())) return;
 	GetMaterialSet(SkeletalMesh, MaterialSetOut, bPreferAssetMaterials);
 }
 
-void USkeletalMeshToolTarget::GetMaterialSet(const USkeletalMesh* SkeletalMeshIn, FComponentMaterialSet& MaterialSetOut,
+void USkeletalMeshReadOnlyToolTarget::GetMaterialSet(const USkeletalMesh* SkeletalMeshIn, FComponentMaterialSet& MaterialSetOut,
 	bool bPreferAssetMaterials)
 {
 	const TArray<FSkeletalMaterial>& Materials = SkeletalMeshIn->GetMaterials(); 
@@ -49,13 +55,13 @@ void USkeletalMeshToolTarget::GetMaterialSet(const USkeletalMesh* SkeletalMeshIn
 	}
 }
 
-bool USkeletalMeshToolTarget::CommitMaterialSetUpdate(const FComponentMaterialSet& MaterialSet, bool bApplyToAsset)
+bool USkeletalMeshReadOnlyToolTarget::CommitMaterialSetUpdate(const FComponentMaterialSet& MaterialSet, bool bApplyToAsset)
 {
 	if (!ensure(IsValid())) return false;
 	return CommitMaterialSetUpdate(SkeletalMesh, MaterialSet, bApplyToAsset);
 }
 
-bool USkeletalMeshToolTarget::CommitMaterialSetUpdate(USkeletalMesh* SkeletalMeshIn, 
+bool USkeletalMeshReadOnlyToolTarget::CommitMaterialSetUpdate(USkeletalMesh* SkeletalMeshIn, 
 	const FComponentMaterialSet& MaterialSet, bool bApplyToAsset)
 {
 	if (!bApplyToAsset)
@@ -100,7 +106,7 @@ bool USkeletalMeshToolTarget::CommitMaterialSetUpdate(USkeletalMesh* SkeletalMes
 	return true;
 }
 
-const FMeshDescription* USkeletalMeshToolTarget::GetMeshDescription()
+const FMeshDescription* USkeletalMeshReadOnlyToolTarget::GetMeshDescription()
 {
 	if (!ensure(IsValid()))
 	{
@@ -116,7 +122,7 @@ const FMeshDescription* USkeletalMeshToolTarget::GetMeshDescription()
 	return CachedMeshDescription.Get();
 }
 
-void USkeletalMeshToolTarget::GetMeshDescription(const USkeletalMesh* SkeletalMeshIn, FMeshDescription& MeshDescription)
+void USkeletalMeshReadOnlyToolTarget::GetMeshDescription(const USkeletalMesh* SkeletalMeshIn, FMeshDescription& MeshDescription)
 {
 	using namespace USkeletalMeshToolTargetLocals;
 
@@ -140,6 +146,20 @@ void USkeletalMeshToolTarget::GetMeshDescription(const USkeletalMesh* SkeletalMe
 	}
 }
 
+TSharedPtr<FDynamicMesh3, ESPMode::ThreadSafe> USkeletalMeshReadOnlyToolTarget::GetDynamicMesh()
+{
+	return GetDynamicMeshViaMeshDescription(*this);
+}
+
+USkeletalMesh* USkeletalMeshReadOnlyToolTarget::GetSkeletalMesh() const
+{
+	return IsValid() ? SkeletalMesh : nullptr;
+}
+
+//
+// USkeletalMeshToolTarget
+//
+
 void USkeletalMeshToolTarget::CommitMeshDescription(const FCommitter& Committer)
 {
 	if (ensure(IsValid()) == false) return;
@@ -194,23 +214,35 @@ void USkeletalMeshToolTarget::CommitMeshDescription(USkeletalMesh* SkeletalMeshI
 	SkeletalMeshIn->PostEditChange();
 }
 
-TSharedPtr<FDynamicMesh3, ESPMode::ThreadSafe> USkeletalMeshToolTarget::GetDynamicMesh()
-{
-	return GetDynamicMeshViaMeshDescription(*this);
-}
-
 void USkeletalMeshToolTarget::CommitDynamicMesh(const FDynamicMesh3& Mesh, const FDynamicMeshCommitInfo& CommitInfo)
 {
 	CommitDynamicMeshViaMeshDescription(*this, Mesh, CommitInfo);
 }
 
-USkeletalMesh* USkeletalMeshToolTarget::GetSkeletalMesh() const
+
+// Factory
+
+bool USkeletalMeshReadOnlyToolTargetFactory::CanBuildTarget(UObject* SourceObject, const FToolTargetTypeRequirements& Requirements) const
 {
-	return IsValid() ? SkeletalMesh : nullptr;
+	// We are using an exact cast here to prevent subclasses, which might not meet all
+	// requirements for functionality such as the deprecated DestructibleMesh, from 
+	// being caught up as valid targets.
+	// If you want to make the tool target work with some subclass of USkeletalMesh,
+	// just add another factory that allows that class specifically (but make sure that
+	// GetMeshDescription and such work properly)
+
+	return ExactCast<USkeletalMesh>(SourceObject) 
+		&& Requirements.AreSatisfiedBy(USkeletalMeshReadOnlyToolTarget::StaticClass());
 }
 
+UToolTarget* USkeletalMeshReadOnlyToolTargetFactory::BuildTarget(UObject* SourceObject, const FToolTargetTypeRequirements& Requirements)
+{
+	USkeletalMeshReadOnlyToolTarget* Target = NewObject<USkeletalMeshReadOnlyToolTarget>();
+	Target->SkeletalMesh = Cast<USkeletalMesh>(SourceObject);
+	check(Target->SkeletalMesh && Requirements.AreSatisfiedBy(Target));
 
-// Factory
+	return Target;
+}
 
 bool USkeletalMeshToolTargetFactory::CanBuildTarget(UObject* SourceObject, const FToolTargetTypeRequirements& Requirements) const
 {
diff --git a/Engine/Plugins/Runtime/MeshModelingToolset/Source/ModelingComponentsEditorOnly/Public/ToolTargets/SkeletalMeshComponentToolTarget.h b/Engine/Plugins/Runtime/MeshModelingToolset/Source/ModelingComponentsEditorOnly/Public/ToolTargets/SkeletalMeshComponentToolTarget.h
--- a/Engine/Plugins/Runtime/MeshModelingToolset/Source/ModelingComponentsEditorOnly/Public/ToolTargets/SkeletalMeshComponentToolTarget.h
+++ b/Engine/Plugins/Runtime/MeshModelingToolset/Source/ModelingComponentsEditorOnly/Public/ToolTargets/SkeletalMeshComponentToolTarget.h
@@ -17,17 +17,16 @@
 
 class USkeletalMesh;
 
+
 /**
- * A tool target backed by a skeletal mesh component that can provide and take a mesh
+ * A tool target backed by a read-only skeletal mesh component that can provide and take a mesh
  * description.
  */
 UCLASS(Transient)
-class MODELINGCOMPONENTSEDITORONLY_API USkeletalMeshComponentToolTarget :
+class MODELINGCOMPONENTSEDITORONLY_API USkeletalMeshComponentReadOnlyToolTarget :
 	public UPrimitiveComponentToolTarget,
-	public IMeshDescriptionCommitter,
 	public IMeshDescriptionProvider,
 	public IDynamicMeshProvider, 
-	public IDynamicMeshCommitter,
 	public IMaterialProvider,
 	public ISkeletalMeshBackedTarget
 {
@@ -37,36 +36,69 @@ public:
 	// IMeshDescriptionProvider implementation
 	const FMeshDescription* GetMeshDescription() override;
 
-	// IMeshDescritpionCommitter implementation
-	void CommitMeshDescription(const FCommitter& Committer) override;
-	using IMeshDescriptionCommitter::CommitMeshDescription; // unhide the other overload
-
 	// IMaterialProvider implementation
 	int32 GetNumMaterials() const override;
 	UMaterialInterface* GetMaterial(int32 MaterialIndex) const override;
 	void GetMaterialSet(FComponentMaterialSet& MaterialSetOut, bool bPreferAssetMaterials) const override;
-	bool CommitMaterialSetUpdate(const FComponentMaterialSet& MaterialSet, bool bApplyToAsset) override;	
+	bool CommitMaterialSetUpdate(const FComponentMaterialSet& MaterialSet, bool bApplyToAsset) override;
 
 	// IDynamicMeshProvider
 	virtual TSharedPtr<UE::Geometry::FDynamicMesh3, ESPMode::ThreadSafe> GetDynamicMesh() override;
 
-	// IDynamicMeshCommitter
-	virtual void CommitDynamicMesh(const UE::Geometry::FDynamicMesh3& Mesh, const FDynamicMeshCommitInfo& CommitInfo) override;
-	using IDynamicMeshCommitter::CommitDynamicMesh; // unhide the other overload
-
 	// ISkeletalMeshBackedTarget implementation
 	USkeletalMesh* GetSkeletalMesh() const override;
 
 protected:
 	// So that the tool target factory can poke into Component.
-	friend class USkeletalMeshComponentToolTargetFactory;
+	friend class USkeletalMeshComponentReadOnlyToolTargetFactory;
 
-private:
 	// Until USkeletalMesh stores its internal representation as FMeshDescription, we need to
 	// retain the storage here to cover the lifetime of the pointer returned by GetMeshDescription(). 
 	TUniquePtr<FMeshDescription> CachedMeshDescription;
 };
 
+
+/**
+ * A tool target backed by a skeletal mesh component that can provide and take a mesh
+ * description.
+ */
+UCLASS(Transient)
+class MODELINGCOMPONENTSEDITORONLY_API USkeletalMeshComponentToolTarget :
+	public USkeletalMeshComponentReadOnlyToolTarget,
+	public IMeshDescriptionCommitter,
+	public IDynamicMeshCommitter
+{
+	GENERATED_BODY()
+
+public:
+	// IMeshDescriptionCommitter implementation
+	void CommitMeshDescription(const FCommitter& Committer) override;
+	using IMeshDescriptionCommitter::CommitMeshDescription; // unhide the other overload
+
+	// IDynamicMeshCommitter
+	virtual void CommitDynamicMesh(const UE::Geometry::FDynamicMesh3& Mesh, const FDynamicMeshCommitInfo& CommitInfo) override;
+	using IDynamicMeshCommitter::CommitDynamicMesh; // unhide the other overload
+
+protected:
+	// So that the tool target factory can poke into Component.
+	friend class USkeletalMeshComponentToolTargetFactory;
+};
+
+
+/** Factory for USkeletalMeshComponentReadOnlyToolTarget to be used by the target manager. */
+UCLASS(Transient)
+class MODELINGCOMPONENTSEDITORONLY_API USkeletalMeshComponentReadOnlyToolTargetFactory : public UToolTargetFactory
+{
+	GENERATED_BODY()
+
+public:
+
+	bool CanBuildTarget(UObject* SourceObject, const FToolTargetTypeRequirements& TargetTypeInfo) const override;
+
+	UToolTarget* BuildTarget(UObject* SourceObject, const FToolTargetTypeRequirements& TargetTypeInfo) override;
+};
+
+
 /** Factory for USkeletalMeshComponentToolTarget to be used by the target manager. */
 UCLASS(Transient)
 class MODELINGCOMPONENTSEDITORONLY_API USkeletalMeshComponentToolTargetFactory : public UToolTargetFactory
diff --git a/Engine/Plugins/Runtime/MeshModelingToolset/Source/ModelingComponentsEditorOnly/Public/ToolTargets/SkeletalMeshToolTarget.h b/Engine/Plugins/Runtime/MeshModelingToolset/Source/ModelingComponentsEditorOnly/Public/ToolTargets/SkeletalMeshToolTarget.h
--- a/Engine/Plugins/Runtime/MeshModelingToolset/Source/ModelingComponentsEditorOnly/Public/ToolTargets/SkeletalMeshToolTarget.h
+++ b/Engine/Plugins/Runtime/MeshModelingToolset/Source/ModelingComponentsEditorOnly/Public/ToolTargets/SkeletalMeshToolTarget.h
@@ -18,15 +18,13 @@
 class USkeletalMesh;
 
 /**
- * A tool target backed by a skeletal mesh.
+ * A tool target backed by a read-only skeletal mesh.
  */
 UCLASS(Transient)
-class MODELINGCOMPONENTSEDITORONLY_API USkeletalMeshToolTarget :
+class MODELINGCOMPONENTSEDITORONLY_API USkeletalMeshReadOnlyToolTarget :
 	public UToolTarget,
-	public IMeshDescriptionCommitter,
 	public IMeshDescriptionProvider,
-	public IDynamicMeshProvider, 
-	public IDynamicMeshCommitter,
+	public IDynamicMeshProvider,
 	public IMaterialProvider,
 	public ISkeletalMeshBackedTarget
 {
@@ -39,10 +37,6 @@ public:
 	// IMeshDescriptionProvider implementation
 	const FMeshDescription* GetMeshDescription() override;
 
-	// IMeshDescritpionCommitter implementation
-	void CommitMeshDescription(const FCommitter& Committer) override;
-	using IMeshDescriptionCommitter::CommitMeshDescription; // unhide the other overload
-
 	// IMaterialProvider implementation
 	int32 GetNumMaterials() const override;
 	UMaterialInterface* GetMaterial(int32 MaterialIndex) const override;
@@ -52,11 +46,6 @@ public:
 	// IDynamicMeshProvider
 	virtual TSharedPtr<UE::Geometry::FDynamicMesh3, ESPMode::ThreadSafe> GetDynamicMesh() override;
 
-	// IDynamicMeshCommitter
-	virtual void CommitDynamicMesh(const UE::Geometry::FDynamicMesh3& Mesh, 
-		const FDynamicMeshCommitInfo& CommitInfo) override;
-	using IDynamicMeshCommitter::CommitDynamicMesh; // unhide the other overload
-
 	// ISkeletalMeshBackedTarget implementation
 	USkeletalMesh* GetSkeletalMesh() const override;
 
@@ -64,25 +53,67 @@ protected:
 	USkeletalMesh* SkeletalMesh = nullptr;
 
 	// So that the tool target factory can poke into Component.
-	friend class USkeletalMeshToolTargetFactory;
-
-	
+	friend class USkeletalMeshReadOnlyToolTargetFactory;
+	friend class USkeletalMeshComponentReadOnlyToolTarget;
 	friend class USkeletalMeshComponentToolTarget;
 
 	static void GetMeshDescription(const USkeletalMesh* SkeletalMesh, FMeshDescription& MeshDescriptionOut);
-	static void CommitMeshDescription(USkeletalMesh* SkeletalMesh, 
-		FMeshDescription* MeshDescription, const FCommitter& Committer);
 	static void GetMaterialSet(const USkeletalMesh* SkeletalMesh, FComponentMaterialSet& MaterialSetOut,
 		bool bPreferAssetMaterials);
 	static bool CommitMaterialSetUpdate(USkeletalMesh* SkeletalMesh,
 		const FComponentMaterialSet& MaterialSet, bool bApplyToAsset);
-	
-private:
+
 	// Until USkeletalMesh stores its internal representation as FMeshDescription, we need to
 	// retain the storage here to cover the lifetime of the pointer returned by GetMeshDescription(). 
 	TUniquePtr<FMeshDescription> CachedMeshDescription;	
 };
 
+
+/**
+ * A tool target backed by a skeletal mesh.
+ */
+UCLASS(Transient)
+class MODELINGCOMPONENTSEDITORONLY_API USkeletalMeshToolTarget :
+	public USkeletalMeshReadOnlyToolTarget,
+	public IMeshDescriptionCommitter,
+	public IDynamicMeshCommitter
+{
+	GENERATED_BODY()
+
+public:
+	// IMeshDescriptionCommitter implementation
+	void CommitMeshDescription(const FCommitter& Committer) override;
+	using IMeshDescriptionCommitter::CommitMeshDescription; // unhide the other overload
+
+	// IDynamicMeshCommitter
+	virtual void CommitDynamicMesh(const UE::Geometry::FDynamicMesh3& Mesh, 
+		const FDynamicMeshCommitInfo& CommitInfo) override;
+	using IDynamicMeshCommitter::CommitDynamicMesh; // unhide the other overload
+
+protected:
+	// So that the tool target factory can poke into Component.
+	friend class USkeletalMeshToolTargetFactory;
+	friend class USkeletalMeshComponentToolTarget;
+
+	static void CommitMeshDescription(USkeletalMesh* SkeletalMesh, 
+		FMeshDescription* MeshDescription, const FCommitter& Committer);
+};
+
+
+/** Factory for USkeletalMeshReadOnlyToolTarget to be used by the target manager. */
+UCLASS(Transient)
+class MODELINGCOMPONENTSEDITORONLY_API USkeletalMeshReadOnlyToolTargetFactory : public UToolTargetFactory
+{
+	GENERATED_BODY()
+
+public:
+
+	bool CanBuildTarget(UObject* SourceObject, const FToolTargetTypeRequirements& TargetTypeInfo) const override;
+
+	UToolTarget* BuildTarget(UObject* SourceObject, const FToolTargetTypeRequirements& TargetTypeInfo) override;
+};
+
+
 /** Factory for USkeletalMeshToolTarget to be used by the target manager. */
 UCLASS(Transient)
 class MODELINGCOMPONENTSEDITORONLY_API USkeletalMeshToolTargetFactory : public UToolTargetFactory
