MeshModelingTools: Fix MeshElementsVisualizer not updating the bounds of MeshWireframeComponent on mesh changes. This fixes wireframe disappearing in UV editor.

#rnx

diff --git a/Engine/Plugins/Runtime/MeshModelingToolset/Source/ModelingComponents/Private/Drawing/MeshElementsVisualizer.cpp b/Engine/Plugins/Runtime/MeshModelingToolset/Source/ModelingComponents/Private/Drawing/MeshElementsVisualizer.cpp
--- a/Engine/Plugins/Runtime/MeshModelingToolset/Source/ModelingComponents/Private/Drawing/MeshElementsVisualizer.cpp
+++ b/Engine/Plugins/Runtime/MeshModelingToolset/Source/ModelingComponents/Private/Drawing/MeshElementsVisualizer.cpp
@@ -193,11 +193,11 @@ void UMeshElementsVisualizer::UpdateVisibility()
 	WireframeComponent->NormalSeamColor = Settings->NormalSeamColor;
 	WireframeComponent->ColorSeamColor = Settings->ColorSeamColor;
 
-	WireframeComponent->MarkRenderStateDirty();
+	WireframeComponent->UpdateWireframe();
 }
 
 
 void UMeshElementsVisualizer::NotifyMeshChanged()
 {
-	WireframeComponent->MarkRenderStateDirty();
+	WireframeComponent->UpdateWireframe();
 }
\ No newline at end of file
diff --git a/Engine/Plugins/Runtime/MeshModelingToolset/Source/ModelingComponents/Private/Drawing/MeshWireframeComponent.cpp b/Engine/Plugins/Runtime/MeshModelingToolset/Source/ModelingComponents/Private/Drawing/MeshWireframeComponent.cpp
--- a/Engine/Plugins/Runtime/MeshModelingToolset/Source/ModelingComponents/Private/Drawing/MeshWireframeComponent.cpp
+++ b/Engine/Plugins/Runtime/MeshModelingToolset/Source/ModelingComponents/Private/Drawing/MeshWireframeComponent.cpp
@@ -321,6 +321,11 @@ UMeshWireframeComponent::UMeshWireframeComponent()
 void UMeshWireframeComponent::SetWireframeSourceProvider(TSharedPtr<IMeshWireframeSourceProvider> Provider)
 {
 	SourceProvider = Provider;
+	UpdateWireframe();
+}
+
+void UMeshWireframeComponent::UpdateWireframe()
+{
 	if (SourceProvider)
 	{
 		SourceProvider->AccessMesh([&](const IMeshWireframeSource& Source)
diff --git a/Engine/Plugins/Runtime/MeshModelingToolset/Source/ModelingComponents/Public/Drawing/MeshWireframeComponent.h b/Engine/Plugins/Runtime/MeshModelingToolset/Source/ModelingComponents/Public/Drawing/MeshWireframeComponent.h
--- a/Engine/Plugins/Runtime/MeshModelingToolset/Source/ModelingComponents/Public/Drawing/MeshWireframeComponent.h
+++ b/Engine/Plugins/Runtime/MeshModelingToolset/Source/ModelingComponents/Public/Drawing/MeshWireframeComponent.h
@@ -101,7 +101,8 @@ public:
 	 */
 	void SetWireframeSourceProvider(TSharedPtr<IMeshWireframeSourceProvider> Provider);
 
-
+	/** Causes the rendered wireframe to be updated from its source. */
+	void UpdateWireframe();
 public:
 	
 	/**
