MeshModelingTools: Change the dynamic mesh provider interface to return a mesh by value so that it can be used with UDynamicMesh and better aligns with the use in ModelingToolTargetUtil.

#rnx

diff --git a/Engine/Plugins/Experimental/UVEditor/Source/UVEditor/Private/UVEditorMode.cpp b/Engine/Plugins/Experimental/UVEditor/Source/UVEditor/Private/UVEditorMode.cpp
--- a/Engine/Plugins/Experimental/UVEditor/Source/UVEditor/Private/UVEditorMode.cpp
+++ b/Engine/Plugins/Experimental/UVEditor/Source/UVEditor/Private/UVEditorMode.cpp
@@ -17,9 +17,9 @@
 #include "PreviewMesh.h"
 #include "PreviewScene.h"
 #include "TargetInterfaces/AssetBackedTarget.h"
+#include "TargetInterfaces/DynamicMeshCommitter.h"
+#include "TargetInterfaces/DynamicMeshProvider.h"
 #include "TargetInterfaces/MaterialProvider.h"
-#include "TargetInterfaces/MeshDescriptionCommitter.h"
-#include "TargetInterfaces/MeshDescriptionProvider.h"
 #include "ToolSetupUtil.h"
 #include "ToolTargets/UVEditorToolMeshInput.h"
 #include "ToolTargetManager.h"
@@ -126,12 +126,8 @@ const FToolTargetTypeRequirements& UUVEditorMode::GetToolTargetRequirements()
 	static const FToolTargetTypeRequirements ToolTargetRequirements =
 		FToolTargetTypeRequirements({
 			UMaterialProvider::StaticClass(),
-
-			// What we actually care about is dynamic meshes, but we don't currently have
-			// a standardized dynamic mesh commiter/provider interface, because UDynamicMesh
-			// doesn't implement IDynamicMeshCommitter
-			UMeshDescriptionCommitter::StaticClass(),
-			UMeshDescriptionProvider::StaticClass()
+			UDynamicMeshCommitter::StaticClass(),
+			UDynamicMeshProvider::StaticClass()
 			});
 	return ToolTargetRequirements;
 }
diff --git a/Engine/Plugins/Runtime/MeshModelingToolset/Source/ModelingComponents/Private/ConversionUtils/DynamicMeshViaMeshDescriptionUtil.cpp b/Engine/Plugins/Runtime/MeshModelingToolset/Source/ModelingComponents/Private/ConversionUtils/DynamicMeshViaMeshDescriptionUtil.cpp
--- a/Engine/Plugins/Runtime/MeshModelingToolset/Source/ModelingComponents/Private/ConversionUtils/DynamicMeshViaMeshDescriptionUtil.cpp
+++ b/Engine/Plugins/Runtime/MeshModelingToolset/Source/ModelingComponents/Private/ConversionUtils/DynamicMeshViaMeshDescriptionUtil.cpp
@@ -8,12 +8,12 @@
 #include "MeshConversionOptions.h" //FConversionToMeshDescriptionOptions
 #include "MeshDescriptionToDynamicMesh.h"
 
-TSharedPtr<FDynamicMesh3, ESPMode::ThreadSafe> UE::Geometry::GetDynamicMeshViaMeshDescription(
+FDynamicMesh3 UE::Geometry::GetDynamicMeshViaMeshDescription(
 	IMeshDescriptionProvider& MeshDescriptionProvider)
 {
-	TSharedPtr<FDynamicMesh3> DynamicMesh = MakeShared<FDynamicMesh3>();
+	FDynamicMesh3 DynamicMesh;
 	FMeshDescriptionToDynamicMesh Converter;
-	Converter.Convert(MeshDescriptionProvider.GetMeshDescription(), *DynamicMesh);
+	Converter.Convert(MeshDescriptionProvider.GetMeshDescription(), DynamicMesh);
 	return DynamicMesh;
 }
 
diff --git a/Engine/Plugins/Runtime/MeshModelingToolset/Source/ModelingComponents/Private/ModelingToolTargetUtil.cpp b/Engine/Plugins/Runtime/MeshModelingToolset/Source/ModelingComponents/Private/ModelingToolTargetUtil.cpp
--- a/Engine/Plugins/Runtime/MeshModelingToolset/Source/ModelingComponents/Private/ModelingToolTargetUtil.cpp
+++ b/Engine/Plugins/Runtime/MeshModelingToolset/Source/ModelingComponents/Private/ModelingToolTargetUtil.cpp
@@ -188,7 +188,7 @@ FDynamicMesh3 UE::ToolTarget::GetDynamicMeshCopy(UToolTarget* Target, bool bWant
 	IDynamicMeshProvider* DynamicMeshProvider = Cast<IDynamicMeshProvider>(Target);
 	if (DynamicMeshProvider && !bWantMeshTangents)
 	{
-		return *DynamicMeshProvider->GetDynamicMesh();
+		return DynamicMeshProvider->GetDynamicMesh();
 	}
 
 	IMeshDescriptionProvider* MeshDescriptionProvider = Cast<IMeshDescriptionProvider>(Target);
@@ -246,7 +246,28 @@ UE::ToolTarget::EDynamicMeshUpdateResult UE::ToolTarget::CommitMeshDescriptionUp
 }
 
 
+void UE::ToolTarget::Internal::CommitDynamicMeshViaIPersistentDynamicMeshSource(
+	IPersistentDynamicMeshSource& DynamicMeshSource, 
+	const FDynamicMesh3& UpdatedMesh, bool bHaveModifiedTopology)
+{
+	UDynamicMesh* DynamicMesh = DynamicMeshSource.GetDynamicMeshContainer();
+	TUniquePtr<FDynamicMesh3> CurrentMesh = DynamicMesh->ExtractMesh();
+	TSharedPtr<FDynamicMesh3> CurrentMeshShared(CurrentMesh.Release());
+
+	DynamicMesh->EditMesh([&](FDynamicMesh3& EditMesh)
+		{
+			EditMesh.CompactCopy(UpdatedMesh);
+		});
 
+	TSharedPtr<FDynamicMesh3> NewMeshShared = MakeShared<FDynamicMesh3>();
+	DynamicMesh->ProcessMesh([&](const FDynamicMesh3& ReadMesh) { *NewMeshShared = ReadMesh; });
+
+	TUniquePtr<FMeshReplacementChange> ReplaceChange = MakeUnique<FMeshReplacementChange>(CurrentMeshShared, NewMeshShared);
+
+	DynamicMeshSource.CommitDynamicMeshChange(MoveTemp(ReplaceChange), LOCTEXT("CommitDynamicMeshUpdate_MeshSource", "Update Mesh"));
+
+	// todo support bModifiedTopology flag?
+}
 
 UE::ToolTarget::EDynamicMeshUpdateResult UE::ToolTarget::CommitDynamicMeshUpdate(
 	UToolTarget* Target, const FDynamicMesh3& UpdatedMesh, 
@@ -262,24 +283,9 @@ UE::ToolTarget::EDynamicMeshUpdateResult UE::ToolTarget::CommitDynamicMeshUpdate
 	IPersistentDynamicMeshSource* DynamicMeshSource = Cast<IPersistentDynamicMeshSource>(Target);
 	if (DynamicMeshSource)
 	{
-		UDynamicMesh* DynamicMesh = DynamicMeshSource->GetDynamicMeshContainer();
-		TUniquePtr<FDynamicMesh3> CurrentMesh = DynamicMesh->ExtractMesh();
-		TSharedPtr<FDynamicMesh3> CurrentMeshShared(CurrentMesh.Release());
-
-		DynamicMesh->EditMesh([&](FDynamicMesh3& EditMesh) 
-		{ 
-			EditMesh.CompactCopy(UpdatedMesh);
-		});
-
-		TSharedPtr<FDynamicMesh3> NewMeshShared = MakeShared<FDynamicMesh3>();
-		DynamicMesh->ProcessMesh([&](const FDynamicMesh3& ReadMesh) { *NewMeshShared = ReadMesh; });
-
-		TUniquePtr<FMeshReplacementChange> ReplaceChange = MakeUnique<FMeshReplacementChange>(CurrentMeshShared, NewMeshShared);
+		Internal::CommitDynamicMeshViaIPersistentDynamicMeshSource(
+			*DynamicMeshSource, UpdatedMesh, bHaveModifiedTopology);
 		
-		DynamicMeshSource->CommitDynamicMeshChange(MoveTemp(ReplaceChange), LOCTEXT("CommitDynamicMeshUpdate_MeshSource", "Edit Mesh"));
-
-
-		// todo support bModifiedTopology flag?
 		return EDynamicMeshUpdateResult::Ok;
 	}
 
diff --git a/Engine/Plugins/Runtime/MeshModelingToolset/Source/ModelingComponents/Public/ConversionUtils/DynamicMeshViaMeshDescriptionUtil.h b/Engine/Plugins/Runtime/MeshModelingToolset/Source/ModelingComponents/Public/ConversionUtils/DynamicMeshViaMeshDescriptionUtil.h
--- a/Engine/Plugins/Runtime/MeshModelingToolset/Source/ModelingComponents/Public/ConversionUtils/DynamicMeshViaMeshDescriptionUtil.h
+++ b/Engine/Plugins/Runtime/MeshModelingToolset/Source/ModelingComponents/Public/ConversionUtils/DynamicMeshViaMeshDescriptionUtil.h
@@ -15,7 +15,7 @@ class IMeshDescriptionProvider;
 namespace UE {
 namespace Geometry {
 
-	MODELINGCOMPONENTS_API TSharedPtr<FDynamicMesh3, ESPMode::ThreadSafe> GetDynamicMeshViaMeshDescription(
+	MODELINGCOMPONENTS_API FDynamicMesh3 GetDynamicMeshViaMeshDescription(
 		IMeshDescriptionProvider& MeshDescriptionProvider);
 
 	MODELINGCOMPONENTS_API void CommitDynamicMeshViaMeshDescription(IMeshDescriptionCommitter& MeshDescriptionCommitter,
diff --git a/Engine/Plugins/Runtime/MeshModelingToolset/Source/ModelingComponents/Public/ModelingToolTargetUtil.h b/Engine/Plugins/Runtime/MeshModelingToolset/Source/ModelingComponents/Public/ModelingToolTargetUtil.h
--- a/Engine/Plugins/Runtime/MeshModelingToolset/Source/ModelingComponents/Public/ModelingToolTargetUtil.h
+++ b/Engine/Plugins/Runtime/MeshModelingToolset/Source/ModelingComponents/Public/ModelingToolTargetUtil.h
@@ -12,6 +12,7 @@ class UPrimitiveComponent;
 class AActor;
 struct FMeshDescription;
 struct FCreateMeshObjectParams;
+class IPersistentDynamicMeshSource;
 
 //
 // UE::ToolTarget:: namespace contains utility/helper functions for interacting with UToolTargets.
@@ -158,6 +159,18 @@ MODELINGCOMPONENTS_API EDynamicMeshUpdateResult CommitDynamicMeshUVUpdate(UToolT
 MODELINGCOMPONENTS_API bool ConfigureCreateMeshObjectParams(UToolTarget* SourceTarget, FCreateMeshObjectParams& DerivedParamsOut);
 
 
+namespace Internal
+{
+	/**
+	 * Not intended for direct use by tools, just for use by tool target util functions and tool target
+	 * implementations that may need to do this operation. Uses the IPersistentDynamicMeshSource interface
+	 * to perform an update of the dynamic mesh.
+	 * Currently ignores bHaveModifiedTopology.
+	 */
+	MODELINGCOMPONENTS_API void CommitDynamicMeshViaIPersistentDynamicMeshSource(
+		IPersistentDynamicMeshSource& DynamicMeshSource,
+		const UE::Geometry::FDynamicMesh3& UpdatedMesh, bool bHaveModifiedTopology);
+}
 
 
 }  // end namespace ToolTarget
diff --git a/Engine/Plugins/Runtime/MeshModelingToolset/Source/ModelingComponents/Public/TargetInterfaces/DynamicMeshProvider.h b/Engine/Plugins/Runtime/MeshModelingToolset/Source/ModelingComponents/Public/TargetInterfaces/DynamicMeshProvider.h
--- a/Engine/Plugins/Runtime/MeshModelingToolset/Source/ModelingComponents/Public/TargetInterfaces/DynamicMeshProvider.h
+++ b/Engine/Plugins/Runtime/MeshModelingToolset/Source/ModelingComponents/Public/TargetInterfaces/DynamicMeshProvider.h
@@ -23,13 +23,7 @@ class MODELINGCOMPONENTS_API IDynamicMeshProvider
 
 public:
 	/**
-	 * Gives back an editable dynamic mesh. Changes may or may not affect the target, but
-	 * target guarantees that they are safe to perform without crashing. In particular,
-	 * note that it's unspecified whether the next GetMesh() call after an edit will give
-	 * back the edited mesh or a copy of the original.
-	 * 
-	 * To guarantee changes having an effect, target must also be a IDynamicMeshCommitter 
-	 * and the mesh must be committed.
+	 * Gives a copy of a dynamic mesh for tools to operate on.
 	 */
-	virtual TSharedPtr<UE::Geometry::FDynamicMesh3, ESPMode::ThreadSafe> GetDynamicMesh() = 0;
+	virtual UE::Geometry::FDynamicMesh3 GetDynamicMesh() = 0;
 };
diff --git a/Engine/Plugins/Runtime/MeshModelingToolset/Source/ModelingComponentsEditorOnly/Private/ToolTargets/DynamicMeshComponentToolTarget.cpp b/Engine/Plugins/Runtime/MeshModelingToolset/Source/ModelingComponentsEditorOnly/Private/ToolTargets/DynamicMeshComponentToolTarget.cpp
--- a/Engine/Plugins/Runtime/MeshModelingToolset/Source/ModelingComponentsEditorOnly/Private/ToolTargets/DynamicMeshComponentToolTarget.cpp
+++ b/Engine/Plugins/Runtime/MeshModelingToolset/Source/ModelingComponentsEditorOnly/Private/ToolTargets/DynamicMeshComponentToolTarget.cpp
@@ -7,6 +7,7 @@
 #include "MeshDescriptionToDynamicMesh.h"
 #include "StaticMeshAttributes.h"
 #include "Materials/Material.h"
+#include "ModelingToolTargetUtil.h"
 
 #include "Misc/ITransaction.h"
 #include "ScopedTransaction.h"
@@ -14,6 +15,10 @@
 
 #define LOCTEXT_NAMESPACE "UDynamicMeshComponentToolTarget"
 
+namespace DynamicMeshComponentToolTargetLocals
+{
+	const FText CommitMeshTransactionName = LOCTEXT("DynamicMeshComponentToolTargetCommit", "Update Mesh");
+}
 
 bool UDynamicMeshComponentToolTarget::IsValid() const
 {
@@ -134,7 +139,7 @@ void UDynamicMeshComponentToolTarget::CommitMeshDescription(const FCommitter& Co
 
 	TUniquePtr<FMeshReplacementChange> ReplaceChange = MakeUnique<FMeshReplacementChange>(CurrentMeshShared, NewMeshShared);
 
-	CommitDynamicMeshChange(MoveTemp(ReplaceChange), LOCTEXT("DynamicMeshComponentToolTargetCommit", "Update Mesh"));
+	CommitDynamicMeshChange(MoveTemp(ReplaceChange), DynamicMeshComponentToolTargetLocals::CommitMeshTransactionName);
 }
 
 
@@ -172,7 +177,19 @@ void UDynamicMeshComponentToolTarget::CommitDynamicMeshChange(TUniquePtr<FToolCo
 	}
 }
 
+FDynamicMesh3 UDynamicMeshComponentToolTarget::GetDynamicMesh()
+{
+	UDynamicMesh* DynamicMesh = GetDynamicMeshContainer();
+	FDynamicMesh3 Mesh;
+	DynamicMesh->ProcessMesh([&](const FDynamicMesh3& ReadMesh) { Mesh = ReadMesh; });
+	return Mesh;
+}
 
+void UDynamicMeshComponentToolTarget::CommitDynamicMesh(const FDynamicMesh3& UpdatedMesh, const FDynamicMeshCommitInfo& CommitInfo)
+{
+	UE::ToolTarget::Internal::CommitDynamicMeshViaIPersistentDynamicMeshSource(
+		*this, UpdatedMesh, CommitInfo.bTopologyChanged);
+}
 
 
 
diff --git a/Engine/Plugins/Runtime/MeshModelingToolset/Source/ModelingComponentsEditorOnly/Private/ToolTargets/SkeletalMeshComponentToolTarget.cpp b/Engine/Plugins/Runtime/MeshModelingToolset/Source/ModelingComponentsEditorOnly/Private/ToolTargets/SkeletalMeshComponentToolTarget.cpp
--- a/Engine/Plugins/Runtime/MeshModelingToolset/Source/ModelingComponentsEditorOnly/Private/ToolTargets/SkeletalMeshComponentToolTarget.cpp
+++ b/Engine/Plugins/Runtime/MeshModelingToolset/Source/ModelingComponentsEditorOnly/Private/ToolTargets/SkeletalMeshComponentToolTarget.cpp
@@ -106,7 +106,7 @@ const FMeshDescription* USkeletalMeshComponentReadOnlyToolTarget::GetMeshDescrip
 }
 
 
-TSharedPtr<FDynamicMesh3, ESPMode::ThreadSafe> USkeletalMeshComponentReadOnlyToolTarget::GetDynamicMesh()
+FDynamicMesh3 USkeletalMeshComponentReadOnlyToolTarget::GetDynamicMesh()
 {
 	return GetDynamicMeshViaMeshDescription(*this);
 }
diff --git a/Engine/Plugins/Runtime/MeshModelingToolset/Source/ModelingComponentsEditorOnly/Private/ToolTargets/SkeletalMeshToolTarget.cpp b/Engine/Plugins/Runtime/MeshModelingToolset/Source/ModelingComponentsEditorOnly/Private/ToolTargets/SkeletalMeshToolTarget.cpp
--- a/Engine/Plugins/Runtime/MeshModelingToolset/Source/ModelingComponentsEditorOnly/Private/ToolTargets/SkeletalMeshToolTarget.cpp
+++ b/Engine/Plugins/Runtime/MeshModelingToolset/Source/ModelingComponentsEditorOnly/Private/ToolTargets/SkeletalMeshToolTarget.cpp
@@ -3,6 +3,7 @@
 #include "ToolTargets/SkeletalMeshToolTarget.h"
 
 #include "ConversionUtils/DynamicMeshViaMeshDescriptionUtil.h"
+#include "DynamicMesh/DynamicMesh3.h"
 #include "Rendering/SkeletalMeshModel.h"
 #include "Widgets/Notifications/SNotificationList.h"
 #include "Framework/Notifications/NotificationManager.h"
@@ -146,7 +147,7 @@ void USkeletalMeshReadOnlyToolTarget::GetMeshDescription(const USkeletalMesh* Sk
 	}
 }
 
-TSharedPtr<FDynamicMesh3, ESPMode::ThreadSafe> USkeletalMeshReadOnlyToolTarget::GetDynamicMesh()
+FDynamicMesh3 USkeletalMeshReadOnlyToolTarget::GetDynamicMesh()
 {
 	return GetDynamicMeshViaMeshDescription(*this);
 }
diff --git a/Engine/Plugins/Runtime/MeshModelingToolset/Source/ModelingComponentsEditorOnly/Private/ToolTargets/StaticMeshComponentToolTarget.cpp b/Engine/Plugins/Runtime/MeshModelingToolset/Source/ModelingComponentsEditorOnly/Private/ToolTargets/StaticMeshComponentToolTarget.cpp
--- a/Engine/Plugins/Runtime/MeshModelingToolset/Source/ModelingComponentsEditorOnly/Private/ToolTargets/StaticMeshComponentToolTarget.cpp
+++ b/Engine/Plugins/Runtime/MeshModelingToolset/Source/ModelingComponentsEditorOnly/Private/ToolTargets/StaticMeshComponentToolTarget.cpp
@@ -5,6 +5,7 @@
 #include "ComponentReregisterContext.h"
 #include "Components/StaticMeshComponent.h"
 #include "ConversionUtils/DynamicMeshViaMeshDescriptionUtil.h"
+#include "DynamicMesh/DynamicMesh3.h"
 #include "Engine/StaticMesh.h"
 #include "Materials/Material.h"
 #include "ToolTargets/StaticMeshToolTarget.h"
@@ -175,7 +176,7 @@ void UStaticMeshComponentToolTarget::CommitMeshDescription(const FCommitter& Com
 	Component->RecreatePhysicsState();
 }
 
-TSharedPtr<FDynamicMesh3, ESPMode::ThreadSafe> UStaticMeshComponentToolTarget::GetDynamicMesh()
+FDynamicMesh3 UStaticMeshComponentToolTarget::GetDynamicMesh()
 {
 	return GetDynamicMeshViaMeshDescription(*this);
 }
diff --git a/Engine/Plugins/Runtime/MeshModelingToolset/Source/ModelingComponentsEditorOnly/Private/ToolTargets/StaticMeshToolTarget.cpp b/Engine/Plugins/Runtime/MeshModelingToolset/Source/ModelingComponentsEditorOnly/Private/ToolTargets/StaticMeshToolTarget.cpp
--- a/Engine/Plugins/Runtime/MeshModelingToolset/Source/ModelingComponentsEditorOnly/Private/ToolTargets/StaticMeshToolTarget.cpp
+++ b/Engine/Plugins/Runtime/MeshModelingToolset/Source/ModelingComponentsEditorOnly/Private/ToolTargets/StaticMeshToolTarget.cpp
@@ -4,6 +4,7 @@
 
 #include "AssetUtils/MeshDescriptionUtil.h"
 #include "ConversionUtils/DynamicMeshViaMeshDescriptionUtil.h"
+#include "DynamicMesh/DynamicMesh3.h"
 #include "Engine/StaticMesh.h"
 #include "Materials/Material.h"
 #include "RenderingThread.h"
@@ -252,7 +253,7 @@ void UStaticMeshToolTarget::CommitMeshDescription(UStaticMesh* StaticMeshIn, FMe
 	StaticMeshIn->PostEditChange();
 }
 
-TSharedPtr<FDynamicMesh3, ESPMode::ThreadSafe> UStaticMeshToolTarget::GetDynamicMesh()
+FDynamicMesh3 UStaticMeshToolTarget::GetDynamicMesh()
 {
 	return GetDynamicMeshViaMeshDescription(*this);
 }
diff --git a/Engine/Plugins/Runtime/MeshModelingToolset/Source/ModelingComponentsEditorOnly/Private/ToolTargets/VolumeDynamicMeshToolTarget.cpp b/Engine/Plugins/Runtime/MeshModelingToolset/Source/ModelingComponentsEditorOnly/Private/ToolTargets/VolumeDynamicMeshToolTarget.cpp
--- a/Engine/Plugins/Runtime/MeshModelingToolset/Source/ModelingComponentsEditorOnly/Private/ToolTargets/VolumeDynamicMeshToolTarget.cpp
+++ b/Engine/Plugins/Runtime/MeshModelingToolset/Source/ModelingComponentsEditorOnly/Private/ToolTargets/VolumeDynamicMeshToolTarget.cpp
@@ -66,7 +66,7 @@ void UVolumeDynamicMeshToolTarget::GetMaterialSet(FComponentMaterialSet& Materia
 	}
 }
 
-TSharedPtr<UE::Geometry::FDynamicMesh3, ESPMode::ThreadSafe> UVolumeDynamicMeshToolTarget::GetDynamicMesh()
+FDynamicMesh3 UVolumeDynamicMeshToolTarget::GetDynamicMesh()
 {
 	UBrushComponent* BrushComponent = Cast<UBrushComponent>(Component);
 	if (!BrushComponent)
@@ -79,9 +79,9 @@ TSharedPtr<UE::Geometry::FDynamicMesh3, ESPMode::ThreadSafe> UVolumeDynamicMeshT
 		return nullptr;
 	}
 
-	TSharedPtr<FDynamicMesh3, ESPMode::ThreadSafe> DynamicMesh = MakeShared<FDynamicMesh3, ESPMode::ThreadSafe>();
-	UE::Conversion::VolumeToDynamicMesh(Volume, *DynamicMesh, GetVolumeToMeshOptions());
-	FMeshNormals::InitializeMeshToPerTriangleNormals(DynamicMesh.Get());
+	FDynamicMesh3 DynamicMesh;
+	UE::Conversion::VolumeToDynamicMesh(Volume, DynamicMesh, GetVolumeToMeshOptions());
+	FMeshNormals::InitializeMeshToPerTriangleNormals(&DynamicMesh);
 
 	return DynamicMesh;
 }
diff --git a/Engine/Plugins/Runtime/MeshModelingToolset/Source/ModelingComponentsEditorOnly/Private/ToolTargets/VolumeMeshDescriptionToolTarget.cpp b/Engine/Plugins/Runtime/MeshModelingToolset/Source/ModelingComponentsEditorOnly/Private/ToolTargets/VolumeMeshDescriptionToolTarget.cpp
--- a/Engine/Plugins/Runtime/MeshModelingToolset/Source/ModelingComponentsEditorOnly/Private/ToolTargets/VolumeMeshDescriptionToolTarget.cpp
+++ b/Engine/Plugins/Runtime/MeshModelingToolset/Source/ModelingComponentsEditorOnly/Private/ToolTargets/VolumeMeshDescriptionToolTarget.cpp
@@ -21,11 +21,8 @@ const FMeshDescription* UVolumeMeshDescriptionToolTarget::GetMeshDescription()
 		// Editor.h. However, that path doesn't assign polygroups to the result, which we
 		// typically want when using this target, hence the path through a dynamic mesh.
 
-		TSharedPtr<FDynamicMesh3, ESPMode::ThreadSafe> DynamicMesh = GetDynamicMesh();
-		if (!DynamicMesh)
-		{
-			return nullptr;
-		}
+		TSharedPtr<FDynamicMesh3, ESPMode::ThreadSafe> DynamicMesh = 
+			MakeShared<FDynamicMesh3, ESPMode::ThreadSafe>(GetDynamicMesh());
 
 		ConvertedMeshDescription = MakeShared<FMeshDescription, ESPMode::ThreadSafe>();
 		FStaticMeshAttributes Attributes(*ConvertedMeshDescription);
diff --git a/Engine/Plugins/Runtime/MeshModelingToolset/Source/ModelingComponentsEditorOnly/Public/ToolTargets/DynamicMeshComponentToolTarget.h b/Engine/Plugins/Runtime/MeshModelingToolset/Source/ModelingComponentsEditorOnly/Public/ToolTargets/DynamicMeshComponentToolTarget.h
--- a/Engine/Plugins/Runtime/MeshModelingToolset/Source/ModelingComponentsEditorOnly/Public/ToolTargets/DynamicMeshComponentToolTarget.h
+++ b/Engine/Plugins/Runtime/MeshModelingToolset/Source/ModelingComponentsEditorOnly/Public/ToolTargets/DynamicMeshComponentToolTarget.h
@@ -4,6 +4,8 @@
 
 #include "CoreMinimal.h"
 
+#include "TargetInterfaces/DynamicMeshCommitter.h"
+#include "TargetInterfaces/DynamicMeshProvider.h"
 #include "TargetInterfaces/MaterialProvider.h"
 #include "TargetInterfaces/MeshDescriptionCommitter.h"
 #include "TargetInterfaces/MeshDescriptionProvider.h"
@@ -22,6 +24,8 @@ class MODELINGCOMPONENTSEDITORONLY_API UDynamicMeshComponentToolTarget :
 	public UPrimitiveComponentToolTarget,
 	public IMeshDescriptionCommitter, 
 	public IMeshDescriptionProvider, 
+	public IDynamicMeshProvider,
+	public IDynamicMeshCommitter,
 	public IMaterialProvider,
 	public IPersistentDynamicMeshSource
 {
@@ -39,6 +43,13 @@ public:
 	virtual void CommitMeshDescription(const FCommitter& Committer) override;
 	using IMeshDescriptionCommitter::CommitMeshDescription; // unhide the other overload
 
+	// IDynamicMeshProvider implementation
+	virtual UE::Geometry::FDynamicMesh3 GetDynamicMesh() override;
+
+	// IDynamicMeshCommitter implementation
+	virtual void CommitDynamicMesh(const UE::Geometry::FDynamicMesh3& Mesh, const FDynamicMeshCommitInfo& CommitInfo) override;
+	using IDynamicMeshCommitter::CommitDynamicMesh; // unhide the other overload
+
 	// IMaterialProvider implementation
 	int32 GetNumMaterials() const override;
 	UMaterialInterface* GetMaterial(int32 MaterialIndex) const override;
diff --git a/Engine/Plugins/Runtime/MeshModelingToolset/Source/ModelingComponentsEditorOnly/Public/ToolTargets/SkeletalMeshComponentToolTarget.h b/Engine/Plugins/Runtime/MeshModelingToolset/Source/ModelingComponentsEditorOnly/Public/ToolTargets/SkeletalMeshComponentToolTarget.h
--- a/Engine/Plugins/Runtime/MeshModelingToolset/Source/ModelingComponentsEditorOnly/Public/ToolTargets/SkeletalMeshComponentToolTarget.h
+++ b/Engine/Plugins/Runtime/MeshModelingToolset/Source/ModelingComponentsEditorOnly/Public/ToolTargets/SkeletalMeshComponentToolTarget.h
@@ -43,7 +43,7 @@ public:
 	bool CommitMaterialSetUpdate(const FComponentMaterialSet& MaterialSet, bool bApplyToAsset) override;
 
 	// IDynamicMeshProvider
-	virtual TSharedPtr<UE::Geometry::FDynamicMesh3, ESPMode::ThreadSafe> GetDynamicMesh() override;
+	virtual UE::Geometry::FDynamicMesh3 GetDynamicMesh() override;
 
 	// ISkeletalMeshBackedTarget implementation
 	USkeletalMesh* GetSkeletalMesh() const override;
diff --git a/Engine/Plugins/Runtime/MeshModelingToolset/Source/ModelingComponentsEditorOnly/Public/ToolTargets/SkeletalMeshToolTarget.h b/Engine/Plugins/Runtime/MeshModelingToolset/Source/ModelingComponentsEditorOnly/Public/ToolTargets/SkeletalMeshToolTarget.h
--- a/Engine/Plugins/Runtime/MeshModelingToolset/Source/ModelingComponentsEditorOnly/Public/ToolTargets/SkeletalMeshToolTarget.h
+++ b/Engine/Plugins/Runtime/MeshModelingToolset/Source/ModelingComponentsEditorOnly/Public/ToolTargets/SkeletalMeshToolTarget.h
@@ -44,7 +44,7 @@ public:
 	bool CommitMaterialSetUpdate(const FComponentMaterialSet& MaterialSet, bool bApplyToAsset) override;	
 
 	// IDynamicMeshProvider
-	virtual TSharedPtr<UE::Geometry::FDynamicMesh3, ESPMode::ThreadSafe> GetDynamicMesh() override;
+	virtual UE::Geometry::FDynamicMesh3 GetDynamicMesh() override;
 
 	// ISkeletalMeshBackedTarget implementation
 	USkeletalMesh* GetSkeletalMesh() const override;
diff --git a/Engine/Plugins/Runtime/MeshModelingToolset/Source/ModelingComponentsEditorOnly/Public/ToolTargets/StaticMeshComponentToolTarget.h b/Engine/Plugins/Runtime/MeshModelingToolset/Source/ModelingComponentsEditorOnly/Public/ToolTargets/StaticMeshComponentToolTarget.h
--- a/Engine/Plugins/Runtime/MeshModelingToolset/Source/ModelingComponentsEditorOnly/Public/ToolTargets/StaticMeshComponentToolTarget.h
+++ b/Engine/Plugins/Runtime/MeshModelingToolset/Source/ModelingComponentsEditorOnly/Public/ToolTargets/StaticMeshComponentToolTarget.h
@@ -52,7 +52,7 @@ public:
 	using IMeshDescriptionCommitter::CommitMeshDescription; // unhide the other overload
 
 	// IDynamicMeshProvider
-	virtual TSharedPtr<UE::Geometry::FDynamicMesh3, ESPMode::ThreadSafe> GetDynamicMesh() override;
+	virtual UE::Geometry::FDynamicMesh3 GetDynamicMesh() override;
 
 	// IDynamicMeshCommitter
 	virtual void CommitDynamicMesh(const UE::Geometry::FDynamicMesh3& Mesh, const FDynamicMeshCommitInfo& CommitInfo) override;
diff --git a/Engine/Plugins/Runtime/MeshModelingToolset/Source/ModelingComponentsEditorOnly/Public/ToolTargets/StaticMeshToolTarget.h b/Engine/Plugins/Runtime/MeshModelingToolset/Source/ModelingComponentsEditorOnly/Public/ToolTargets/StaticMeshToolTarget.h
--- a/Engine/Plugins/Runtime/MeshModelingToolset/Source/ModelingComponentsEditorOnly/Public/ToolTargets/StaticMeshToolTarget.h
+++ b/Engine/Plugins/Runtime/MeshModelingToolset/Source/ModelingComponentsEditorOnly/Public/ToolTargets/StaticMeshToolTarget.h
@@ -18,7 +18,7 @@
 class UStaticMesh;
 
 /**
- * A tool target backed by a static mesh component that can provide and take a mesh
+ * A tool target backed by a static mesh asset that can provide and take a mesh
  * description.
  */
 UCLASS(Transient)
@@ -54,7 +54,7 @@ public:
 	using IMeshDescriptionCommitter::CommitMeshDescription; // unhide the other overload
 
 	// IDynamicMeshProvider
-	virtual TSharedPtr<UE::Geometry::FDynamicMesh3, ESPMode::ThreadSafe> GetDynamicMesh() override;
+	virtual UE::Geometry::FDynamicMesh3 GetDynamicMesh() override;
 
 	// IDynamicMeshCommitter
 	virtual void CommitDynamicMesh(const UE::Geometry::FDynamicMesh3& Mesh, const FDynamicMeshCommitInfo& CommitInfo) override;
diff --git a/Engine/Plugins/Runtime/MeshModelingToolset/Source/ModelingComponentsEditorOnly/Public/ToolTargets/VolumeDynamicMeshToolTarget.h b/Engine/Plugins/Runtime/MeshModelingToolset/Source/ModelingComponentsEditorOnly/Public/ToolTargets/VolumeDynamicMeshToolTarget.h
--- a/Engine/Plugins/Runtime/MeshModelingToolset/Source/ModelingComponentsEditorOnly/Public/ToolTargets/VolumeDynamicMeshToolTarget.h
+++ b/Engine/Plugins/Runtime/MeshModelingToolset/Source/ModelingComponentsEditorOnly/Public/ToolTargets/VolumeDynamicMeshToolTarget.h
@@ -39,7 +39,7 @@ public:
 	const UE::Conversion::FVolumeToMeshOptions& GetVolumeToMeshOptions() { return VolumeToMeshOptions; }
 
 	// IDynamicMeshProvider implementation
-	virtual TSharedPtr<UE::Geometry::FDynamicMesh3, ESPMode::ThreadSafe> GetDynamicMesh() override;
+	virtual UE::Geometry::FDynamicMesh3 GetDynamicMesh() override;
 
 	// IDynamicMeshCommitter implementation
 	virtual void CommitDynamicMesh(const UE::Geometry::FDynamicMesh3& Mesh, const FDynamicMeshCommitInfo&) override;
