GeometryScripting: add support for reading from Volumes in SceneUtilityFunctions::CopyMeshFromComponent()

#rnx

diff --git a/Engine/Plugins/Experimental/GeometryScripting/Source/GeometryScriptingCore/Private/SceneUtilityFunctions.cpp b/Engine/Plugins/Experimental/GeometryScripting/Source/GeometryScriptingCore/Private/SceneUtilityFunctions.cpp
--- a/Engine/Plugins/Experimental/GeometryScripting/Source/GeometryScriptingCore/Private/SceneUtilityFunctions.cpp
+++ b/Engine/Plugins/Experimental/GeometryScripting/Source/GeometryScriptingCore/Private/SceneUtilityFunctions.cpp
@@ -6,9 +6,12 @@
 #include "UDynamicMesh.h"
 #include "GeometryScript/MeshAssetFunctions.h"
 #include "DynamicMesh/MeshTransforms.h"
+#include "DynamicMesh/MeshNormals.h"
 
 #include "Components/StaticMeshComponent.h"
+#include "Components/BrushComponent.h"
 #include "Components/DynamicMeshComponent.h"
+#include "ConversionUtils/VolumeToDynamicMesh.h"
 
 using namespace UE::Geometry;
 
@@ -31,9 +34,8 @@ UDynamicMesh* UGeometryScriptLibrary_SceneUtilityFunctions::CopyMeshFromComponen
 {
 	Outcome = EGeometryScriptOutcomePins::Failure;
 
-	if (Cast<UStaticMeshComponent>(Component) != nullptr)
+	if (UStaticMeshComponent* StaticMeshComponent = Cast<UStaticMeshComponent>(Component))
 	{
-		UStaticMeshComponent* StaticMeshComponent = Cast<UStaticMeshComponent>(Component);
 		LocalToWorld = StaticMeshComponent->GetComponentTransform();
 		UStaticMesh* StaticMesh = StaticMeshComponent->GetStaticMesh();
 		if (StaticMesh)
@@ -49,9 +51,8 @@ UDynamicMesh* UGeometryScriptLibrary_SceneUtilityFunctions::CopyMeshFromComponen
 			UE::Geometry::AppendError(Debug, EGeometryScriptErrorType::InvalidInputs, LOCTEXT("CopyMeshFromComponent_MissingStaticMesh", "CopyMeshFromComponent: StaticMeshComponent has a null StaticMesh"));
 		}
 	}
-	else if (Cast<UDynamicMeshComponent>(Component) != nullptr)
+	else if (UDynamicMeshComponent* DynamicMeshComponent = Cast<UDynamicMeshComponent>(Component))
 	{
-		UDynamicMeshComponent* DynamicMeshComponent = Cast<UDynamicMeshComponent>(Component);
 		LocalToWorld = DynamicMeshComponent->GetComponentTransform();
 		UDynamicMesh* CopyDynamicMesh = DynamicMeshComponent->GetDynamicMesh();
 		if (CopyDynamicMesh)
@@ -67,6 +68,38 @@ UDynamicMesh* UGeometryScriptLibrary_SceneUtilityFunctions::CopyMeshFromComponen
 			UE::Geometry::AppendError(Debug, EGeometryScriptErrorType::InvalidInputs, LOCTEXT("CopyMeshFromComponent_MissingDynamicMesh", "CopyMeshFromComponent: DynamicMeshComponent has a null DynamicMesh"));
 		}
 	}
+	else if (UBrushComponent* BrushComponent = Cast<UBrushComponent>(Component))
+	{
+		LocalToWorld = BrushComponent->GetComponentTransform();
+
+		UE::Conversion::FVolumeToMeshOptions VolOptions;
+		VolOptions.bMergeVertices = true;
+		VolOptions.bAutoRepairMesh = true;
+		VolOptions.bOptimizeMesh = true;
+		VolOptions.bSetGroups = true;
+
+		FDynamicMesh3 ConvertedMesh(EMeshComponents::FaceGroups);
+		UE::Conversion::BrushComponentToDynamicMesh(BrushComponent, ConvertedMesh, VolOptions);
+
+		// compute normals for current polygroup topology
+		ConvertedMesh.EnableAttributes();
+		if (Options.bWantNormals)
+		{
+			FDynamicMeshNormalOverlay* Normals = ConvertedMesh.Attributes()->PrimaryNormals();
+			FMeshNormals::InitializeOverlayTopologyFromFaceGroups(&ConvertedMesh, Normals);
+			FMeshNormals::QuickRecomputeOverlayNormals(ConvertedMesh);
+		}
+
+		if (ConvertedMesh.TriangleCount() > 0)
+		{
+			ToDynamicMesh->SetMesh(MoveTemp(ConvertedMesh));
+			Outcome = EGeometryScriptOutcomePins::Success;
+		}
+		else
+		{
+			UE::Geometry::AppendError(Debug, EGeometryScriptErrorType::InvalidInputs, LOCTEXT("CopyMeshFromComponent_InvalidBrushConversion", "CopyMeshFromComponent: BrushComponent conversion produced 0 triangles"));
+		}
+	}
 
 	// transform mesh to world
 	if (Outcome == EGeometryScriptOutcomePins::Success && bTransformToWorld)
diff --git a/Engine/Plugins/Runtime/MeshModelingToolset/Source/ModelingComponents/Private/ConversionUtils/VolumeToDynamicMesh.cpp b/Engine/Plugins/Runtime/MeshModelingToolset/Source/ModelingComponents/Private/ConversionUtils/VolumeToDynamicMesh.cpp
--- a/Engine/Plugins/Runtime/MeshModelingToolset/Source/ModelingComponents/Private/ConversionUtils/VolumeToDynamicMesh.cpp
+++ b/Engine/Plugins/Runtime/MeshModelingToolset/Source/ModelingComponents/Private/ConversionUtils/VolumeToDynamicMesh.cpp
@@ -5,6 +5,7 @@
 #include "ConstrainedDelaunay2.h"
 #include "Polygon2.h"
 #include "DynamicMesh/DynamicMesh3.h"
+#include "DynamicMesh/MeshTransforms.h"
 #include "GameFramework/Volume.h"
 #include "MeshBoundaryLoops.h"
 #include "MeshQueries.h"
@@ -13,6 +14,7 @@
 #include "DynamicMesh/Operations/MergeCoincidentMeshEdges.h"
 #include "Operations/MinimalHoleFiller.h"
 #include "Operations/PlanarFlipsOptimization.h"
+#include "Components/BrushComponent.h"
 
 #if WITH_EDITOR
 #include "Engine/Classes/Engine/Polys.h"
@@ -25,6 +27,43 @@ namespace Conversion {
 
 void VolumeToDynamicMesh(AVolume* Volume, FDynamicMesh3& Mesh, 
 	const FVolumeToMeshOptions& Options)
+{
+	UModel* Model = Volume->Brush;
+	if (Model == nullptr)		// model is null in some cases, eg default physics volume
+	{
+		return;
+	}
+
+	BrushToDynamicMesh(*Model, Mesh, Options);
+
+	if (Options.bInWorldSpace)
+	{
+		UE::Geometry::FTransform3d XForm(Volume->GetTransform());
+		MeshTransforms::ApplyTransform(Mesh, XForm);
+	}
+}
+
+
+void BrushComponentToDynamicMesh(UBrushComponent* Component, UE::Geometry::FDynamicMesh3& Mesh, const FVolumeToMeshOptions& Options)
+{
+	UModel* Model = Component->Brush;
+	if (Model == nullptr)		// model is null in some cases, eg default physics volume
+	{
+		return;
+	}
+
+	BrushToDynamicMesh(*Model, Mesh, Options);
+
+	if (Options.bInWorldSpace)
+	{
+		UE::Geometry::FTransform3d XForm(Component->GetComponentTransform());
+		MeshTransforms::ApplyTransform(Mesh, XForm);
+	}
+}
+
+
+
+void BrushToDynamicMesh(UModel& BrushModel, UE::Geometry::FDynamicMesh3& Mesh, const FVolumeToMeshOptions& Options)
 {
 	Mesh.Clear();
 	if (Options.bSetGroups)
@@ -32,10 +71,6 @@ void VolumeToDynamicMesh(AVolume* Volume, FDynamicMesh3& Mesh,
 		Mesh.EnableTriangleGroups();
 	}
 
-	UModel* Model = Volume->Brush;
-	UE::Geometry::FTransform3d XForm = (Options.bInWorldSpace) ? 
-		UE::Geometry::FTransform3d(Volume->GetTransform()) : UE::Geometry::FTransform3d::Identity();
-
 #if WITH_EDITOR
 	// In the editor, the preferred source of geometry for a volume is the Polys array,
 	// which the bsp nodes are generated from, because the polys may be broken up into
@@ -44,7 +79,7 @@ void VolumeToDynamicMesh(AVolume* Volume, FDynamicMesh3& Mesh,
 
 	// We do not try to merge any vertices yet.
 
-	const TArray<FPoly>& Polygons = Model->Polys->Element;
+	const TArray<FPoly>& Polygons = BrushModel.Polys->Element;
 	for (FPoly Poly : Polygons)
 	{
 		int32 NumVerts = Poly.Vertices.Num();
@@ -62,7 +97,7 @@ void VolumeToDynamicMesh(AVolume* Volume, FDynamicMesh3& Mesh,
 		Vids.SetNum(NumVerts);
 		for (int32 VertexIndex = 0; VertexIndex < NumVerts; ++VertexIndex)
 		{
-			FVector3d Point = XForm.TransformPosition((FVector3d)Poly.Vertices[VertexIndex]);
+			FVector3d Point = (FVector3d)Poly.Vertices[VertexIndex];
 			Vids[VertexIndex] = Mesh.AppendVertex(Point);
 			ToTriangulate.AppendVertex(Plane.ToPlaneUV(Point, 2));
 		}
@@ -88,7 +123,7 @@ void VolumeToDynamicMesh(AVolume* Volume, FDynamicMesh3& Mesh,
 #else
 	// Each "BspNode" is a planar polygon, triangulate each polygon and accumulate in a mesh.
 	// Note that this does not make any attempt to weld vertices/edges
-	for (const FBspNode& Node : Model->Nodes)
+	for (const FBspNode& Node : BrushModel.Nodes)
 	{
 		FVector3d Normal = (FVector3d)Node.Plane;
 		FFrame3d Plane(Node.Plane.W * Normal, Normal);
@@ -101,9 +136,8 @@ void VolumeToDynamicMesh(AVolume* Volume, FDynamicMesh3& Mesh,
 			Vids.SetNum(NumVerts);
 			for (int32 VertexIndex = 0; VertexIndex < NumVerts; ++VertexIndex)
 			{
-				const FVert& Vert = Model->Verts[Node.iVertPool + VertexIndex];
-				FVector3d Point = (FVector3d)Model->Points[Vert.pVertex];
-				Point = XForm.TransformPosition(Point);
+				const FVert& Vert = BrushModel.Verts[Node.iVertPool + VertexIndex];
+				FVector3d Point = (FVector3d)BrushModel.Points[Vert.pVertex];
 				Vids[VertexIndex] = Mesh.AppendVertex(Point);
 				ToTriangulate.AppendVertex(Plane.ToPlaneUV(Point, 2));
 			}
@@ -124,8 +158,8 @@ void VolumeToDynamicMesh(AVolume* Volume, FDynamicMesh3& Mesh,
 			{
 				Mesh.AppendTriangle(Vids[Tri.A], Vids[Tri.B], Vids[Tri.C], GroupID);
 			}
-		}
 	}
+}
 #endif
 
 	if (Options.bMergeVertices)
@@ -158,4 +192,8 @@ void VolumeToDynamicMesh(AVolume* Volume, FDynamicMesh3& Mesh,
 	}
 }
 
+
+
+
+
 }}//end namespace UE::Conversion
\ No newline at end of file
diff --git a/Engine/Plugins/Runtime/MeshModelingToolset/Source/ModelingComponents/Public/ConversionUtils/VolumeToDynamicMesh.h b/Engine/Plugins/Runtime/MeshModelingToolset/Source/ModelingComponents/Public/ConversionUtils/VolumeToDynamicMesh.h
--- a/Engine/Plugins/Runtime/MeshModelingToolset/Source/ModelingComponents/Public/ConversionUtils/VolumeToDynamicMesh.h
+++ b/Engine/Plugins/Runtime/MeshModelingToolset/Source/ModelingComponents/Public/ConversionUtils/VolumeToDynamicMesh.h
@@ -6,6 +6,8 @@
 #include "FrameTypes.h"
 
 class AVolume;
+class UBrushComponent;
+class UModel;
 PREDECLARE_GEOMETRY(class FDynamicMesh3);
 
 namespace UE {
@@ -21,8 +23,19 @@ struct FVolumeToMeshOptions
 };
 
 /**
- * Converts a volume to a dynamic mesh. Does not initialize normals and does not delete the volume.
+ * Converts an AVolume to a FDynamicMesh3. Does not initialize normals and does not delete the volume.
  */
 void MODELINGCOMPONENTS_API VolumeToDynamicMesh(AVolume* Volume, UE::Geometry::FDynamicMesh3& Mesh, const FVolumeToMeshOptions& Options);
 
+/**
+* Converts a UBrushComponent to a FDynamicMesh3. Does not initialize normals and does not delete the volume.
+*/
+void MODELINGCOMPONENTS_API BrushComponentToDynamicMesh(UBrushComponent* Component, UE::Geometry::FDynamicMesh3& Mesh, const FVolumeToMeshOptions& Options);
+
+/**
+* Converts a UModel to a FDynamicMesh3. Does not initialize normals.
+*/
+void MODELINGCOMPONENTS_API BrushToDynamicMesh(UModel& BrushModel, UE::Geometry::FDynamicMesh3& Mesh, const FVolumeToMeshOptions& Options);
+
+
 }}//end namespace UE::Conversion
\ No newline at end of file
