ModelingComponents: Fix issue with the PointSet using wrong colors due to a missing sRGB conversion.

#rnx

diff --git a/Engine/Plugins/Runtime/MeshModelingToolset/Source/ModelingComponents/Private/Drawing/PointSetComponent.cpp b/Engine/Plugins/Runtime/MeshModelingToolset/Source/ModelingComponents/Private/Drawing/PointSetComponent.cpp
--- a/Engine/Plugins/Runtime/MeshModelingToolset/Source/ModelingComponents/Private/Drawing/PointSetComponent.cpp
+++ b/Engine/Plugins/Runtime/MeshModelingToolset/Source/ModelingComponents/Private/Drawing/PointSetComponent.cpp
@@ -102,12 +102,17 @@ public:
 				int UseVertexBufferIndex = 4 * i;
 				int UseIndexBufferIndex = 6 * i;
 
+				// The color stored in the vertices actually gets interpreted as a linear color by the material,
+				// whereas it is more convenient for the user of the LineSet to specify colors as sRGB. So we actually
+				// have to convert it back to linear. The ToFColor(false) call just scales back into 0-255 space.
+				FColor color = FLinearColor::FromSRGBColor(Point.Color).ToFColor(false);
+
 				const FVector2D UV(Point.Size, 0.0f);
 				for (int j = 0; j < 4; ++j)
 				{
 					VertexBuffers.PositionVertexBuffer.VertexPosition(UseVertexBufferIndex + j) = Point.Position;
 					VertexBuffers.StaticMeshVertexBuffer.SetVertexUV(UseVertexBufferIndex + j, 0, UV);
-					VertexBuffers.ColorVertexBuffer.VertexColor(UseVertexBufferIndex + j) = Point.Color;
+					VertexBuffers.ColorVertexBuffer.VertexColor(UseVertexBufferIndex + j) = color;
 					VertexBuffers.StaticMeshVertexBuffer.SetVertexTangents(UseVertexBufferIndex + j, FVector::ZeroVector, FVector::ZeroVector, TangentVectors[j]);
 				}
 
