MeshModelingTools: combine the volume component tool targets into one tool target. Make that tool target implement the IPhysicsDataSource interface. Fix bug in PhysicsDataCollection.cpp that caused crash in PhysicsTools::InitializePhysicsToolObjectPropertySet later.

#rnx

diff --git a/Engine/Plugins/Editor/ModelingToolsEditorMode/Source/ModelingToolsEditorMode/Private/ModelingToolsEditorMode.cpp b/Engine/Plugins/Editor/ModelingToolsEditorMode/Source/ModelingToolsEditorMode/Private/ModelingToolsEditorMode.cpp
--- a/Engine/Plugins/Editor/ModelingToolsEditorMode/Source/ModelingToolsEditorMode/Private/ModelingToolsEditorMode.cpp
+++ b/Engine/Plugins/Editor/ModelingToolsEditorMode/Source/ModelingToolsEditorMode/Private/ModelingToolsEditorMode.cpp
@@ -8,8 +8,7 @@
 #include "ToolTargetManager.h"
 #include "ContextObjectStore.h"
 #include "ToolTargets/StaticMeshComponentToolTarget.h"
-#include "ToolTargets/VolumeDynamicMeshToolTarget.h"
-#include "ToolTargets/VolumeMeshDescriptionToolTarget.h"
+#include "ToolTargets/VolumeComponentToolTarget.h"
 #include "ToolTargets/DynamicMeshComponentToolTarget.h"
 #include "ToolTargets/SkeletalMeshComponentToolTarget.h"
 #include "Framework/Commands/UICommandList.h"
@@ -316,8 +315,7 @@ void UModelingToolsEditorMode::Enter()
 
 	// Register builders for tool targets that the mode uses.
 	GetInteractiveToolsContext()->TargetManager->AddTargetFactory(NewObject<UStaticMeshComponentToolTargetFactory>(GetToolManager()));
-	GetInteractiveToolsContext()->TargetManager->AddTargetFactory(NewObject<UVolumeDynamicMeshToolTargetFactory>(GetToolManager()));
-	GetInteractiveToolsContext()->TargetManager->AddTargetFactory(NewObject<UVolumeMeshDescriptionToolTargetFactory>(GetToolManager()));
+	GetInteractiveToolsContext()->TargetManager->AddTargetFactory(NewObject<UVolumeComponentToolTargetFactory>(GetToolManager()));
 	GetInteractiveToolsContext()->TargetManager->AddTargetFactory(NewObject<UDynamicMeshComponentToolTargetFactory>(GetToolManager()));
 
 	// Register read-only skeletal mesh tool targets. Currently tools that write to meshes risk breaking
diff --git a/Engine/Plugins/Runtime/MeshModelingToolset/Source/ModelingComponents/Private/Physics/PhysicsDataCollection.cpp b/Engine/Plugins/Runtime/MeshModelingToolset/Source/ModelingComponents/Private/Physics/PhysicsDataCollection.cpp
--- a/Engine/Plugins/Runtime/MeshModelingToolset/Source/ModelingComponents/Private/Physics/PhysicsDataCollection.cpp
+++ b/Engine/Plugins/Runtime/MeshModelingToolset/Source/ModelingComponents/Private/Physics/PhysicsDataCollection.cpp
@@ -13,13 +13,13 @@ using namespace UE::Geometry;
 
 void FPhysicsDataCollection::InitializeFromComponent(const UActorComponent* Component, bool bInitializeAggGeom)
 {
+	SourceComponent = Component;
 	BodySetup = nullptr;
 
 	if ( const UStaticMeshComponent* StaticMeshComponent = Cast<UStaticMeshComponent>(Component) )
 	{
 		if (UStaticMesh* StaticMesh = StaticMeshComponent->GetStaticMesh())
 		{
-			SourceComponent = StaticMeshComponent;
 			SourceStaticMesh = StaticMesh;
 			BodySetup = StaticMesh->GetBodySetup();
 		}
diff --git a/Engine/Plugins/Runtime/MeshModelingToolset/Source/ModelingComponentsEditorOnly/Private/EditorModelingObjectsCreationAPI.cpp b/Engine/Plugins/Runtime/MeshModelingToolset/Source/ModelingComponentsEditorOnly/Private/EditorModelingObjectsCreationAPI.cpp
--- a/Engine/Plugins/Runtime/MeshModelingToolset/Source/ModelingComponentsEditorOnly/Private/EditorModelingObjectsCreationAPI.cpp
+++ b/Engine/Plugins/Runtime/MeshModelingToolset/Source/ModelingComponentsEditorOnly/Private/EditorModelingObjectsCreationAPI.cpp
@@ -16,7 +16,7 @@
 #include "Engine/StaticMeshActor.h"
 #include "Materials/Material.h"
 
-#include "ToolTargets/VolumeDynamicMeshToolTarget.h"  // for CVarModelingMaxVolumeTriangleCount
+#include "ToolTargets/VolumeComponentToolTarget.h"  // for CVarModelingMaxVolumeTriangleCount
 #include "Engine/BlockingVolume.h"
 #include "Components/BrushComponent.h"
 #include "Engine/Polys.h"
diff --git a/Engine/Plugins/Runtime/MeshModelingToolset/Source/ModelingComponentsEditorOnly/Private/ToolTargets/VolumeComponentToolTarget.cpp b/Engine/Plugins/Runtime/MeshModelingToolset/Source/ModelingComponentsEditorOnly/Private/ToolTargets/VolumeComponentToolTarget.cpp
new file mode 100644
--- /dev/null
+++ b/Engine/Plugins/Runtime/MeshModelingToolset/Source/ModelingComponentsEditorOnly/Private/ToolTargets/VolumeComponentToolTarget.cpp
@@ -0,0 +1,185 @@
+// Copyright Epic Games, Inc. All Rights Reserved.
+
+#pragma once
+
+#include "ToolTargets/VolumeComponentToolTarget.h"
+
+#include "Components/BrushComponent.h"
+#include "ConversionUtils/DynamicMeshToVolume.h"
+#include "ConversionUtils/VolumeToDynamicMesh.h"
+#include "DynamicMeshToMeshDescription.h"
+#include "GameFramework/Volume.h"
+#include "MeshDescriptionToDynamicMesh.h"
+#include "DynamicMesh/MeshNormals.h"
+#include "StaticMeshAttributes.h"
+#include "ToolSetupUtil.h"
+
+using namespace UE::Geometry;
+
+TAutoConsoleVariable<int32> CVarModelingMaxVolumeTriangleCount(
+	TEXT("modeling.VolumeMaxTriCount"),
+	500,
+	TEXT("Limit on triangle count for Volumes that will be emitted by modeling tools. Meshes above this limit will be auto-simplified."));
+
+UVolumeComponentToolTarget::UVolumeComponentToolTarget()
+{
+	// TODO: These should be user-configurable somewhere
+	VolumeToMeshOptions.bInWorldSpace = false;
+	VolumeToMeshOptions.bSetGroups = true;
+	VolumeToMeshOptions.bMergeVertices = true;
+
+	// When a volume has cracks, this option seems to make the geometry
+	// worse rather than better, since the filled in triangles are sometimes
+	// degenerate, folded in on themselves, etc.
+	VolumeToMeshOptions.bAutoRepairMesh = false;
+
+	VolumeToMeshOptions.bOptimizeMesh = true;
+}
+
+int32 UVolumeComponentToolTarget::GetNumMaterials() const
+{
+	return IsValid() ? 1 : 0;
+}
+
+UMaterialInterface* UVolumeComponentToolTarget::GetMaterial(int32 MaterialIndex) const
+{
+	if (!IsValid() || MaterialIndex > 0)
+	{
+		return nullptr;
+	}
+
+	return ToolSetupUtil::GetDefaultEditVolumeMaterial();
+}
+
+void UVolumeComponentToolTarget::GetMaterialSet(FComponentMaterialSet& MaterialSetOut, bool bPreferAssetMaterials) const
+{
+	MaterialSetOut.Materials.Reset();
+	if (IsValid() == false)
+	{
+		return;
+	}
+
+	UMaterialInterface* Material = ToolSetupUtil::GetDefaultEditVolumeMaterial();
+	if (Material)
+	{
+		MaterialSetOut.Materials.Add(Material);
+	}
+}
+
+FDynamicMesh3 UVolumeComponentToolTarget::GetDynamicMesh()
+{
+	UBrushComponent* BrushComponent = Cast<UBrushComponent>(Component);
+	if (!BrushComponent)
+	{
+		return nullptr;
+	}
+	AVolume* Volume = Cast<AVolume>(BrushComponent->GetOwner());
+	if (!Volume)
+	{
+		return nullptr;
+	}
+
+	FDynamicMesh3 DynamicMesh;
+	UE::Conversion::VolumeToDynamicMesh(Volume, DynamicMesh, GetVolumeToMeshOptions());
+	FMeshNormals::InitializeMeshToPerTriangleNormals(&DynamicMesh);
+
+	return DynamicMesh;
+}
+
+void UVolumeComponentToolTarget::CommitDynamicMesh(const UE::Geometry::FDynamicMesh3& Mesh, const FDynamicMeshCommitInfo&)
+{
+	check(IsValid());
+
+	UBrushComponent* BrushComponent = Cast<UBrushComponent>(Component);
+	check(BrushComponent);
+	AVolume* Volume = Cast<AVolume>(BrushComponent->GetOwner());
+	check(Volume);
+
+	FTransform Transform = GetWorldTransform();
+
+	UE::Conversion::FMeshToVolumeOptions ConversionOptions;
+	ConversionOptions.bAutoSimplify = true;
+	ConversionOptions.MaxTriangles = FMath::Max(1, CVarModelingMaxVolumeTriangleCount.GetValueOnGameThread());
+	UE::Conversion::DynamicMeshToVolume(Mesh, Volume, ConversionOptions);
+
+	Volume->SetActorTransform(Transform);
+	Volume->PostEditChange();
+}
+
+const FMeshDescription* UVolumeComponentToolTarget::GetMeshDescription()
+{
+	if (!ConvertedMeshDescription.IsValid())
+	{
+		// Note: We can go directly from a volume to a mesh description using GetBrushMesh() in
+		// Editor.h. However, that path doesn't assign polygroups to the result, which we
+		// typically want when using this target, hence the path through a dynamic mesh.
+
+		TSharedPtr<FDynamicMesh3, ESPMode::ThreadSafe> DynamicMesh = 
+			MakeShared<FDynamicMesh3, ESPMode::ThreadSafe>(GetDynamicMesh());
+
+		ConvertedMeshDescription = MakeShared<FMeshDescription, ESPMode::ThreadSafe>();
+		FStaticMeshAttributes Attributes(*ConvertedMeshDescription);
+		Attributes.Register();
+		FDynamicMeshToMeshDescription Converter;
+		Converter.Convert(DynamicMesh.Get(), *ConvertedMeshDescription);
+	}
+	return ConvertedMeshDescription.Get();
+}
+
+void UVolumeComponentToolTarget::CommitMeshDescription(const FCommitter& Committer)
+{
+	check(IsValid());
+
+	// Let the user fill our mesh description with the Committer
+	if (!ConvertedMeshDescription.IsValid())
+	{
+		ConvertedMeshDescription = MakeShared<FMeshDescription, ESPMode::ThreadSafe>();
+		FStaticMeshAttributes Attributes(*ConvertedMeshDescription);
+		Attributes.Register();
+	}
+	FCommitterParams CommitParams;
+	CommitParams.MeshDescriptionOut = ConvertedMeshDescription.Get();
+	Committer(CommitParams);
+
+	// The conversion we have right now is from dynamic mesh to volume, so we convert
+	// to dynamic mesh first.
+	FDynamicMesh3 DynamicMesh;
+	FMeshDescriptionToDynamicMesh Converter;
+	Converter.Convert(ConvertedMeshDescription.Get(), DynamicMesh);
+
+	CommitDynamicMesh(DynamicMesh);
+}
+
+UBodySetup* UVolumeComponentToolTarget::GetBodySetup() const
+{
+	if (UBrushComponent* BrushComponent = Cast<UBrushComponent>(Component))
+	{
+		return BrushComponent->GetBodySetup();
+	}
+	return nullptr;
+}
+
+
+// Factory
+
+bool UVolumeComponentToolTargetFactory::CanBuildTarget(UObject* SourceObject, const FToolTargetTypeRequirements& Requirements) const
+{
+	UBrushComponent* BrushComponent = Cast<UBrushComponent>(SourceObject);
+	if (!BrushComponent)
+	{
+		return false;
+	}
+
+	return Cast<AVolume>(BrushComponent->GetOwner()) && Requirements.AreSatisfiedBy(UVolumeComponentToolTarget::StaticClass());
+}
+
+UToolTarget* UVolumeComponentToolTargetFactory::BuildTarget(UObject* SourceObject, const FToolTargetTypeRequirements& Requirements)
+{
+	UVolumeComponentToolTarget* Target = NewObject<UVolumeComponentToolTarget>();
+	Target->Component = Cast<UBrushComponent>(SourceObject);
+
+	check(Target->Component && Requirements.AreSatisfiedBy(Target)
+		&& Cast<AVolume>(Cast<UBrushComponent>(SourceObject)->GetOwner()));
+
+	return Target;
+}
\ No newline at end of file
diff --git a/Engine/Plugins/Runtime/MeshModelingToolset/Source/ModelingComponentsEditorOnly/Private/ToolTargets/VolumeDynamicMeshToolTarget.cpp b/Engine/Plugins/Runtime/MeshModelingToolset/Source/ModelingComponentsEditorOnly/Private/ToolTargets/VolumeDynamicMeshToolTarget.cpp
deleted file mode 100644
diff --git a/Engine/Plugins/Runtime/MeshModelingToolset/Source/ModelingComponentsEditorOnly/Private/ToolTargets/VolumeMeshDescriptionToolTarget.cpp b/Engine/Plugins/Runtime/MeshModelingToolset/Source/ModelingComponentsEditorOnly/Private/ToolTargets/VolumeMeshDescriptionToolTarget.cpp
deleted file mode 100644
diff --git a/Engine/Plugins/Runtime/MeshModelingToolset/Source/ModelingComponentsEditorOnly/Public/ToolTargets/VolumeComponentToolTarget.h b/Engine/Plugins/Runtime/MeshModelingToolset/Source/ModelingComponentsEditorOnly/Public/ToolTargets/VolumeComponentToolTarget.h
new file mode 100644
--- /dev/null
+++ b/Engine/Plugins/Runtime/MeshModelingToolset/Source/ModelingComponentsEditorOnly/Public/ToolTargets/VolumeComponentToolTarget.h
@@ -0,0 +1,95 @@
+// Copyright Epic Games, Inc. All Rights Reserved.
+
+#pragma once
+
+#include "CoreMinimal.h"
+
+#include "ConversionUtils/VolumeToDynamicMesh.h" // FVolumeToMeshOptions
+#include "TargetInterfaces/DynamicMeshCommitter.h"
+#include "TargetInterfaces/DynamicMeshProvider.h"
+#include "TargetInterfaces/MaterialProvider.h"
+#include "TargetInterfaces/MeshDescriptionCommitter.h"
+#include "TargetInterfaces/MeshDescriptionProvider.h"
+#include "TargetInterfaces/PhysicsDataSource.h"
+#include "ToolTargets/PrimitiveComponentToolTarget.h"
+
+#include "VolumeComponentToolTarget.generated.h"
+
+/**
+ * The CVar "modeling.VolumeMaxTriCount" is used as a cap on triangles that the various Modeling Mode
+ * Tools will allow an output AVolume to have. If this triangle count is exceeded, the mesh used to
+ * create/update the AVolume will be auto-simplified. This is necessary because all AVolume process is
+ * done on the game thread, and a large Volume (eg with 100k faces) will hang the editor for a long time
+ * when it is created. The default is set to 500.
+ */
+extern MODELINGCOMPONENTSEDITORONLY_API TAutoConsoleVariable<int32> CVarModelingMaxVolumeTriangleCount;
+
+/**
+ * A tool target backed by AVolume
+ */
+UCLASS(Transient)
+class MODELINGCOMPONENTSEDITORONLY_API UVolumeComponentToolTarget : public UPrimitiveComponentToolTarget,
+	public IMaterialProvider, public IPhysicsDataSource,
+	public IDynamicMeshCommitter, public IDynamicMeshProvider,
+	public IMeshDescriptionCommitter, public IMeshDescriptionProvider
+{
+	GENERATED_BODY()
+
+public:
+
+	UVolumeComponentToolTarget();
+
+	const UE::Conversion::FVolumeToMeshOptions& GetVolumeToMeshOptions() { return VolumeToMeshOptions; }
+
+	// IDynamicMeshProvider implementation
+	virtual UE::Geometry::FDynamicMesh3 GetDynamicMesh() override;
+
+	// IDynamicMeshCommitter implementation
+	virtual void CommitDynamicMesh(const UE::Geometry::FDynamicMesh3& Mesh, const FDynamicMeshCommitInfo&) override;
+	using IDynamicMeshCommitter::CommitDynamicMesh; // unhide the other overload
+
+	// IMaterialProvider implementation
+	virtual int32 GetNumMaterials() const override;
+	virtual UMaterialInterface* GetMaterial(int32 MaterialIndex) const override;
+	// Ignores bPreferAssetMaterials
+	virtual void GetMaterialSet(FComponentMaterialSet& MaterialSetOut, bool bPreferAssetMaterials) const override;
+
+	// Doesn't actually do anything for a volume
+	virtual bool CommitMaterialSetUpdate(const FComponentMaterialSet& MaterialSet, bool bApplyToAsset) override { return false; }
+
+	// IMeshDescriptionProvider implementation
+	const FMeshDescription* GetMeshDescription() override;
+
+	// IMeshDescritpionCommitter implementation
+	virtual void CommitMeshDescription(const FCommitter& Committer) override;
+	using IMeshDescriptionCommitter::CommitMeshDescription; // unhide the other overload
+
+	// IPhysicsDataSource implementation
+	virtual UBodySetup* GetBodySetup() const override;
+	// always returns null because volumes do not support IInterface_CollisionDataProvider
+	virtual IInterface_CollisionDataProvider* GetComplexCollisionProvider() const override { return nullptr; }
+
+	// Rest provided by parent class
+
+protected:
+	UE::Conversion::FVolumeToMeshOptions VolumeToMeshOptions;
+
+	// This isn't for caching- we have to take ownership of the mesh description because it is
+	// expected for things like a static mesh.
+	TSharedPtr<FMeshDescription, ESPMode::ThreadSafe> ConvertedMeshDescription;
+
+	friend class UVolumeComponentToolTargetFactory;
+};
+
+/** Factory for UVolumeComponentToolTarget to be used by the target manager. */
+UCLASS(Transient)
+class MODELINGCOMPONENTSEDITORONLY_API UVolumeComponentToolTargetFactory : public UToolTargetFactory
+{
+	GENERATED_BODY()
+
+public:
+
+	virtual bool CanBuildTarget(UObject* SourceObject, const FToolTargetTypeRequirements& TargetTypeInfo) const override;
+
+	virtual UToolTarget* BuildTarget(UObject* SourceObject, const FToolTargetTypeRequirements& TargetTypeInfo) override;
+};
\ No newline at end of file
diff --git a/Engine/Plugins/Runtime/MeshModelingToolset/Source/ModelingComponentsEditorOnly/Public/ToolTargets/VolumeDynamicMeshToolTarget.h b/Engine/Plugins/Runtime/MeshModelingToolset/Source/ModelingComponentsEditorOnly/Public/ToolTargets/VolumeDynamicMeshToolTarget.h
deleted file mode 100644
diff --git a/Engine/Plugins/Runtime/MeshModelingToolset/Source/ModelingComponentsEditorOnly/Public/ToolTargets/VolumeMeshDescriptionToolTarget.h b/Engine/Plugins/Runtime/MeshModelingToolset/Source/ModelingComponentsEditorOnly/Public/ToolTargets/VolumeMeshDescriptionToolTarget.h
deleted file mode 100644
