MeshModelingTools: Add IPhysicsDataSource implementation to DynamicMeshComponentToolTarget.

#rnx

diff --git a/Engine/Plugins/Runtime/MeshModelingToolset/Source/ModelingComponentsEditorOnly/Private/ToolTargets/DynamicMeshComponentToolTarget.cpp b/Engine/Plugins/Runtime/MeshModelingToolset/Source/ModelingComponentsEditorOnly/Private/ToolTargets/DynamicMeshComponentToolTarget.cpp
--- a/Engine/Plugins/Runtime/MeshModelingToolset/Source/ModelingComponentsEditorOnly/Private/ToolTargets/DynamicMeshComponentToolTarget.cpp
+++ b/Engine/Plugins/Runtime/MeshModelingToolset/Source/ModelingComponentsEditorOnly/Private/ToolTargets/DynamicMeshComponentToolTarget.cpp
@@ -191,6 +191,16 @@ void UDynamicMeshComponentToolTarget::CommitDynamicMesh(const FDynamicMesh3& Upd
 		*this, UpdatedMesh, CommitInfo.bTopologyChanged);
 }
 
+UBodySetup* UDynamicMeshComponentToolTarget::GetBodySetup() const
+{
+	return IsValid() ? Cast<UDynamicMeshComponent>(Component)->GetBodySetup() : nullptr;
+}
+
+
+IInterface_CollisionDataProvider* UDynamicMeshComponentToolTarget::GetComplexCollisionProvider() const
+{
+	return IsValid() ? Cast<UDynamicMeshComponent>(Component) : nullptr;
+}
 
 
 // Factory
diff --git a/Engine/Plugins/Runtime/MeshModelingToolset/Source/ModelingComponentsEditorOnly/Public/ToolTargets/DynamicMeshComponentToolTarget.h b/Engine/Plugins/Runtime/MeshModelingToolset/Source/ModelingComponentsEditorOnly/Public/ToolTargets/DynamicMeshComponentToolTarget.h
--- a/Engine/Plugins/Runtime/MeshModelingToolset/Source/ModelingComponentsEditorOnly/Public/ToolTargets/DynamicMeshComponentToolTarget.h
+++ b/Engine/Plugins/Runtime/MeshModelingToolset/Source/ModelingComponentsEditorOnly/Public/ToolTargets/DynamicMeshComponentToolTarget.h
@@ -10,6 +10,7 @@
 #include "TargetInterfaces/MeshDescriptionCommitter.h"
 #include "TargetInterfaces/MeshDescriptionProvider.h"
 #include "TargetInterfaces/DynamicMeshSource.h"
+#include "TargetInterfaces/PhysicsDataSource.h"
 #include "ToolTargets/PrimitiveComponentToolTarget.h"
 
 #include "DynamicMeshComponentToolTarget.generated.h"
@@ -27,7 +28,8 @@ class MODELINGCOMPONENTSEDITORONLY_API UDynamicMeshComponentToolTarget :
 	public IDynamicMeshProvider,
 	public IDynamicMeshCommitter,
 	public IMaterialProvider,
-	public IPersistentDynamicMeshSource
+	public IPersistentDynamicMeshSource,
+	public IPhysicsDataSource
 {
 	GENERATED_BODY()
 
@@ -62,6 +64,10 @@ public:
 	virtual bool HasDynamicMeshComponent() const override;
 	virtual UDynamicMeshComponent* GetDynamicMeshComponent() override;
 
+	// IPhysicsDataSource implementation
+	virtual UBodySetup* GetBodySetup() const override;
+	virtual IInterface_CollisionDataProvider* GetComplexCollisionProvider() const override;
+
 	// Rest provided by parent class
 
 protected:
