Added support for world-space hit-testing and snapping to Volume and DynamicMeshComponent mesh vertices/edges/faces, by having Modeling Mode maintain a persistent set of spatial data structures for these scene meshes. With this change, new ToolSceneQueriesUtil functions can be used to provide more general scene mesh-element hit-testing/snapping via ModelingSceneSnappingManager. Updates to existing Tools/Mechanics to use these new query functions will be done in later CLs.

ToolsFramework:
- added new USceneSnappingManager::ExecuteSceneHitQuery API function, as well as accompanying FSceneHitQueryRequest/FSceneHitQueryResult structs
- added FSceneQueryVisibilityFilter struct which contains ComponentsToIgnore/InvisibleComponentsToInclude lists for snapping queries and provides shared IsVisible() implementation

ModelingComponents:
- add FLevelObjectsObserver, tracks Actors added/removed to a Level via various editor events
- add FSceneGeometrySpatialCache, maintains spatial data structures for known types of mesh-backed PrimitiveComponents, an octree based on their world-space bounding-boxes, and provides various spatial-query functions
- ModelingSceneSnappingManager now maintains a FSceneGeometrySpatialCache for a set of Actors/Components it is notified about, and updates the spatial cache if those components are modified or translated. Currently supports Volumes and DynamicMeshComponents.
- add ModelingSceneSnappingManager::ExecuteSceneHitQuery implementation, does raycasts into the world (for static mesh components) and FSceneGeometrySpatialCache for Volumes/DMCs
- modified ModelingSceneSnappingManager::ExecuteSceneSnapQuery to also do combined query against StaticMeshComponents via world-linetrace, and FSceneGeometrySpatialCache for Volumes/DMCs
- add ToolSceneQueriesUtil::FindNearestVisibleObjectHit variants that take a USceneSnappingManager or UInteractiveTool, and do queries via the SceneSnappingManger, instead of directly linetracing into a UWorld

ModleingToolsEditorMode:
- add FLevelObjectsObserver instance and use it to drive updates to active ModelingSceneSnappingManager

#rnx

diff --git a/Engine/Plugins/Editor/ModelingToolsEditorMode/Source/ModelingToolsEditorMode/Private/ModelingToolsEditorMode.cpp b/Engine/Plugins/Editor/ModelingToolsEditorMode/Source/ModelingToolsEditorMode/Private/ModelingToolsEditorMode.cpp
--- a/Engine/Plugins/Editor/ModelingToolsEditorMode/Source/ModelingToolsEditorMode/Private/ModelingToolsEditorMode.cpp
+++ b/Engine/Plugins/Editor/ModelingToolsEditorMode/Source/ModelingToolsEditorMode/Private/ModelingToolsEditorMode.cpp
@@ -20,13 +20,11 @@
 #include "Selection/PersistentMeshSelectionManager.h"
 #include "Selection/StoredMeshSelectionUtil.h"
 #include "Snapping/ModelingSceneSnappingManager.h"
+#include "Scene/LevelObjectsObserver.h"
 
 #include "Features/IModularFeatures.h"
 #include "ModelingModeToolExtensions.h"
 
-//#include "SingleClickTool.h"
-//#include "MeshSurfacePointTool.h"
-//#include "MeshVertexDragTool.h"
 #include "DynamicMeshSculptTool.h"
 #include "MeshVertexSculptTool.h"
 #include "EditMeshPolygonsTool.h"
@@ -119,8 +117,6 @@
 #define LOCTEXT_NAMESPACE "UModelingToolsEditorMode"
 
 
-//#define ENABLE_DEBUG_PRINTING
-
 const FEditorModeID UModelingToolsEditorMode::EM_ModelingToolsEditorModeId = TEXT("EM_ModelingToolsEditorMode");
 
 FDateTime UModelingToolsEditorMode::LastModeStartTimestamp;
@@ -309,6 +305,12 @@ public:
 
 };
 
+
+
+
+
+
+
 void UModelingToolsEditorMode::Enter()
 {
 	UEdMode::Enter();
@@ -330,6 +332,26 @@ void UModelingToolsEditorMode::Enter()
 
 	// register snapping manager
 	UE::Geometry::RegisterSceneSnappingManager(GetInteractiveToolsContext());
+	SceneSnappingManager = UE::Geometry::FindModelingSceneSnappingManager(GetToolManager());
+
+	// register level objects observer that will update the snapping manager as the scene changes
+	LevelObjectsObserver = MakeShared<FLevelObjectsObserver>();
+	LevelObjectsObserver->OnActorAdded.AddLambda([this](AActor* Actor)
+	{
+		if (SceneSnappingManager)
+		{
+			SceneSnappingManager->OnActorAdded(Actor, [](UPrimitiveComponent*) { return true; });
+		}
+	});
+	LevelObjectsObserver->OnActorRemoved.AddLambda([this](AActor* Actor)
+	{
+		if (SceneSnappingManager)
+		{
+			SceneSnappingManager->OnActorRemoved(Actor);
+		}
+	});
+	// tracker will auto-populate w/ the current level, but must have registered the handlers first!
+	LevelObjectsObserver->Initialize(GetWorld());
 
 	// register selection manager, if this feature is enabled in the mode settings
 	const UModelingToolsEditorModeSettings* ModelingModeSettings = GetDefault<UModelingToolsEditorModeSettings>();
@@ -349,7 +371,7 @@ void UModelingToolsEditorMode::Enter()
 		{
 			return UE::Modeling::GetNewAssetPathName(BaseName, TargetWorld, SuggestedFolder);
 		});
-		MeshCreatedEventHandle = ModelCreationAPI->OnModelingMeshCreated.AddLambda([](const FCreateMeshObjectResult& CreatedInfo) 
+		MeshCreatedEventHandle = ModelCreationAPI->OnModelingMeshCreated.AddLambda([this](const FCreateMeshObjectResult& CreatedInfo) 
 		{
 			if (CreatedInfo.NewAsset != nullptr)
 			{
@@ -776,8 +798,11 @@ void UModelingToolsEditorMode::Exit()
 	// will be called before our Exit()
 	//UE::TransformGizmoUtil::DeregisterTransformGizmoContextObject(ToolsContext.Get());
 	
-	// deregister snapping manager
+	// deregister snapping manager and shut down level objects tracker
+	LevelObjectsObserver->Shutdown();		// do this first because it is going to fire events on the snapping manager
+	LevelObjectsObserver.Reset();
 	UE::Geometry::DeregisterSceneSnappingManager(GetInteractiveToolsContext());
+	SceneSnappingManager = nullptr;
 
 	// TODO: cannot deregister currently because if another mode is also registering, its Enter()
 	// will be called before our Exit()
diff --git a/Engine/Plugins/Editor/ModelingToolsEditorMode/Source/ModelingToolsEditorMode/Public/ModelingToolsEditorMode.h b/Engine/Plugins/Editor/ModelingToolsEditorMode/Source/ModelingToolsEditorMode/Public/ModelingToolsEditorMode.h
--- a/Engine/Plugins/Editor/ModelingToolsEditorMode/Source/ModelingToolsEditorMode/Public/ModelingToolsEditorMode.h
+++ b/Engine/Plugins/Editor/ModelingToolsEditorMode/Source/ModelingToolsEditorMode/Public/ModelingToolsEditorMode.h
@@ -15,6 +15,8 @@
 class FEditorComponentSourceFactory;
 class FUICommandList;
 class FStylusStateTracker;		// for stylus events
+class FLevelObjectsObserver;
+class UModelingSceneSnappingManager;
 
 UCLASS(Transient)
 class UModelingToolsEditorMode : public UBaseLegacyWidgetEdMode
@@ -73,6 +75,11 @@ protected:
 
 	TUniquePtr<FStylusStateTracker> StylusStateTracker;
 
+	TSharedPtr<FLevelObjectsObserver> LevelObjectsObserver;
+
+	UPROPERTY()
+	UModelingSceneSnappingManager* SceneSnappingManager;
+
 	void ModelingModeShortcutRequested(EModelingModeActionCommands Command);
 	void FocusCameraAtCursorHotkey();
 
diff --git a/Engine/Plugins/Runtime/MeshModelingToolset/Source/ModelingComponents/Private/Scene/SceneGeometrySpatialCache.cpp b/Engine/Plugins/Runtime/MeshModelingToolset/Source/ModelingComponents/Private/Scene/SceneGeometrySpatialCache.cpp
new file mode 100644
--- /dev/null
+++ b/Engine/Plugins/Runtime/MeshModelingToolset/Source/ModelingComponents/Private/Scene/SceneGeometrySpatialCache.cpp
@@ -0,0 +1,550 @@
+// Copyright Epic Games, Inc. All Rights Reserved.
+
+#pragma once
+
+#include "Scene/SceneGeometrySpatialCache.h"
+
+#include "Engine/Brush.h"
+#include "GameFramework/Volume.h"
+#include "Components/BrushComponent.h"
+#include "ConversionUtils/VolumeToDynamicMesh.h"
+
+using namespace UE::Geometry;
+
+
+namespace UE
+{
+namespace Local
+{
+	
+static bool IsHiddenComponent(UPrimitiveComponent* Component)
+{
+#if WITH_EDITOR
+	return (Component != nullptr) && ( (Component->IsVisible() == false) || (Component->IsVisibleInEditor() == false) );
+#else
+	return (Component != nullptr) && (Component->IsVisible() == false);
+#endif
+}
+
+}
+}
+
+
+std::atomic<int32> ISceneGeometrySpatial::UniqueIDGenerator = 1;
+
+ISceneGeometrySpatial::ISceneGeometrySpatial()
+{
+	SpatialID = UniqueIDGenerator++;
+}
+
+
+
+
+class FSceneVolumeSpatial : public ISceneGeometrySpatial
+{
+public:
+	virtual void Initialize(UBrushComponent* Component);
+
+	virtual bool IsValid() override;
+	virtual bool IsVisible() override;
+
+	virtual UPrimitiveComponent* GetComponent() override;
+	virtual FAxisAlignedBox3d GetWorldBounds() override;
+
+	virtual bool FindNearestHit(
+		const FRay3d& WorldRay, 
+		FSceneGeometryPoint& HitResultOut,
+		double MaxDistance = TNumericLimits<double>::Max()) override;
+	virtual bool GetGeometry(EGeometryPointType GeometryType, int32 GeometryIndex, bool bWorldSpace, FVector3d& A, FVector3d& B, FVector3d& C) override;
+
+	virtual void OnGeometryModified(bool bDeferSpatialRebuild) override;
+	virtual void OnTransformModified() override;
+
+public:
+	TWeakObjectPtr<UBrushComponent> Component;
+	FAxisAlignedBox3d LocalBounds;
+	UE::Geometry::FTransform3d Transform;
+	FColliderMesh ColliderMesh;
+	bool bDeferredRebuildPending;
+};
+
+
+
+
+
+void FSceneVolumeSpatial::Initialize(UBrushComponent* ComponentIn)
+{
+	this->Component = ComponentIn;
+	bDeferredRebuildPending = true;
+}
+
+bool FSceneVolumeSpatial::IsValid()
+{
+	return Component.IsValid(false);
+}
+
+bool FSceneVolumeSpatial::IsVisible()
+{
+	return IsValid() && UE::Local::IsHiddenComponent(Component.Get()) == false;
+}
+
+UPrimitiveComponent* FSceneVolumeSpatial::GetComponent() 
+{ 
+	return Component.Get(); 
+}
+
+FAxisAlignedBox3d FSceneVolumeSpatial::GetWorldBounds()
+{
+	return FAxisAlignedBox3d(LocalBounds, Transform);
+}
+
+bool FSceneVolumeSpatial::FindNearestHit(
+	const FRay3d& WorldRay, 
+	FSceneGeometryPoint& HitResultOut,
+	double MaxDistance)
+{
+	if (IsValid() == false)
+	{
+		return false;
+	}
+
+	if (bDeferredRebuildPending)
+	{
+		OnGeometryModified(false);
+		ensure(bDeferredRebuildPending == false);
+	}
+
+	FRay3d LocalRay = Transform.InverseTransformRay(WorldRay);
+	double RayHitT; int32 HitTID; FVector3d HitBaryCoords;
+	if ( ColliderMesh.GetAABBTree().FindNearestHitTriangle(LocalRay, RayHitT, HitTID, HitBaryCoords) )
+	{
+		HitResultOut.Component = Component.Get();
+		HitResultOut.Actor = HitResultOut.Component->GetOwner();
+		HitResultOut.GeometryIndex = HitTID;
+		HitResultOut.GeometryType = EGeometryPointType::Triangle;
+		HitResultOut.PointBaryCoords = HitBaryCoords;
+		HitResultOut.LocalPoint = LocalRay.PointAt(RayHitT);
+		HitResultOut.WorldPoint = Transform.TransformPosition(HitResultOut.LocalPoint);
+		HitResultOut.Ray = WorldRay;
+		HitResultOut.RayDistance = WorldRay.GetParameter(HitResultOut.WorldPoint);
+		return true;
+	}
+	return false;
+}
+
+
+bool FSceneVolumeSpatial::GetGeometry(EGeometryPointType GeometryType, int32 GeometryIndex, bool bWorldSpace, FVector3d& A, FVector3d& B, FVector3d& C)
+{
+	if (GeometryType == EGeometryPointType::Triangle)
+	{
+		if (ColliderMesh.IsTriangle(GeometryIndex))
+		{
+			ColliderMesh.GetTriVertices(GeometryIndex, A, B, C);
+			if (bWorldSpace)
+			{
+				A = Transform.TransformPosition(A);
+				B = Transform.TransformPosition(B);
+				C = Transform.TransformPosition(C);
+			}
+			return true;
+		}
+	}
+	return false;
+}
+
+
+void FSceneVolumeSpatial::OnGeometryModified(bool bDeferSpatialRebuild)
+{
+	ColliderMesh.Reset(false);
+	if (IsValid() == false)
+	{
+		return;
+	}
+
+	Transform = (UE::Geometry::FTransform3d)Component->GetComponentTransform();
+	LocalBounds = FAxisAlignedBox3d( Component->CalcBounds(FTransform::Identity).GetBox() );
+
+	if (bDeferSpatialRebuild)
+	{
+		bDeferredRebuildPending = true;
+		return;
+	}
+
+	AVolume* Volume = Cast<AVolume>(Component->GetOwner());
+	if (Volume)
+	{
+		UE::Conversion::FVolumeToMeshOptions Options;
+		Options.bSetGroups = false;
+		Options.bMergeVertices = false;
+		Options.bAutoRepairMesh = false;
+		Options.bOptimizeMesh = false;
+		FDynamicMesh3 TmpMesh;
+		UE::Conversion::VolumeToDynamicMesh(Volume, TmpMesh, Options);
+		ColliderMesh.Initialize(TmpMesh);
+	}
+
+	bDeferredRebuildPending = false;
+}
+
+
+void FSceneVolumeSpatial::OnTransformModified()
+{
+	if (IsValid())
+	{
+		Transform = (UE::Geometry::FTransform3d)Component->GetComponentTransform();
+	}
+}
+
+
+
+
+
+
+class FSceneDynamicMeshSpatial : public ISceneGeometrySpatial
+{
+public:
+	virtual void Initialize(UDynamicMeshComponent* Component);
+
+	virtual bool IsValid() override;
+	virtual bool IsVisible() override;
+
+	virtual UPrimitiveComponent* GetComponent() override;
+	virtual FAxisAlignedBox3d GetWorldBounds() override;
+
+	virtual bool FindNearestHit(
+		const FRay3d& WorldRay, 
+		FSceneGeometryPoint& HitResultOut,
+		double MaxDistance = TNumericLimits<double>::Max()) override;
+	virtual bool GetGeometry(EGeometryPointType GeometryType, int32 GeometryIndex, bool bWorldSpace, FVector3d& A, FVector3d& B, FVector3d& C) override;
+
+	virtual void OnGeometryModified(bool bDeferSpatialRebuild) override;
+	virtual void OnTransformModified() override;
+
+public:
+	TWeakObjectPtr<UDynamicMeshComponent> Component;
+	FAxisAlignedBox3d LocalBounds;
+	UE::Geometry::FTransform3d Transform;
+	FDynamicMeshAABBTree3 AABBTree;
+	bool bDeferredRebuildPending;
+};
+
+
+
+void FSceneDynamicMeshSpatial::Initialize(UDynamicMeshComponent* ComponentIn)
+{
+	this->Component = ComponentIn;
+	bDeferredRebuildPending = true;
+}
+
+bool FSceneDynamicMeshSpatial::IsValid()
+{
+	return Component.IsValid(false);
+}
+
+bool FSceneDynamicMeshSpatial::IsVisible()
+{
+	return IsValid() && UE::Local::IsHiddenComponent(Component.Get()) == false;
+}
+
+UPrimitiveComponent* FSceneDynamicMeshSpatial::GetComponent() 
+{ 
+	return Component.Get(); 
+}
+
+FAxisAlignedBox3d FSceneDynamicMeshSpatial::GetWorldBounds()
+{
+	return FAxisAlignedBox3d(LocalBounds, Transform);
+}
+
+bool FSceneDynamicMeshSpatial::FindNearestHit(
+	const FRay3d& WorldRay, 
+	FSceneGeometryPoint& HitResultOut,
+	double MaxDistance)
+{
+	if (IsValid() == false)
+	{
+		return false;
+	}
+
+	if (bDeferredRebuildPending)
+	{
+		OnGeometryModified(false);
+		ensure(bDeferredRebuildPending == false);
+	}
+
+	FRay3d LocalRay = Transform.InverseTransformRay(WorldRay);
+	double RayHitT; int32 HitTID; FVector3d HitBaryCoords;
+	if (AABBTree.FindNearestHitTriangle(LocalRay, RayHitT, HitTID, HitBaryCoords))
+	{
+		HitResultOut.Component = Component.Get();
+		HitResultOut.Actor = HitResultOut.Component->GetOwner();
+		HitResultOut.GeometryIndex = HitTID;
+		HitResultOut.GeometryType = EGeometryPointType::Triangle;
+		HitResultOut.PointBaryCoords = HitBaryCoords;
+		HitResultOut.LocalPoint = LocalRay.PointAt(RayHitT);
+		HitResultOut.WorldPoint = Transform.TransformPosition(HitResultOut.LocalPoint);
+		HitResultOut.Ray = WorldRay;
+		HitResultOut.RayDistance = WorldRay.GetParameter(HitResultOut.WorldPoint);
+		return true;
+	}
+	return false;
+}
+
+
+bool FSceneDynamicMeshSpatial::GetGeometry(EGeometryPointType GeometryType, int32 GeometryIndex, bool bWorldSpace, FVector3d& A, FVector3d& B, FVector3d& C)
+{
+	bool bFound = false;
+	if (GeometryType == EGeometryPointType::Triangle)
+	{
+		UDynamicMeshComponent* DynamicMeshComponent = Component.Get();
+		if (DynamicMeshComponent)
+		{
+			DynamicMeshComponent->ProcessMesh([&](const FDynamicMesh3& Mesh)
+			{
+				if (Mesh.IsTriangle(GeometryIndex))
+				{
+					Mesh.GetTriVertices(GeometryIndex, A, B, C);
+					if (bWorldSpace)
+					{
+						A = Transform.TransformPosition(A);
+						B = Transform.TransformPosition(B);
+						C = Transform.TransformPosition(C);
+					}
+					bFound = true;
+				}
+			});
+		}
+	}
+	return bFound;
+}
+
+
+void FSceneDynamicMeshSpatial::OnGeometryModified(bool bDeferSpatialRebuild)
+{
+	if (IsValid() == false)
+	{
+		return;
+	}
+	UDynamicMeshComponent* DynamicMeshComponent = Component.Get();
+	if (!DynamicMeshComponent)
+	{
+		return;
+	}
+
+	Transform = (UE::Geometry::FTransform3d)DynamicMeshComponent->GetComponentTransform();
+	// may be reasonable to use USceneComponent::GetLocalBounds() here. Unclear if it will have already been updated on mesh changes...
+	DynamicMeshComponent->ProcessMesh([&](const FDynamicMesh3& Mesh)
+	{
+		LocalBounds = Mesh.GetBounds(true);
+	});
+
+	if (bDeferSpatialRebuild)
+	{
+		bDeferredRebuildPending = true;
+		return;
+	}
+
+	DynamicMeshComponent->ProcessMesh([&](const FDynamicMesh3& Mesh)
+	{
+		AABBTree.SetMesh(&Mesh, true);
+	});
+
+	bDeferredRebuildPending = false;
+}
+
+
+void FSceneDynamicMeshSpatial::OnTransformModified()
+{
+	if (IsValid())
+	{
+		Transform = (UE::Geometry::FTransform3d)Component->GetComponentTransform();
+	}
+}
+
+
+
+
+
+
+FSceneGeometrySpatialCache::FSceneGeometrySpatialCache()
+{
+	BoundsOctree.RootDimension = 100000;		// 1000m, this can do until we come up w/ something better?
+	BoundsOctree.SetMaxTreeDepth(7);
+}
+
+FSceneGeometrySpatialCache::~FSceneGeometrySpatialCache()
+{
+}
+
+
+bool FSceneGeometrySpatialCache::EnableComponentTracking(UPrimitiveComponent* Component, FSceneGeometryID& IdentifierOut)
+{
+	if (ComponentMap.Contains(Component))
+	{
+		IdentifierOut = ComponentMap[Component];
+		return true;
+	}
+
+	TSharedPtr<ISceneGeometrySpatial> NewSpatial;
+	if (Cast<UBrushComponent>(Component) != nullptr)
+	{
+		IdentifierOut = FSceneGeometryID{ UniqueIDGenerator++ };
+
+		TSharedPtr<FSceneVolumeSpatial> Spatial = MakeShared<FSceneVolumeSpatial>();
+		Spatial->Initialize(Cast<UBrushComponent>(Component));
+		NewSpatial = Spatial;
+	}
+	else if (Cast<UDynamicMeshComponent>(Component) != nullptr)
+	{
+		IdentifierOut = FSceneGeometryID{ UniqueIDGenerator++ };
+
+		TSharedPtr<FSceneDynamicMeshSpatial> Spatial = MakeShared<FSceneDynamicMeshSpatial>();
+		Spatial->Initialize(Cast<UDynamicMeshComponent>(Component));
+		NewSpatial = Spatial;
+	}
+
+	if (NewSpatial != nullptr)
+	{
+		ComponentMap.Add(Component, IdentifierOut);
+		Spatials.Add(IdentifierOut, NewSpatial);
+
+		// do elsewhere?
+		NewSpatial->OnGeometryModified(true);
+
+		BoundsOctree.InsertObject(IdentifierOut.UniqueID, NewSpatial->GetWorldBounds());
+
+		return true;
+	}
+
+	return false;
+}
+
+
+
+void FSceneGeometrySpatialCache::DisableComponentTracking(UPrimitiveComponent* Component)
+{
+	if (ComponentMap.Contains(Component))
+	{
+		FSceneGeometryID Identifier = ComponentMap[Component];
+		ComponentMap.Remove(Component);
+		DisableGeometryTracking(Identifier);
+	}
+}
+
+
+void FSceneGeometrySpatialCache::DisableGeometryTracking(FSceneGeometryID& Identifier)
+{
+	if (Spatials.Contains(Identifier))
+	{
+		BoundsOctree.RemoveObject(Identifier.UniqueID);
+
+		Spatials.Remove(Identifier);
+	}
+}
+
+
+
+bool FSceneGeometrySpatialCache::FindNearestHit(
+	const FRay3d& WorldRay,
+	FSceneGeometryPoint& HitResultOut,
+	FSceneGeometryID& HitIdentifier,
+	const FSceneQueryVisibilityFilter* VisibilityFilter,
+	double MaxDistance)
+{
+
+	HitResultOut = FSceneGeometryPoint();
+	bool bFoundAnyHit = false;
+
+	// TODO: this current approach will force any hit objects to run any deferred AABBTree builds.
+	// What we probably should do is have a FindAllHitObjects query that only tests the boxes, and
+	// then for those hits, parallel-build any pending AABBTrees before testing them.
+
+	// cast ray into octree
+	int32 HitSpatialID = BoundsOctree.FindNearestHitObject(WorldRay,
+		[&](int32 Identifier) { 
+			return Spatials[FSceneGeometryID(Identifier)]->GetWorldBounds(); 
+		},
+		[&](int Identifier, const FRay3d& Ray) {
+			TSharedPtr<ISceneGeometrySpatial> Spatial = Spatials[FSceneGeometryID(Identifier)];
+			if (VisibilityFilter != nullptr && VisibilityFilter->IsVisible(Spatial->GetComponent()) == false)
+			{
+				return TNumericLimits<double>::Max();
+			}		
+			FSceneGeometryPoint Hit;
+			if (Spatial->FindNearestHit(WorldRay, Hit, MaxDistance))
+			{
+				return Hit.RayDistance;
+			}
+			return TNumericLimits<double>::Max();
+		});
+
+	// if we hit something above, repeat it's hit-query here (dumb...)
+	if (HitSpatialID >= 0)
+	{
+		TSharedPtr<ISceneGeometrySpatial> Spatial = Spatials[FSceneGeometryID(HitSpatialID)];
+		FSceneGeometryPoint Hit;
+		bool bHit = Spatial->FindNearestHit(WorldRay, Hit, MaxDistance);
+		if (ensure(bHit))
+		{
+			HitResultOut = Hit;
+			HitIdentifier = FSceneGeometryID(HitSpatialID);
+			return true;
+		}
+	}
+
+	return false;
+}
+
+
+bool FSceneGeometrySpatialCache::GetGeometry(FSceneGeometryID Identifier, EGeometryPointType GeometryType, int32 GeometryIndex, bool bWorldSpace, FVector3d& A, FVector3d& B, FVector3d& C)
+{
+	TSharedPtr<ISceneGeometrySpatial>* Found = Spatials.Find(Identifier);
+	if (Found != nullptr)
+	{
+		return (*Found)->GetGeometry(GeometryType, GeometryIndex, bWorldSpace, A, B, C);
+	}
+	return false;
+}
+
+
+bool FSceneGeometrySpatialCache::HaveCacheForComponent(UPrimitiveComponent* Component)
+{
+	const FSceneGeometryID* FoundID = ComponentMap.Find(Component);
+	if (FoundID != nullptr)
+	{
+		TSharedPtr<ISceneGeometrySpatial> Spatial = Spatials[*FoundID];
+		return Spatial->IsValid();
+	}
+	return false;
+}
+
+
+bool FSceneGeometrySpatialCache::NotifyTransformUpdate(UPrimitiveComponent* Component)
+{
+	const FSceneGeometryID* FoundID = ComponentMap.Find(Component);
+	if (FoundID != nullptr)
+	{
+		TSharedPtr<ISceneGeometrySpatial> Spatial = Spatials[*FoundID];
+		Spatial->OnTransformModified();
+
+		BoundsOctree.ReinsertObject(FoundID->UniqueID, Spatial->GetWorldBounds());
+
+		return true;
+	}
+	return false;	
+}
+
+bool FSceneGeometrySpatialCache::NotifyGeometryUpdate(UPrimitiveComponent* Component, bool bDeferRebuild)
+{
+	const FSceneGeometryID* FoundID = ComponentMap.Find(Component);
+	if (FoundID != nullptr)
+	{
+		TSharedPtr<ISceneGeometrySpatial> Spatial = Spatials[*FoundID];
+		Spatial->OnGeometryModified(bDeferRebuild);
+
+		BoundsOctree.ReinsertObject(FoundID->UniqueID, Spatial->GetWorldBounds());
+
+		return true;
+	}
+	return false;	
+}
\ No newline at end of file
diff --git a/Engine/Plugins/Runtime/MeshModelingToolset/Source/ModelingComponents/Private/Snapping/ModelingSceneSnappingManager.cpp b/Engine/Plugins/Runtime/MeshModelingToolset/Source/ModelingComponents/Private/Snapping/ModelingSceneSnappingManager.cpp
--- a/Engine/Plugins/Runtime/MeshModelingToolset/Source/ModelingComponents/Private/Snapping/ModelingSceneSnappingManager.cpp
+++ b/Engine/Plugins/Runtime/MeshModelingToolset/Source/ModelingComponents/Private/Snapping/ModelingSceneSnappingManager.cpp
@@ -4,13 +4,21 @@
 
 #include "InteractiveToolsContext.h"
 #include "InteractiveToolManager.h"
+#include "InteractiveToolObjects.h"
 #include "ContextObjectStore.h"
 
+#include "Scene/SceneGeometrySpatialCache.h"
+
 #include "Engine/World.h"
+#include "GameFramework/Actor.h"
+#include "EngineUtils.h" // for TActorIterator<>
 #include "Engine/StaticMesh.h"
 #include "Components/PrimitiveComponent.h"
 #include "Components/StaticMeshComponent.h"
+#include "Components/DynamicMeshComponent.h"
+#include "UObject/UObjectGlobals.h"
 
+using namespace UE::Geometry;
 
 #define LOCTEXT_NAMESPACE "USceneSnappingManager"
 
@@ -71,11 +79,189 @@ void UModelingSceneSnappingManager::Initialize(TObjectPtr<UInteractiveToolsConte
 
 	QueriesAPI = (ParentContext && ParentContext->ToolManager) ?
 		ParentContext->ToolManager->GetContextQueriesAPI() : nullptr;
+
+	SpatialCache = MakeShared<FSceneGeometrySpatialCache>();
+
+#if WITH_EDITOR
+	OnObjectModifiedHandler = FCoreUObjectDelegates::OnObjectModified.AddLambda([this](UObject* Object) { 
+		HandleGlobalObjectModifiedDelegate(Object);
+	});
+
+	if (GEngine)
+	{
+		OnComponentTransformChangedHandle = GEngine->OnComponentTransformChanged().AddLambda([this](USceneComponent* Component, ETeleportType) {
+			HandleGlobalComponentTransformChangedDelegate(Component);
+		});
+	}
+#endif
 }
 
 void UModelingSceneSnappingManager::Shutdown()
 {
+#if WITH_EDITOR
+	FCoreUObjectDelegates::OnObjectModified.Remove(OnObjectModifiedHandler);
+
+	if (GEngine)
+	{
+		GEngine->OnComponentTransformChanged().Remove(OnComponentTransformChangedHandle);
+	}
+#endif
+}
+
+
+bool UModelingSceneSnappingManager::IsComponentTypeSupported(const UPrimitiveComponent* Component) const
+{
+	// supported types
+	return ( Cast<UStaticMeshComponent>(Component) != nullptr ||
+			 Cast<UBrushComponent>(Component) != nullptr ||
+			 Cast<UDynamicMeshComponent>(Component) != nullptr );
+}
+
+void UModelingSceneSnappingManager::OnActorAdded(AActor* Actor, TFunctionRef<bool(UPrimitiveComponent*)> PrimitiveFilter)
+{
+	// TODO: additional filtering here?
+
+	// ignore internal tool framework Actors
+	if (Cast<AInternalToolFrameworkActor>(Actor) != nullptr)
+	{
+		return;
+	}
+
+	TArray<AActor*> AllActors;
+	Actor->GetAllChildActors(AllActors, true);
+	AllActors.Add(Actor);
+	for (AActor* ProcessActor : AllActors)
+	{
+		for (UActorComponent* Component : ProcessActor->GetComponents())
+		{
+			UPrimitiveComponent* PrimComponent = Cast<UPrimitiveComponent>(Component);
+			if (PrimComponent && PrimitiveFilter(PrimComponent) )
+			{
+				OnComponentAdded(PrimComponent);
+			}
+		}
+	}
+}
+
+void UModelingSceneSnappingManager::OnActorRemoved(AActor* Actor)
+{
+	// cannot assume actor is still valid...
+	TArray<UPrimitiveComponent*> ToRemove;
+	for (TPair<UPrimitiveComponent*, AActor*> Pair : ComponentToActorMap)
+	{
+		if (Pair.Value == Actor)
+		{
+			ToRemove.Add(Pair.Key);
+		}
+	}
+	for (UPrimitiveComponent* Component : ToRemove)
+	{
+		OnComponentRemoved(Component);
+	}
+}
+
+void UModelingSceneSnappingManager::OnComponentAdded(UPrimitiveComponent* Component)
+{
+	if (Component->GetOwner() == nullptr)
+	{
+		return;
+	}
+	if (IsComponentTypeSupported(Component) == false)
+	{
+		return;
+	}
+
+	UE::Geometry::FSceneGeometryID GeometryID;
+	SpatialCache->EnableComponentTracking(Component, GeometryID);
+
+	ComponentToActorMap.Add(Component, Component->GetOwner());
+
+	// UDynamicMeshComponent gets special treatment because it has an explicit mesh-modified handler
+	// we can rely on to detect geometry changes
+	if (UDynamicMeshComponent* DynamicMeshComponent = Cast<UDynamicMeshComponent>(Component))
+	{
+		DynamicMeshComponent->OnMeshChanged.AddUObject(this, &UModelingSceneSnappingManager::HandleDynamicMeshModifiedDelegate, DynamicMeshComponent);
+
+		DynamicMeshComponents.Add(DynamicMeshComponent, TWeakObjectPtr<UDynamicMeshComponent>(DynamicMeshComponent) );
+	}
+
+}
+
+void UModelingSceneSnappingManager::OnComponentRemoved(UPrimitiveComponent* Component)
+{
+	SpatialCache->DisableComponentTracking(Component);
+	ComponentToActorMap.Remove(Component);
+
+	// If this is a DynamicMeshComponent, we want to stop listening to mesh change events, but
+	// we need to make sure it is still alive first.
+	TWeakObjectPtr<UDynamicMeshComponent>* FoundDynamicMesh = DynamicMeshComponents.Find(Component);
+	if (FoundDynamicMesh != nullptr)
+	{
+		if (FoundDynamicMesh->IsValid())
+		{
+			UDynamicMeshComponent* DynamicMeshComponent = FoundDynamicMesh->Get();
+			DynamicMeshComponent->OnMeshChanged.RemoveAll(this);
+		}
 
+		DynamicMeshComponents.Remove(Component);
+	}
+}
+
+
+void UModelingSceneSnappingManager::HandleGlobalObjectModifiedDelegate(UObject* Object)
+{
+	if (Cast<UDynamicMeshComponent>(Object))
+	{
+		return;		// have special-case handling
+	}
+
+	UPrimitiveComponent* PrimComponent = Cast<UPrimitiveComponent>(Object);
+	if (PrimComponent)
+	{
+		OnComponentModified(PrimComponent);
+	}
+}
+
+void UModelingSceneSnappingManager::OnComponentModified(UActorComponent* Component)
+{
+	const bool bDeferRebuild = false;
+	UPrimitiveComponent* PrimComponent = Cast<UPrimitiveComponent>(Component);
+	if (PrimComponent)
+	{
+		SpatialCache->NotifyGeometryUpdate(PrimComponent, bDeferRebuild);
+	}
+}
+
+
+void UModelingSceneSnappingManager::HandleGlobalComponentTransformChangedDelegate(USceneComponent* Component)
+{
+	UPrimitiveComponent* PrimComponent = Cast<UPrimitiveComponent>(Component);
+	if (PrimComponent && ComponentToActorMap.Contains(PrimComponent))
+	{
+		SpatialCache->NotifyTransformUpdate(PrimComponent);
+	}
+}
+
+
+void UModelingSceneSnappingManager::HandleDynamicMeshModifiedDelegate(UDynamicMeshComponent* Component)
+{
+	const bool bDeferRebuild = false;
+	SpatialCache->NotifyGeometryUpdate(Component, bDeferRebuild);
+}
+
+
+void UModelingSceneSnappingManager::BuildSpatialCacheForWorld(
+	UWorld* World,
+	TFunctionRef<bool(AActor*)> ActorFilter,
+	TFunctionRef<bool(UPrimitiveComponent*)> PrimitiveFilter)
+{
+	for ( TActorIterator<AActor> ActorItr(World); ActorItr; ++ActorItr)
+	{
+		if (ActorFilter(*ActorItr))
+		{
+			OnActorAdded(*ActorItr, PrimitiveFilter);
+		}
+	}
 }
 
 
@@ -122,51 +308,26 @@ bool UModelingSceneSnappingManager::ExecuteSceneSnapQueryRotation(const FSceneSn
 
 
 
-static bool IsHiddenActor(AActor* Actor)
-{
-#if WITH_EDITOR
-	return (Actor != nullptr) && ( Actor->IsHidden() || Actor->IsHiddenEd() );
-#else
-	return (Actor != nullptr) && ( Actor->IsHidden() );
-#endif
-}
-
-static bool IsHiddenComponent(UPrimitiveComponent* Component)
-{
-#if WITH_EDITOR
-	return (Component != nullptr) && ( (Component->IsVisible() == false) || (Component->IsVisibleInEditor() == false) );
-#else
-	return (Component != nullptr) && (Component->IsVisible() == false);
-#endif
-}
 
 
-static bool IsVisibleObjectHit_Internal(const FHitResult& HitResult)
-{
-	AActor* Actor = HitResult.GetActor();
-	if (IsHiddenActor(Actor))
-	{
-		return false;
-	}
-	UPrimitiveComponent* Component = HitResult.GetComponent();
-	if ( IsHiddenComponent(Component) )
-	{
-		return false;
-	}
-	return true;
-}
-
-static bool FindNearestVisibleObjectHit_Internal(UWorld* World, FHitResult& HitResultOut, const FVector& Start, const FVector& End,
-	bool bIsSceneGeometrySnapQuery, const TArray<const UPrimitiveComponent*>* ComponentsToIgnore, 
-	const TArray<const UPrimitiveComponent*>* InvisibleComponentsToInclude)
+static bool FindNearestVisibleObjectHit_Internal(
+	UWorld* World, 
+	FHitResult& HitResultOut, 
+	const FRay3d& Ray,
+	bool bIsSceneGeometrySnapQuery, 
+	const FSceneQueryVisibilityFilter& VisibilityFilter,
+	UE::Geometry::FSceneGeometrySpatialCache* SpatialCache)
 {
 	FCollisionObjectQueryParams ObjectQueryParams(FCollisionObjectQueryParams::AllObjects);
 	FCollisionQueryParams QueryParams = FCollisionQueryParams::DefaultQueryParam;
 	QueryParams.bTraceComplex = true;
 	QueryParams.bReturnFaceIndex = bIsSceneGeometrySnapQuery;
 
+	// TODO: QueryParams.AddIgnoredComponents ?
+
 	TArray<FHitResult> OutHits;
-	if (World->LineTraceMultiByObjectType(OutHits, Start, End, ObjectQueryParams, QueryParams) == false)
+	FVector RayEnd = (FVector)(Ray.PointAt(HALF_WORLD_MAX));
+	if (World->LineTraceMultiByObjectType(OutHits, (FVector)Ray.Origin, RayEnd, ObjectQueryParams, QueryParams) == false)
 	{
 		return false;
 	}
@@ -174,17 +335,16 @@ static bool FindNearestVisibleObjectHit_Internal(UWorld* World, FHitResult& HitR
 	float NearestVisible = TNumericLimits<float>::Max();
 	for (const FHitResult& CurResult : OutHits)
 	{
-		if (CurResult.Distance < NearestVisible)
+		// if we have hit Component in the SpatialCache, prefer to use that
+		if (SpatialCache && SpatialCache->HaveCacheForComponent(CurResult.GetComponent()) )
 		{
-			if (IsVisibleObjectHit_Internal(CurResult) 
-				|| (InvisibleComponentsToInclude && InvisibleComponentsToInclude->Contains(CurResult.GetComponent())))
-			{
-				if (!ComponentsToIgnore || !ComponentsToIgnore->Contains(CurResult.GetComponent()))
-				{
-					HitResultOut = CurResult;
-					NearestVisible = CurResult.Distance;
-				}
-			}
+			continue;
+		}
+
+		if (CurResult.Distance < NearestVisible && VisibilityFilter.IsVisible(CurResult.Component.Get()) )
+		{
+			HitResultOut = CurResult;
+			NearestVisible = CurResult.Distance;
 		}
 	}
 
@@ -192,6 +352,106 @@ static bool FindNearestVisibleObjectHit_Internal(UWorld* World, FHitResult& HitR
 }
 
 
+template<typename VectorType>
+bool GetComponentHitTriangle_Internal(FHitResult HitResult, VectorType* TriVertices)
+{
+	if (HitResult.Component.IsValid() == false || HitResult.FaceIndex < 0)
+	{
+		return false;
+	}
+
+	if (UPrimitiveComponent* Component = HitResult.Component.Get())
+	{
+		// physics collision data is created from StaticMesh RenderData
+		// so use HitResult.FaceIndex to extract triangle from the LOD0 mesh
+		// (note: this may be incorrect if there are multiple sections...in that case I think we have to
+		//  first find section whose accumulated index range would contain .FaceIndexX)
+		UStaticMesh* StaticMesh = Cast<UStaticMeshComponent>(Component)->GetStaticMesh();
+		FStaticMeshLODResources& LOD = StaticMesh->GetRenderData()->LODResources[0];
+		FIndexArrayView Indices = LOD.IndexBuffer.GetArrayView();
+		int32 TriIdx = 3 * HitResult.FaceIndex;
+		FVector Positions[3];
+		Positions[0] = LOD.VertexBuffers.PositionVertexBuffer.VertexPosition(Indices[TriIdx]);
+		Positions[1] = LOD.VertexBuffers.PositionVertexBuffer.VertexPosition(Indices[TriIdx + 1]);
+		Positions[2] = LOD.VertexBuffers.PositionVertexBuffer.VertexPosition(Indices[TriIdx + 2]);
+
+		// transform to world space
+		FTransform ComponentTransform = Component->GetComponentTransform();
+		Positions[0] = ComponentTransform.TransformPosition(Positions[0]);
+		Positions[1] = ComponentTransform.TransformPosition(Positions[1]);
+		Positions[2] = ComponentTransform.TransformPosition(Positions[2]);
+
+		TriVertices[0] = (VectorType)Positions[0];
+		TriVertices[1] = (VectorType)Positions[1];
+		TriVertices[2] = (VectorType)Positions[2];
+		return true;
+	}
+
+	return false;
+}
+
+
+
+
+bool UModelingSceneSnappingManager::ExecuteSceneHitQuery(const FSceneHitQueryRequest& Request, FSceneHitQueryResult& ResultOut) const
+{
+	if (!QueriesAPI)
+	{
+		return false;
+	}
+
+	FRay3d WorldRay(Request.WorldRay);
+
+	double MinCurrentHitDist = TNumericLimits<double>::Max();
+
+	FViewCameraState ViewState;
+	QueriesAPI->GetCurrentViewState(ViewState);
+
+	FSceneGeometryPoint HitPoint;
+	FSceneGeometryID HitIdentifier;
+	bool bHitSceneCache = SpatialCache->FindNearestHit(WorldRay, HitPoint, HitIdentifier, &Request.VisibilityFilter );
+	if (bHitSceneCache)
+	{
+		FVector3d A, B, C;
+		SpatialCache->GetGeometry(HitIdentifier, HitPoint.GeometryType, HitPoint.GeometryIndex, true, A, B, C);
+		MinCurrentHitDist = HitPoint.RayDistance;
+		ResultOut.TargetActor = HitPoint.Actor;
+		ResultOut.TargetComponent = HitPoint.Component;
+		ResultOut.HitTriIndex = HitPoint.GeometryIndex;
+		ResultOut.Position = HitPoint.WorldPoint;
+		ResultOut.Normal = VectorUtil::Normal(A, B, C);
+		ResultOut.TriVertices[0] = A;
+		ResultOut.TriVertices[1] = B;
+		ResultOut.TriVertices[2] = C;
+
+		ResultOut.InitializeHitResult(Request);
+	}
+
+
+	FHitResult HitResult;
+	bool bHitWorld = FindNearestVisibleObjectHit_Internal(
+		QueriesAPI->GetCurrentEditingWorld(), 
+		HitResult, WorldRay, 
+		true, Request.VisibilityFilter, SpatialCache.Get());
+
+	if (bHitWorld && HitResult.Distance < MinCurrentHitDist )
+	{
+		MinCurrentHitDist =  HitResult.Distance;
+		ResultOut.TargetActor = HitResult.GetActor();
+		ResultOut.TargetComponent = HitResult.GetComponent();
+		ResultOut.HitTriIndex = HitResult.FaceIndex;
+		ResultOut.Position = (FVector3d)HitResult.Location;
+		ResultOut.Normal = (FVector3d)HitResult.Normal;
+		ResultOut.HitResult = HitResult;
+		if (Request.bWantHitGeometryInfo)
+		{ 
+			GetComponentHitTriangle_Internal<FVector3d>(HitResult, ResultOut.TriVertices);
+		}
+	}
+
+	return (MinCurrentHitDist < TNumericLimits<double>::Max());
+}
+
 
 
 
@@ -232,99 +492,138 @@ bool UModelingSceneSnappingManager::ExecuteSceneSnapQueryPosition(const FSceneSn
 	// If a hit is found, we look up what triangle was hit, and then test its vertices and edges
 	//
 
+
+	// try snapping to vertices
+	auto TrySnapToVertex = [](const FSceneSnapQueryRequest& RequestIn, const FRay3d& WorldRay, FSceneSnapQueryResult& SnapResult, double& SmallestSnapAngle)
+	{
+		if ((RequestIn.TargetTypes & ESceneSnapQueryTargetType::MeshVertex) != ESceneSnapQueryTargetType::None)
+		{
+			for (int j = 0; j < 3; ++j)
+			{
+				double VisualAngle = OpeningAngleDeg(RequestIn.Position, SnapResult.TriVertices[j], (FVector)WorldRay.Origin);
+				if (VisualAngle < SmallestSnapAngle)
+				{
+					SmallestSnapAngle = VisualAngle;
+					SnapResult.Position = SnapResult.TriVertices[j];
+					SnapResult.TargetType = ESceneSnapQueryTargetType::MeshVertex;
+					SnapResult.TriSnapIndex = j;
+				}
+			}
+		}
+	};
+
+
+	auto TrySnapToEdge = [](const FSceneSnapQueryRequest& RequestIn, const FRay3d& WorldRay, FSceneSnapQueryResult& SnapResult, double& SmallestSnapAngle)
+	{
+		if ( ((RequestIn.TargetTypes & ESceneSnapQueryTargetType::MeshEdge) != ESceneSnapQueryTargetType::None) &&
+			  (SnapResult.TargetType != ESceneSnapQueryTargetType::MeshVertex) )
+		{
+			for (int j = 0; j < 3; ++j)
+			{
+				FVector EdgeNearestPt = NearestSegmentPt(SnapResult.TriVertices[j], SnapResult.TriVertices[(j+1)%3], RequestIn.Position);
+				double VisualAngle = OpeningAngleDeg(RequestIn.Position, EdgeNearestPt, (FVector)WorldRay.Origin);
+				if (VisualAngle < SmallestSnapAngle )
+				{
+					SmallestSnapAngle = VisualAngle;
+					SnapResult.Position = EdgeNearestPt;
+					SnapResult.TargetType = ESceneSnapQueryTargetType::MeshEdge;
+					SnapResult.TriSnapIndex = j;
+				}
+			}
+		}
+	};
+
+
+	FVector3d RayStart = (FVector3d)ViewState.Position;
+	FVector3d RayDirection = (FVector3d)Request.Position - RayStart;
+	FRay3d WorldRay(RayStart, Normalized(RayDirection));
+
+	FSceneQueryVisibilityFilter TmpFilter;
+	TmpFilter.ComponentsToIgnore = Request.ComponentsToIgnore;
+	TmpFilter.InvisibleComponentsToInclude = Request.InvisibleComponentsToInclude;
+
+	// cast ray into scene cache
+	FSceneGeometryPoint HitPoint;
+	FSceneGeometryID HitIdentifier;
+	bool bHitSceneCache = SpatialCache->FindNearestHit(WorldRay, HitPoint, HitIdentifier);
+
 	// cast ray into world
-	FVector RayStart = ViewState.Position;
-	FVector RayDirection = Request.Position - RayStart; RayDirection.Normalize();
-	FVector RayEnd = RayStart + HALF_WORLD_MAX * RayDirection;
 	FHitResult HitResult;
 	bool bHitWorld = FindNearestVisibleObjectHit_Internal(
 		QueriesAPI->GetCurrentEditingWorld(), 
-		HitResult, 
-		RayStart, RayEnd, 
-		true, 
-		Request.ComponentsToIgnore, 
-		Request.InvisibleComponentsToInclude);
+		HitResult, WorldRay, 
+		true, TmpFilter, SpatialCache.Get());
+
+	// if we hit both, determine which is closer
+	if (bHitSceneCache && bHitWorld)
+	{
+		if (HitPoint.RayDistance < HitResult.Distance)
+		{
+			bHitWorld = false;
+		}
+		else
+		{
+			bHitSceneCache = false;
+		}
+	}
+
+	// Try to snap based on whichever thing was hit.
+	// Might be possible to combine parts of these blocks.
+
+	if (bHitSceneCache)
+	{
+		// is this check necessary? isn't this always going to be zero for a scene hit??
+		float VisualAngle = OpeningAngleDeg(Request.Position, (FVector)HitPoint.WorldPoint, (FVector)RayStart);
+		if (VisualAngle < Request.VisualAngleThresholdDegrees)
+		{
+			FVector3d A, B, C;
+			SpatialCache->GetGeometry(HitIdentifier, HitPoint.GeometryType, HitPoint.GeometryIndex, true, A, B, C);
+
+			FSceneSnapQueryResult SnapResult;
+			SnapResult.TriVertices[0] = (FVector)A;
+			SnapResult.TriVertices[1] = (FVector)B;
+			SnapResult.TriVertices[2] = (FVector)C;
+
+			double SmallestAngle = Request.VisualAngleThresholdDegrees;
+
+			// try snapping to vertices
+			TrySnapToVertex(Request, WorldRay, SnapResult, SmallestAngle);
+
+			// try snapping to nearest points on edges
+			TrySnapToEdge(Request, WorldRay, SnapResult, SmallestAngle);
 
-	if (bHitWorld && HitResult.FaceIndex >= 0)
+			// if we found a valid snap, return it
+			if (SmallestAngle < (double)Request.VisualAngleThresholdDegrees)
+			{
+				SnapResult.TargetActor = HitPoint.Actor;
+				SnapResult.TargetComponent = HitPoint.Component;
+				Results.Add(SnapResult);
+				FoundResultCount++;
+			}
+
+		}
+	}
+
+
+	if (bHitWorld)
 	{
+		// is this check necessary? isn't this always going to be zero for a scene hit??
 		float VisualAngle = OpeningAngleDeg(Request.Position, HitResult.ImpactPoint, RayStart);
-		//UE_LOG(LogTemp, Warning, TEXT("[HIT] visualangle %f faceindex %d"), VisualAngle, HitResult.FaceIndex);
 		if (VisualAngle < Request.VisualAngleThresholdDegrees)
 		{
-			UPrimitiveComponent* Component = HitResult.Component.Get();
-			if (Cast<UStaticMeshComponent>(Component) != nullptr)
+			FSceneSnapQueryResult SnapResult;
+			if ( GetComponentHitTriangle_Internal<FVector>(HitResult, SnapResult.TriVertices) )
 			{
-				// HitResult.FaceIndex is apparently an index into the TriMeshCollisionData, not sure how
-				// to directly access it. Calling GetPhysicsTriMeshData is expensive!
-				//UBodySetup* BodySetup = Cast<UStaticMeshComponent>(Component)->GetBodySetup();
-				//UObject* CDPObj = BodySetup->GetOuter();
-				//IInterface_CollisionDataProvider* CDP = Cast<IInterface_CollisionDataProvider>(CDPObj);
-				//FTriMeshCollisionData TriMesh;
-				//CDP->GetPhysicsTriMeshData(&TriMesh, true);
-				//FTriIndices Triangle = TriMesh.Indices[HitResult.FaceIndex];
-				//FVector Positions[3] = { TriMesh.Vertices[Triangle.v0], TriMesh.Vertices[Triangle.v1], TriMesh.Vertices[Triangle.v2] };
-
-				// physics collision data is created from StaticMesh RenderData
-				// so use HitResult.FaceIndex to extract triangle from the LOD0 mesh
-				// (note: this may be incorrect if there are multiple sections...in that case I think we have to
-				//  first find section whose accumulated index range would contain .FaceIndexX)
-				UStaticMesh* StaticMesh = Cast<UStaticMeshComponent>(Component)->GetStaticMesh();
-				FStaticMeshLODResources& LOD = StaticMesh->GetRenderData()->LODResources[0];
-				FIndexArrayView Indices = LOD.IndexBuffer.GetArrayView();
-				int32 TriIdx = 3 * HitResult.FaceIndex;
-				FVector Positions[3];
-				Positions[0] = LOD.VertexBuffers.PositionVertexBuffer.VertexPosition(Indices[TriIdx]);
-				Positions[1] = LOD.VertexBuffers.PositionVertexBuffer.VertexPosition(Indices[TriIdx+1]);
-				Positions[2] = LOD.VertexBuffers.PositionVertexBuffer.VertexPosition(Indices[TriIdx+2]);
-
-				// transform to world space
-				FTransform ComponentTransform = Component->GetComponentTransform();
-				Positions[0] = ComponentTransform.TransformPosition(Positions[0]);
-				Positions[1] = ComponentTransform.TransformPosition(Positions[1]);
-				Positions[2] = ComponentTransform.TransformPosition(Positions[2]);
-
-				FSceneSnapQueryResult SnapResult;
-				SnapResult.TriVertices[0] = Positions[0];
-				SnapResult.TriVertices[1] = Positions[1];
-				SnapResult.TriVertices[2] = Positions[2];
+				double SmallestAngle = Request.VisualAngleThresholdDegrees;
 
 				// try snapping to vertices
-				float SmallestAngle = Request.VisualAngleThresholdDegrees;
-				if ( (Request.TargetTypes & ESceneSnapQueryTargetType::MeshVertex) != ESceneSnapQueryTargetType::None)
-				{
-					for (int j = 0; j < 3; ++j)
-					{
-						VisualAngle = OpeningAngleDeg(Request.Position, Positions[j], RayStart);
-						if (VisualAngle < SmallestAngle)
-						{
-							SmallestAngle = VisualAngle;
-							SnapResult.Position = Positions[j];
-							SnapResult.TargetType = ESceneSnapQueryTargetType::MeshVertex;
-							SnapResult.TriSnapIndex = j;
-						}
-					}
-				}
+				TrySnapToVertex(Request, WorldRay, SnapResult, SmallestAngle);
 
 				// try snapping to nearest points on edges
-				if ( ((Request.TargetTypes & ESceneSnapQueryTargetType::MeshEdge) != ESceneSnapQueryTargetType::None) &&
-					(SnapResult.TargetType != ESceneSnapQueryTargetType::MeshVertex) )
-				{
-					for (int j = 0; j < 3; ++j)
-					{
-						FVector EdgeNearestPt = NearestSegmentPt(Positions[j], Positions[(j+1)%3], Request.Position);
-						VisualAngle = OpeningAngleDeg(Request.Position, EdgeNearestPt, RayStart);
-						if (VisualAngle < SmallestAngle )
-						{
-							SmallestAngle = VisualAngle;
-							SnapResult.Position = EdgeNearestPt;
-							SnapResult.TargetType = ESceneSnapQueryTargetType::MeshEdge;
-							SnapResult.TriSnapIndex = j;
-						}
-					}
-				}
+				TrySnapToEdge(Request, WorldRay, SnapResult, SmallestAngle);
 
 				// if we found a valid snap, return it
-				if (SmallestAngle < Request.VisualAngleThresholdDegrees)
+				if (SmallestAngle < (double)Request.VisualAngleThresholdDegrees)
 				{
 					SnapResult.TargetActor = HitResult.HitObjectHandle.FetchActor();
 					SnapResult.TargetComponent = HitResult.Component.Get();
@@ -333,7 +632,6 @@ bool UModelingSceneSnappingManager::ExecuteSceneSnapQueryPosition(const FSceneSn
 				}
 			}
 		}
-
 	}
 
 	return (FoundResultCount > 0);
diff --git a/Engine/Plugins/Runtime/MeshModelingToolset/Source/ModelingComponents/Private/ToolSceneQueriesUtil.cpp b/Engine/Plugins/Runtime/MeshModelingToolset/Source/ModelingComponents/Private/ToolSceneQueriesUtil.cpp
--- a/Engine/Plugins/Runtime/MeshModelingToolset/Source/ModelingComponents/Private/ToolSceneQueriesUtil.cpp
+++ b/Engine/Plugins/Runtime/MeshModelingToolset/Source/ModelingComponents/Private/ToolSceneQueriesUtil.cpp
@@ -334,4 +334,44 @@ bool ToolSceneQueriesUtil::FindNearestVisibleObjectHit(UWorld* World, FHitResult
 	const TArray<const UPrimitiveComponent*>* IgnoreComponents, const TArray<const UPrimitiveComponent*>* InvisibleComponentsToInclude)
 {
 	return FindNearestVisibleObjectHit(World, HitResultOut, Ray.Origin, Ray.PointAt(HALF_WORLD_MAX), IgnoreComponents, InvisibleComponentsToInclude);
+}
+
+
+
+
+bool ToolSceneQueriesUtil::FindNearestVisibleObjectHit(USceneSnappingManager* SnapManager, FHitResult& HitResultOut, const FRay& Ray,
+	const TArray<const UPrimitiveComponent*>* IgnoreComponents, const TArray<const UPrimitiveComponent*>* InvisibleComponentsToInclude)
+{
+	if (!SnapManager)
+	{
+		return false;
+	}
+
+	FSceneHitQueryRequest Request;
+	Request.WorldRay = (FRay3d)Ray;
+	Request.bWantHitGeometryInfo = false;
+
+	FSceneHitQueryResult Result;
+	if ( SnapManager->ExecuteSceneHitQuery(Request, Result) )
+	{
+		HitResultOut = Result.HitResult;
+		return true;
+	};
+	return false;
+}
+
+
+
+bool ToolSceneQueriesUtil::FindNearestVisibleObjectHit(const UInteractiveTool* Tool, FHitResult& HitResultOut, const FVector& Start, const FVector& End,
+	const TArray<const UPrimitiveComponent*>* IgnoreComponents, const TArray<const UPrimitiveComponent*>* InvisibleComponentsToInclude)
+{
+	FRay WorldRay(Start, (End - Start), false);
+	return FindNearestVisibleObjectHit(USceneSnappingManager::Find(Tool->GetToolManager()), HitResultOut, WorldRay, IgnoreComponents, InvisibleComponentsToInclude);
+}
+
+
+bool ToolSceneQueriesUtil::FindNearestVisibleObjectHit(const UInteractiveTool* Tool, FHitResult& HitResultOut, const FRay& Ray,
+	const TArray<const UPrimitiveComponent*>* IgnoreComponents, const TArray<const UPrimitiveComponent*>* InvisibleComponentsToInclude)
+{
+	return FindNearestVisibleObjectHit(USceneSnappingManager::Find(Tool->GetToolManager()), HitResultOut, Ray, IgnoreComponents, InvisibleComponentsToInclude);
 }
\ No newline at end of file
diff --git a/Engine/Plugins/Runtime/MeshModelingToolset/Source/ModelingComponents/Public/Scene/SceneGeometrySpatialCache.h b/Engine/Plugins/Runtime/MeshModelingToolset/Source/ModelingComponents/Public/Scene/SceneGeometrySpatialCache.h
new file mode 100644
--- /dev/null
+++ b/Engine/Plugins/Runtime/MeshModelingToolset/Source/ModelingComponents/Public/Scene/SceneGeometrySpatialCache.h
@@ -0,0 +1,151 @@
+// Copyright Epic Games, Inc. All Rights Reserved.
+
+#pragma once
+
+#include "BoxTypes.h"
+#include "Spatial/SpatialInterfaces.h"
+#include "DynamicMesh/DynamicMesh3.h"
+#include "DynamicMesh/DynamicMeshAABBTree3.h"
+#include "DynamicMesh/ColliderMesh.h"
+#include "Spatial/SparseDynamicOctree3.h"
+#include "SceneQueries/SceneSnappingManager.h"		// only needed for FSceneQueryVisibilityFilter, maybe can wrap this?
+
+// need these for TWeakObjectPtr
+#include "Components/BrushComponent.h"
+#include "Components/DynamicMeshComponent.h"
+
+class AActor;
+class UPrimitiveComponent;
+
+
+namespace UE
+{
+namespace Geometry
+{
+
+enum class EGeometryPointType
+{
+	Triangle = 0,
+	Edge = 1,
+	Vertex = 2
+};
+
+struct MODELINGCOMPONENTS_API FSceneGeometryPoint
+{
+	// world ray that was cast, stored here for convenience
+	FRay3d Ray;
+	// distance along ray that point occurs at
+	double RayDistance = TNumericLimits<float>::Max();		// float so that square-distance doesn't overflow
+
+	// Actor that was hit, if available
+	AActor* Actor = nullptr;
+	// Component that was hit, if available
+	UPrimitiveComponent* Component = nullptr;
+	// Element Index on hit Component, if available (eg Instance Index on an InstancedStaticMesh)
+	int32 ComponentElementIndex = -1;
+	// Index/ID on geometry of hit Component
+	int32 GeometryIndex = -1;
+	// Geometry type
+	EGeometryPointType GeometryType = EGeometryPointType::Triangle;
+	// Barycentric coordinates of the point, if that is applicable
+	FVector3d PointBaryCoords = FVector3d::Zero();
+	// Component-space point
+	FVector3d LocalPoint = FVector3d::Zero();
+	// World-space point
+	FVector3d WorldPoint = FVector3d::Zero();
+
+	// todo: normals...
+};
+
+
+
+class ISceneGeometrySpatial
+{
+public:
+	ISceneGeometrySpatial();
+	virtual ~ISceneGeometrySpatial() {}
+
+	int32 GetID() const { return SpatialID; }
+
+	virtual bool IsValid() = 0;
+	virtual bool IsVisible() = 0;
+
+	virtual UPrimitiveComponent* GetComponent() = 0;
+
+	virtual FAxisAlignedBox3d GetWorldBounds() = 0;
+	virtual bool FindNearestHit(
+		const FRay3d& WorldRay, 
+		FSceneGeometryPoint& HitResultOut,
+		double MaxDistance = TNumericLimits<double>::Max() ) = 0;
+	virtual bool GetGeometry(EGeometryPointType GeometryType, int32 GeometryIndex, bool bWorldSpace, FVector3d& A, FVector3d& B, FVector3d& C) = 0;
+
+	virtual void OnGeometryModified(bool bDeferSpatialRebuild) = 0;
+	virtual void OnTransformModified() = 0;
+	
+private:
+	int32 SpatialID = 0;
+
+	static std::atomic<int32> UniqueIDGenerator;
+};
+
+
+
+
+struct MODELINGCOMPONENTS_API FSceneGeometryID
+{
+	int32 UniqueID;
+	FSceneGeometryID() { UniqueID = 0; }
+	FSceneGeometryID(int32 ID) { UniqueID = ID; }
+	bool operator==(const FSceneGeometryID& Other) const { return UniqueID == Other.UniqueID; }
+};
+FORCEINLINE uint32 GetTypeHash(const FSceneGeometryID& Identifier)
+{
+	return ::GetTypeHash(Identifier.UniqueID);
+}
+
+
+
+
+class MODELINGCOMPONENTS_API FSceneGeometrySpatialCache
+{
+public:
+
+public:
+	FSceneGeometrySpatialCache();
+	virtual ~FSceneGeometrySpatialCache();
+
+	virtual bool EnableComponentTracking(UPrimitiveComponent* Component, FSceneGeometryID& IdentifierOut);
+
+	virtual void DisableComponentTracking(UPrimitiveComponent* Component);
+	virtual void DisableGeometryTracking(FSceneGeometryID& IdentifierOut);
+
+	virtual bool HaveCacheForComponent(UPrimitiveComponent* Component);
+
+
+	virtual bool NotifyTransformUpdate(UPrimitiveComponent* Component);
+	virtual bool NotifyGeometryUpdate(UPrimitiveComponent* Component, bool bDeferRebuild);
+
+	virtual bool FindNearestHit(
+		const FRay3d& WorldRay,
+		FSceneGeometryPoint& HitResultOut,
+		FSceneGeometryID& HitIdentifier,
+		const FSceneQueryVisibilityFilter* VisibilityFilter = nullptr,
+		double MaxDistance = TNumericLimits<double>::Max());
+
+	virtual bool GetGeometry(FSceneGeometryID Identifier, EGeometryPointType GeometryType, int32 GeometryIndex, bool bWorldSpace, FVector3d& A, FVector3d& B, FVector3d& C);
+
+protected:
+	int32 UniqueIDGenerator = 0;
+
+	TMap<UPrimitiveComponent*, FSceneGeometryID> ComponentMap;
+	TMap<FSceneGeometryID, TSharedPtr<ISceneGeometrySpatial>> Spatials;
+	FSparseDynamicOctree3 BoundsOctree;
+};
+
+
+
+}
+}
+
+
+
diff --git a/Engine/Plugins/Runtime/MeshModelingToolset/Source/ModelingComponents/Public/Snapping/ModelingSceneSnappingManager.h b/Engine/Plugins/Runtime/MeshModelingToolset/Source/ModelingComponents/Public/Snapping/ModelingSceneSnappingManager.h
--- a/Engine/Plugins/Runtime/MeshModelingToolset/Source/ModelingComponents/Public/Snapping/ModelingSceneSnappingManager.h
+++ b/Engine/Plugins/Runtime/MeshModelingToolset/Source/ModelingComponents/Public/Snapping/ModelingSceneSnappingManager.h
@@ -2,6 +2,7 @@
 
 #pragma once
 
+#include "GeometryBase.h"
 #include "SceneQueries/SceneSnappingManager.h"
 #include "ModelingSceneSnappingManager.generated.h"
 
@@ -9,7 +10,8 @@
 class IToolsContextQueriesAPI;
 class UInteractiveToolsContext;
 class UInteractiveToolManager;
-
+class UDynamicMeshComponent;
+PREDECLARE_GEOMETRY(class FSceneGeometrySpatialCache);
 
 /**
  * UModelingSceneSnappingManager is an implementation of snapping suitable for use in
@@ -17,12 +19,15 @@ class UInteractiveToolManager;
  * 
  * Currently Supports:
  *    - snap to position/rotation grid
- *    - snap to vertex position for Static Mesh Assets
- *    - snap to edge position for Static Mesh Assets
+ *    - snap to mesh vertex position
+ *    - snap to mesh edge position
+ * 
+ * Snapping to mesh vertex/edge positions currently works for Volume (BrushComponent), StaticMeshComponent, 
+ * and DynamicMeshComponent. 
  * 
- * Note that currently the StaticMesh vertex/edge snapping is dependent on the Physics
+ * Currently the StaticMesh vertex/edge snapping is dependent on the Physics
  * system, and may fail or return nonsense results in some cases, due to the physics
- * complex-collision mesh deviating from the source-model mesh
+ * complex-collision mesh deviating from the source-model mesh.
  */
 UCLASS()
 class MODELINGCOMPONENTS_API UModelingSceneSnappingManager : public USceneSnappingManager
@@ -34,6 +39,19 @@ public:
 	virtual void Initialize(TObjectPtr<UInteractiveToolsContext> ToolsContext);
 	virtual void Shutdown();
 
+	//
+	// USceneSnappingManager API
+	//
+
+
+	/**
+	* Try to find a Hit Point in the scene that satisfies the HitQuery Request.
+	* @param Request hit query configuration
+	* @param ResultOut hit query result, if return is true
+	* @return true if a valid hit was found
+	* @warning implementations are not required (and may not be able) to support hit-testing
+	*/
+	virtual bool ExecuteSceneHitQuery(const FSceneHitQueryRequest& Request, FSceneHitQueryResult& ResultOut) const override;
 
 	/**
 	* Try to find Snap Targets in the scene that satisfy the Snap Query.
@@ -45,6 +63,35 @@ public:
 	virtual bool ExecuteSceneSnapQuery(const FSceneSnapQueryRequest& Request, TArray<FSceneSnapQueryResult>& Results) const override;
 
 
+	//
+	// API for managing the set of Actors/Components that the snapping system knows about.
+	// Currently this is for Volumes and DynamicMeshActors/Components, the SnappingManager builds
+	// it's own spatial data structure cache for these types of meshes. StaticMeshComponents are
+	// automatically included and handled via world-traces and the physics system.
+	//
+public:
+
+	/** @return true if this Component type is supported by the spatial-cache tracking for hit and snap testing */
+	virtual bool IsComponentTypeSupported(const UPrimitiveComponent* Component) const;
+
+	/** Enable spatial-cache tracking for the Components of the Actor that pass the PrimitiveFilter */
+	virtual void OnActorAdded(AActor* Actor, TFunctionRef<bool(UPrimitiveComponent*)> PrimitiveFilter);
+	/** Disable spatial-cache tracking for any Components of the Actor */
+	virtual void OnActorRemoved(AActor* Actor);
+
+	/** Enable spatial-cache tracking for the Component */
+	virtual void OnComponentAdded(UPrimitiveComponent* Component);
+	/** Disable spatial-cache tracking for the Component */
+	virtual void OnComponentRemoved(UPrimitiveComponent* Component);
+	/** Notify the internal spatial-cache tracking system that the Component has been modified (ie cache needs to be rebiult) */
+	virtual void OnComponentModified(UActorComponent* Component);
+
+	/** Explicitly populate the spatial-cache tracking set with specific Actors/Components (alternative to OnActorAdded route)  */
+	virtual void BuildSpatialCacheForWorld(
+		UWorld* World,
+		TFunctionRef<bool(AActor*)> ActorFilter,
+		TFunctionRef<bool(UPrimitiveComponent*)> PrimitiveFilter );
+
 protected:
 
 	UPROPERTY()
@@ -54,6 +101,22 @@ protected:
 
 	virtual bool ExecuteSceneSnapQueryRotation(const FSceneSnapQueryRequest& Request, TArray<FSceneSnapQueryResult>& Results) const;
 	virtual bool ExecuteSceneSnapQueryPosition(const FSceneSnapQueryRequest& Request, TArray<FSceneSnapQueryResult>& Results) const;
+
+	// This map allows us to identify the Components belonging to Actors. Need to store this
+	// because if the Actor is deleted we will not be able to identify it's (previous) Components
+	TMap<UPrimitiveComponent*, AActor*> ComponentToActorMap;
+
+	// cache for objects in the level
+	TSharedPtr<UE::Geometry::FSceneGeometrySpatialCache> SpatialCache;
+
+	FDelegateHandle OnObjectModifiedHandler;
+	void HandleGlobalObjectModifiedDelegate(UObject* Object);
+
+	FDelegateHandle OnComponentTransformChangedHandle;
+	void HandleGlobalComponentTransformChangedDelegate(USceneComponent* Component);
+
+	TMap<UPrimitiveComponent*, TWeakObjectPtr<UDynamicMeshComponent>> DynamicMeshComponents;
+	void HandleDynamicMeshModifiedDelegate(UDynamicMeshComponent* Component);
 };
 
 
diff --git a/Engine/Plugins/Runtime/MeshModelingToolset/Source/ModelingComponents/Public/ToolSceneQueriesUtil.h b/Engine/Plugins/Runtime/MeshModelingToolset/Source/ModelingComponents/Public/ToolSceneQueriesUtil.h
--- a/Engine/Plugins/Runtime/MeshModelingToolset/Source/ModelingComponents/Public/ToolSceneQueriesUtil.h
+++ b/Engine/Plugins/Runtime/MeshModelingToolset/Source/ModelingComponents/Public/ToolSceneQueriesUtil.h
@@ -8,6 +8,8 @@
 #include "InteractiveToolManager.h"
 #include "VectorTypes.h"
 
+class USceneSnappingManager;
+
 /**
  * Utility functions for Tool implementations to use to do scene queries, generally via IToolsContextQueriesAPI
  */
@@ -148,6 +150,7 @@ namespace ToolSceneQueriesUtil
 	 * @param Start start point of line
 	 * @param End end point of line
 	 * @param IgnoreComponents optional list of Components to ignore
+	 * @param InvisibleComponentsToInclude optional list of Components to explicitly include, even if they are not visible
 	 * @return true if a visible hit was found
 	 */
 	MODELINGCOMPONENTS_API bool FindNearestVisibleObjectHit(UWorld* World, FHitResult& HitResultOut, const FVector& Start, const FVector& End,
@@ -161,10 +164,55 @@ namespace ToolSceneQueriesUtil
 	 * @param HitResultOut the resulting hit, if true is returned
 	 * @param Ray hit ray
 	 * @param IgnoreComponents optional list of Components to ignore
+	 * @param InvisibleComponentsToInclude optional list of Components to explicitly include, even if they are not visible
 	 * @return true if a visible hit was found
 	 */
 	MODELINGCOMPONENTS_API bool FindNearestVisibleObjectHit(UWorld* World, FHitResult& HitResultOut, const FRay& Ray,
 		const TArray<const UPrimitiveComponent*>* IgnoreComponents = nullptr, 
 		const TArray<const UPrimitiveComponent*>* InvisibleComponentsToInclude = nullptr);
 
+
+
+	/**
+	* Find the nearest object hit by the LineTrace from Start to End that is currently visible (provides correct result in Editor)
+	* @param HitResultOut the resulting hit, if true is returned
+	* @param Start start point of line
+	* @param End end point of line
+	* @param IgnoreComponents optional list of Components to ignore
+	* @param InvisibleComponentsToInclude optional list of Components to explicitly include, even if they are not visible
+	* @return true if a visible hit was found
+	*/
+	MODELINGCOMPONENTS_API bool FindNearestVisibleObjectHit(USceneSnappingManager* SnappingManager, 
+		FHitResult& HitResultOut, const FRay& Ray,
+		const TArray<const UPrimitiveComponent*>* IgnoreComponents = nullptr, 
+		const TArray<const UPrimitiveComponent*>* InvisibleComponentsToInclude = nullptr);
+
+
+
+	/**
+	* Find the nearest object hit by the LineTrace from Start to End that is currently visible (provides correct result in Editor)
+	* @param HitResultOut the resulting hit, if true is returned
+	* @param Start start point of line
+	* @param End end point of line
+	* @param IgnoreComponents optional list of Components to ignore
+	* @param InvisibleComponentsToInclude optional list of Components to explicitly include, even if they are not visible
+	* @return true if a visible hit was found
+	*/
+	MODELINGCOMPONENTS_API bool FindNearestVisibleObjectHit(const UInteractiveTool* Tool, FHitResult& HitResultOut, const FVector& Start, const FVector& End,
+		const TArray<const UPrimitiveComponent*>* IgnoreComponents = nullptr, 
+		const TArray<const UPrimitiveComponent*>* InvisibleComponentsToInclude = nullptr);
+
+
+	/**
+	* Find the nearest object hit by the LineTrace from Start to End that is currently visible (provides correct result in Editor)
+	* @param HitResultOut the resulting hit, if true is returned
+	* @param Ray hit ray
+	* @param IgnoreComponents optional list of Components to ignore
+	* @param InvisibleComponentsToInclude optional list of Components to explicitly include, even if they are not visible
+	* @return true if a visible hit was found
+	*/
+	MODELINGCOMPONENTS_API bool FindNearestVisibleObjectHit(const UInteractiveTool* Tool, FHitResult& HitResultOut, const FRay& Ray,
+		const TArray<const UPrimitiveComponent*>* IgnoreComponents = nullptr, 
+		const TArray<const UPrimitiveComponent*>* InvisibleComponentsToInclude = nullptr);
+
 }
\ No newline at end of file
diff --git a/Engine/Plugins/Runtime/MeshModelingToolset/Source/ModelingComponentsEditorOnly/Private/Scene/LevelObjectsObserver.cpp b/Engine/Plugins/Runtime/MeshModelingToolset/Source/ModelingComponentsEditorOnly/Private/Scene/LevelObjectsObserver.cpp
new file mode 100644
--- /dev/null
+++ b/Engine/Plugins/Runtime/MeshModelingToolset/Source/ModelingComponentsEditorOnly/Private/Scene/LevelObjectsObserver.cpp
@@ -0,0 +1,103 @@
+// Copyright Epic Games, Inc. All Rights Reserved.
+
+#include "Scene/LevelObjectsObserver.h"
+
+
+#include "Engine/Engine.h"
+#include "EngineUtils.h" // for TActorIterator<>
+
+#include "Editor.h"
+#include "EditorSupportDelegates.h"
+
+
+void FLevelObjectsObserver::Initialize(UWorld* WorldIn)
+{
+	this->World = WorldIn;
+
+	GEditor->RegisterForUndo(this);
+	OnActorDeletedHandle = GEngine->OnLevelActorDeleted().AddLambda([this](AActor* Actor) { HandleActorDeletedEvent(Actor); });
+	OnActorAddedHandle = GEngine->OnLevelActorAdded().AddLambda([this](AActor* Actor) { HandleActorAddedEvent(Actor); });
+	OnActorListChangedHandle = GEngine->OnLevelActorListChanged().AddLambda([this]() { OnUntrackedLevelChange(); });
+
+	OnUntrackedLevelChange();
+}
+
+
+void FLevelObjectsObserver::Shutdown()
+{
+	for (AActor* Actor : Actors)
+	{
+		OnActorRemoved.Broadcast(Actor);
+	}
+
+	this->World = nullptr;
+	Actors.Reset();
+
+	OnActorAdded.Clear();
+	OnActorRemoved.Clear();
+
+	GEditor->UnregisterForUndo(this);
+	GEngine->OnLevelActorDeleted().Remove(OnActorDeletedHandle);
+	GEngine->OnLevelActorAdded().Remove(OnActorAddedHandle);
+	GEngine->OnLevelActorListChanged().Remove(OnActorListChangedHandle);
+}
+
+
+void FLevelObjectsObserver::PostUndo(bool bSuccess)
+{
+	// on undo/redo anything might have happend so we have to check the entire actor set
+	OnUntrackedLevelChange();
+}
+void FLevelObjectsObserver::PostRedo(bool bSuccess) 
+{ 
+	PostUndo(bSuccess); 
+}
+
+
+
+void FLevelObjectsObserver::OnUntrackedLevelChange()
+{
+	// the world changed and we don't know how, so check all actors
+
+	TSet<AActor*> RemainingActors;
+	RemainingActors.Append(Actors);
+
+	// find any new Actors that we don't know about
+	for ( TActorIterator<AActor> ActorItr(World); ActorItr; ++ActorItr)
+	{
+		AActor* Actor = *ActorItr;
+		if (Actors.Contains(Actor) == false)
+		{
+			Actors.Add(Actor);
+			OnActorAdded.Broadcast(Actor);
+		}
+		RemainingActors.Remove(Actor);		// we saw this actor
+	}
+
+	// remove any Actors we were tracking that no longer exist
+	for (AActor* Actor : RemainingActors)
+	{
+		Actors.Remove(Actor);
+		OnActorRemoved.Broadcast(Actor);
+	}
+}
+
+
+void FLevelObjectsObserver::HandleActorAddedEvent(AActor* Actor)
+{
+	if (Actors.Contains(Actor) == false)
+	{
+		Actors.Add(Actor);
+		OnActorAdded.Broadcast(Actor);
+	}
+}
+
+
+void FLevelObjectsObserver::HandleActorDeletedEvent(AActor* Actor)
+{
+	if (Actors.Contains(Actor))
+	{
+		Actors.Remove(Actor);
+		OnActorRemoved.Broadcast(Actor);
+	}
+}
diff --git a/Engine/Plugins/Runtime/MeshModelingToolset/Source/ModelingComponentsEditorOnly/Public/Scene/LevelObjectsObserver.h b/Engine/Plugins/Runtime/MeshModelingToolset/Source/ModelingComponentsEditorOnly/Public/Scene/LevelObjectsObserver.h
new file mode 100644
--- /dev/null
+++ b/Engine/Plugins/Runtime/MeshModelingToolset/Source/ModelingComponentsEditorOnly/Public/Scene/LevelObjectsObserver.h
@@ -0,0 +1,68 @@
+// Copyright Epic Games, Inc. All Rights Reserved.
+
+#pragma once
+
+#include "EditorUndoClient.h"
+
+class UWorld;
+class AActor;
+
+/**
+ * FLevelObjectsObserver is a utility class that tracks the active-set of Actors
+ * in a World, in the Editor, and emits events when the set changes. The goal is to
+ * guarantee that:
+ *   - OnActorAdded will be emitted when an Actor comes into existence
+ *   - OnActorRemoved will be emitted when an Actor ceases to exist (possibly after GC, though)
+ *   - these signals will be unique, ie once per Actor
+ * 
+ * This is a surprisingly complicated problem in the Editor due to the many ways
+ * that Actors might be created or destroyed, in particular via Undo/Redo. 
+ *
+ * Note that it is not possible to guarantee that OnActorRemoved() fires before an Actor
+ * is GC'd. So, listeners should maintain TWeakObjectPtr<>'s on Actors if they need
+ * to know this information.
+ * 
+ */
+class MODELINGCOMPONENTSEDITORONLY_API FLevelObjectsObserver : public FEditorUndoClient
+{
+public:
+	/**
+	 * Begin watching WorldIn. This will emit OnActorAdded for all Actors in the World.
+	 */
+	void Initialize(UWorld* WorldIn);
+
+	/**
+	 * Stop watching WorldIn. This will emit OnActorRemoved for all Actors that the Observer is currently aware of.
+	 */
+	void Shutdown();
+
+	// FEditorUndoClient implementation
+	virtual void PostUndo(bool bSuccess) override;
+	virtual void PostRedo(bool bSuccess) override;
+
+
+	DECLARE_MULTICAST_DELEGATE_OneParam(FActorAddRemoveSignature, AActor*);
+
+	/**
+	 * OnActorAdded is emitted when a new Actor is detected
+	 */
+	FActorAddRemoveSignature OnActorAdded;
+	/**
+	 * OnActorRemoved is emitted when an existing Actor is destroyed.
+	 * Note that the Actor pointer passed here may point to garbage, ie the Actor
+	 * may have already been GC'd. Do not attempt to call functions on the pointer.
+	 */
+	FActorAddRemoveSignature OnActorRemoved;
+
+protected:
+	UWorld* World = nullptr;		// target World
+	TSet<AActor*> Actors;			// active set of Actors
+
+	FDelegateHandle OnActorDeletedHandle;
+	FDelegateHandle OnActorAddedHandle;
+	FDelegateHandle OnActorListChangedHandle;
+
+	void OnUntrackedLevelChange();
+	void HandleActorAddedEvent(AActor* Actor);
+	void HandleActorDeletedEvent(AActor* Actor);
+};
diff --git a/Engine/Source/Runtime/InteractiveToolsFramework/Private/SceneQueries/SceneSnappingManager.cpp b/Engine/Source/Runtime/InteractiveToolsFramework/Private/SceneQueries/SceneSnappingManager.cpp
--- a/Engine/Source/Runtime/InteractiveToolsFramework/Private/SceneQueries/SceneSnappingManager.cpp
+++ b/Engine/Source/Runtime/InteractiveToolsFramework/Private/SceneQueries/SceneSnappingManager.cpp
@@ -4,6 +4,92 @@
 #include "InteractiveToolManager.h"
 #include "InteractiveGizmoManager.h"
 #include "ContextObjectStore.h"
+#include "GameFramework/Actor.h"
+#include "Components/PrimitiveComponent.h"
+
+
+
+namespace UELocal
+{
+
+static bool IsHiddenActor(const AActor* Actor)
+{
+#if WITH_EDITOR
+	return ( Actor->IsHiddenEd() );
+#else
+	return ( Actor->IsHidden() );
+#endif
+}
+
+static bool IsHiddenComponent(const UPrimitiveComponent* Component)
+{
+#if WITH_EDITOR
+	return (Component->IsVisibleInEditor() == false);
+#else
+	return (Component->IsVisible() == false);
+#endif
+}
+
+
+static bool IsVisibleObject( const UPrimitiveComponent* Component )
+{
+	if (Component == nullptr)
+	{
+		return false;
+	}
+	AActor* Actor = Component->GetOwner();
+	if (Actor == nullptr)
+	{
+		return false;
+	}
+
+	if (IsHiddenActor(Actor))
+	{
+		return false;
+	}
+
+	if ( IsHiddenComponent(Component) )
+	{
+		return false;
+	}
+	return true;
+}
+
+}
+
+
+
+bool FSceneQueryVisibilityFilter::IsVisible(const UPrimitiveComponent* Component) const
+{
+	if (Component == nullptr)
+	{
+		return false;
+	}
+
+	bool bIsVisible = UELocal::IsVisibleObject(Component);
+	if (bIsVisible || (InvisibleComponentsToInclude != nullptr && InvisibleComponentsToInclude->Contains(Component)))
+	{
+		return (ComponentsToIgnore == nullptr) || (ComponentsToIgnore->Contains(Component) == false);
+	}
+	return false;
+}
+
+
+
+
+
+
+void FSceneHitQueryResult::InitializeHitResult(const FSceneHitQueryRequest& FromRequest)
+{
+	HitResult = FHitResult(TargetActor, TargetComponent, (FVector)Position, (FVector)Normal);
+
+	HitResult.Distance = FromRequest.WorldRay.GetParameter(Position);
+	HitResult.FaceIndex = HitTriIndex;
+	// initialize .Time? Need start/end for that...
+}
+
+
+
 
 USceneSnappingManager* USceneSnappingManager::Find(UInteractiveToolManager* ToolManager)
 {
diff --git a/Engine/Source/Runtime/InteractiveToolsFramework/Public/SceneQueries/SceneSnappingManager.h b/Engine/Source/Runtime/InteractiveToolsFramework/Public/SceneQueries/SceneSnappingManager.h
--- a/Engine/Source/Runtime/InteractiveToolsFramework/Public/SceneQueries/SceneSnappingManager.h
+++ b/Engine/Source/Runtime/InteractiveToolsFramework/Public/SceneQueries/SceneSnappingManager.h
@@ -5,8 +5,8 @@
 #include "UObject/Object.h"
 #include "Misc/Optional.h"
 #include "Math/Vector.h"
-#include "Math/Rotator.h"
-#include "Math/Quat.h"
+#include "Math/Ray.h"
+#include "Engine/EngineTypes.h"
 #include "SceneSnappingManager.generated.h"
 
 class AActor;
@@ -16,6 +16,62 @@ class UInteractiveToolManager;
 class UInteractiveGizmoManager;
 
 
+
+struct INTERACTIVETOOLSFRAMEWORK_API FSceneQueryVisibilityFilter
+{
+	/** Optional: components to consider invisible even if they aren't. */
+	const TArray<const UPrimitiveComponent*>* ComponentsToIgnore = nullptr;
+
+	/** Optional: components to consider visible even if they aren't. */
+	const TArray<const UPrimitiveComponent*>* InvisibleComponentsToInclude = nullptr;
+
+	/** @return true if the Component is currently configured as visible (does not consider ComponentsToIgnore or InvisibleComponentsToInclude lists) */
+	bool IsVisible(const UPrimitiveComponent* Component) const;
+};
+
+
+/**
+* Configuration variables for a USceneSnappingManager hit query request.
+*/
+struct INTERACTIVETOOLSFRAMEWORK_API FSceneHitQueryRequest
+{
+	/** scene query ray */
+	FRay3d WorldRay;
+
+	bool bWantHitGeometryInfo = false;
+
+	FSceneQueryVisibilityFilter VisibilityFilter;
+};
+
+
+/**
+* Computed result of a USceneSnappingManager hit query request
+*/
+struct INTERACTIVETOOLSFRAMEWORK_API FSceneHitQueryResult
+{
+	/** Actor that owns hit target */
+	AActor* TargetActor = nullptr;
+	/** Component that owns hit target */
+	UPrimitiveComponent* TargetComponent = nullptr;
+
+	/** hit position*/
+	FVector3d Position = FVector3d::Zero();
+	/** hit normal */
+	FVector3d Normal = FVector3d::UnitZ();
+
+	/** integer ID of triangle that was hit */
+	int HitTriIndex = -1;
+	/** Vertices of triangle that was hit (for debugging, may not be set) */
+	FVector3d TriVertices[3];
+
+	FHitResult HitResult;
+
+	void InitializeHitResult(const FSceneHitQueryRequest& FromRequest);
+};
+
+
+
+
 /** Types of Snap Queries that a USceneSnappingManager may support */
 UENUM()
 enum class ESceneSnapQueryType : uint8
@@ -116,6 +172,18 @@ class INTERACTIVETOOLSFRAMEWORK_API USceneSnappingManager : public UObject
 	GENERATED_BODY()
 public:
 
+	/**
+	* Try to find a Hit Object in the scene that satisfies the Hit Query
+	* @param Request hit query configuration
+	* @param Results hit query result
+	* @return true if any valid hit target was found
+	* @warning implementations are not required (and may not be able) to support hit testing
+	*/
+	virtual bool ExecuteSceneHitQuery(const FSceneHitQueryRequest& Request, FSceneHitQueryResult& ResultOut) const
+	{
+		return false;
+	}
+
 	/**
 	* Try to find Snap Targets/Results in the scene that satisfy the Snap Query.
 	* @param Request snap query configuration
