fix pragma once issue
#rnx

diff --git a/Engine/Plugins/Runtime/MeshModelingToolset/Source/ModelingComponents/Private/Scene/SceneGeometrySpatialCache.cpp b/Engine/Plugins/Runtime/MeshModelingToolset/Source/ModelingComponents/Private/Scene/SceneGeometrySpatialCache.cpp
--- a/Engine/Plugins/Runtime/MeshModelingToolset/Source/ModelingComponents/Private/Scene/SceneGeometrySpatialCache.cpp
+++ b/Engine/Plugins/Runtime/MeshModelingToolset/Source/ModelingComponents/Private/Scene/SceneGeometrySpatialCache.cpp
@@ -1,7 +1,5 @@
 // Copyright Epic Games, Inc. All Rights Reserved.
 
-#pragma once
-
 #include "Scene/SceneGeometrySpatialCache.h"
 
 #include "Engine/Brush.h"
