MeshModelingToolset: Cleanup UV Unwrap tool

diff --git a/Engine/Plugins/Experimental/MeshModelingToolsetExp/Source/MeshModelingToolsExp/Private/Properties/MeshUVChannelProperties.cpp b/Engine/Plugins/Experimental/MeshModelingToolsetExp/Source/MeshModelingToolsExp/Private/Properties/MeshUVChannelProperties.cpp
deleted file mode 100644
diff --git a/Engine/Plugins/Experimental/MeshModelingToolsetExp/Source/MeshModelingToolsExp/Private/RecomputeUVsTool.cpp b/Engine/Plugins/Experimental/MeshModelingToolsetExp/Source/MeshModelingToolsExp/Private/RecomputeUVsTool.cpp
deleted file mode 100644
diff --git a/Engine/Plugins/Experimental/MeshModelingToolsetExp/Source/MeshModelingToolsExp/Public/Properties/MeshUVChannelProperties.h b/Engine/Plugins/Experimental/MeshModelingToolsetExp/Source/MeshModelingToolsExp/Public/Properties/MeshUVChannelProperties.h
deleted file mode 100644
diff --git a/Engine/Plugins/Experimental/MeshModelingToolsetExp/Source/MeshModelingToolsExp/Public/RecomputeUVsTool.h b/Engine/Plugins/Experimental/MeshModelingToolsetExp/Source/MeshModelingToolsExp/Public/RecomputeUVsTool.h
deleted file mode 100644
diff --git a/Engine/Plugins/Runtime/MeshModelingToolset/Source/MeshModelingTools/MeshModelingTools.Build.cs b/Engine/Plugins/Runtime/MeshModelingToolset/Source/MeshModelingTools/MeshModelingTools.Build.cs
--- a/Engine/Plugins/Runtime/MeshModelingToolset/Source/MeshModelingTools/MeshModelingTools.Build.cs
+++ b/Engine/Plugins/Runtime/MeshModelingToolset/Source/MeshModelingTools/MeshModelingTools.Build.cs
@@ -32,8 +32,8 @@ public class MeshModelingTools : ModuleRules
 				"GeometryAlgorithms",
 				"DynamicMesh",
 				"MeshConversion",
-				// "MeshDescription",
-				// "StaticMeshDescription",
+				"MeshDescription",
+				"StaticMeshDescription",
 				// "SkeletalMeshDescription",
 				"ModelingComponents",
 				"ModelingOperators",
@@ -49,7 +49,6 @@ public class MeshModelingTools : ModuleRules
 				"CoreUObject",
 				"Engine",
 				// "RenderCore",
-				// "MeshModelingToolsExp",
 				// "ModelingOperators",
 				// "InputCore",
 				// "PhysicsCore",
diff --git a/Engine/Plugins/Runtime/MeshModelingToolset/Source/MeshModelingTools/Private/Properties/MeshUVChannelProperties.cpp b/Engine/Plugins/Runtime/MeshModelingToolset/Source/MeshModelingTools/Private/Properties/MeshUVChannelProperties.cpp
new file mode 100644
--- /dev/null
+++ b/Engine/Plugins/Runtime/MeshModelingToolset/Source/MeshModelingTools/Private/Properties/MeshUVChannelProperties.cpp
@@ -0,0 +1,69 @@
+// Copyright Epic Games, Inc. All Rights Reserved.
+
+#include "Properties/MeshUVChannelProperties.h"
+
+#include "DynamicMesh/DynamicMesh3.h"
+#include "DynamicMesh/DynamicMeshAttributeSet.h"
+
+#include "MeshDescription.h"
+#include "StaticMeshAttributes.h"
+
+
+void UMeshUVChannelProperties::Initialize(int32 NumUVChannels, bool bInitializeSelection)
+{
+	UVChannelNamesList.Reset();
+	for (int32 k = 0; k < NumUVChannels; ++k)
+	{
+		UVChannelNamesList.Add(FString::Printf(TEXT("UV %d"), k));
+	}
+	if (bInitializeSelection)
+	{
+		UVChannel = (NumUVChannels > 0) ? UVChannelNamesList[0] : TEXT("");
+	}
+}
+
+
+const TArray<FString>& UMeshUVChannelProperties::GetUVChannelNamesFunc() const
+{
+	return UVChannelNamesList;
+}
+
+
+void UMeshUVChannelProperties::Initialize(const FMeshDescription* MeshDescription, bool bInitializeSelection)
+{
+	TVertexInstanceAttributesConstRef<FVector2f> InstanceUVs =
+		MeshDescription->VertexInstanceAttributes().GetAttributesRef<FVector2f>(MeshAttribute::VertexInstance::TextureCoordinate);
+	Initialize(InstanceUVs.GetNumChannels(), bInitializeSelection);
+}
+
+void UMeshUVChannelProperties::Initialize(const FDynamicMesh3* Mesh, bool bInitializeSelection)
+{
+	int32 NumUVChannels = Mesh->HasAttributes() ? Mesh->Attributes()->NumUVLayers() : 0;
+	Initialize(NumUVChannels, bInitializeSelection);
+}
+
+
+
+bool UMeshUVChannelProperties::ValidateSelection(bool bUpdateIfInvalid)
+{
+	int32 FoundIndex = UVChannelNamesList.IndexOfByKey(UVChannel);
+	if (FoundIndex == INDEX_NONE)
+	{
+		if (bUpdateIfInvalid)
+		{
+			UVChannel = (UVChannelNamesList.Num() > 0) ? UVChannelNamesList[0] : TEXT("");
+		}
+		return false;
+	}
+	return true;
+}
+
+int32 UMeshUVChannelProperties::GetSelectedChannelIndex(bool bForceToZeroOnFailure)
+{
+	int32 FoundIndex = UVChannelNamesList.IndexOfByKey(UVChannel);
+	if (FoundIndex == INDEX_NONE)
+	{
+		return (bForceToZeroOnFailure ) ? 0 : -1;
+	}
+	return FoundIndex;
+}
\ No newline at end of file
diff --git a/Engine/Plugins/Runtime/MeshModelingToolset/Source/MeshModelingTools/Private/RecomputeUVsTool.cpp b/Engine/Plugins/Runtime/MeshModelingToolset/Source/MeshModelingTools/Private/RecomputeUVsTool.cpp
new file mode 100644
--- /dev/null
+++ b/Engine/Plugins/Runtime/MeshModelingToolset/Source/MeshModelingTools/Private/RecomputeUVsTool.cpp
@@ -0,0 +1,243 @@
+// Copyright Epic Games, Inc. All Rights Reserved.
+
+#include "RecomputeUVsTool.h"
+#include "InteractiveToolManager.h"
+#include "ToolBuilderUtil.h"
+#include "DynamicMesh/DynamicMesh3.h"
+#include "Polygroups/PolygroupUtil.h"
+#include "ToolSetupUtil.h"
+#include "ModelingToolTargetUtil.h"
+#include "ParameterizationOps/RecomputeUVsOp.h"
+
+using namespace UE::Geometry;
+
+#define LOCTEXT_NAMESPACE "URecomputeUVsTool"
+
+
+/*
+ * ToolBuilder
+ */
+
+USingleSelectionMeshEditingTool* URecomputeUVsToolBuilder::CreateNewTool(const FToolBuilderState& SceneState) const
+{
+	URecomputeUVsTool* NewTool = NewObject<URecomputeUVsTool>(SceneState.ToolManager);
+	return NewTool;
+}
+
+/*
+ * Tool
+ */
+
+
+void URecomputeUVsTool::Setup()
+{
+	UInteractiveTool::Setup();
+
+	InputMesh = MakeShared<FDynamicMesh3, ESPMode::ThreadSafe>();
+	*InputMesh = UE::ToolTarget::GetDynamicMeshCopy(Target);
+
+	FComponentMaterialSet MaterialSet = UE::ToolTarget::GetMaterialSet(Target);
+	FTransform TargetTransform = (FTransform)UE::ToolTarget::GetLocalToWorldTransform(Target);
+
+	// initialize our properties
+
+	Settings = NewObject<URecomputeUVsToolProperties>(this);
+	Settings->RestoreProperties(this);
+	AddToolPropertySource(Settings);
+
+	UVChannelProperties = NewObject<UMeshUVChannelProperties>(this);
+	UVChannelProperties->RestoreProperties(this);
+	UVChannelProperties->Initialize(InputMesh.Get(), false);
+	UVChannelProperties->ValidateSelection(true);
+	UVChannelProperties->WatchProperty(UVChannelProperties->UVChannel, [this](const FString& NewValue)
+		{
+			MaterialSettings->UpdateUVChannels(UVChannelProperties->UVChannelNamesList.IndexOfByKey(UVChannelProperties->UVChannel),
+			                                   UVChannelProperties->UVChannelNamesList);
+		});
+	AddToolPropertySource(UVChannelProperties);
+
+	PolygroupLayerProperties = NewObject<UPolygroupLayersProperties>(this);
+	PolygroupLayerProperties->RestoreProperties(this, TEXT("RecomputeUVsTool"));
+	PolygroupLayerProperties->InitializeGroupLayers(InputMesh.Get());
+	PolygroupLayerProperties->WatchProperty(PolygroupLayerProperties->ActiveGroupLayer, [&](FName) { OnSelectedGroupLayerChanged(); });
+	AddToolPropertySource(PolygroupLayerProperties);
+	UpdateActiveGroupLayer();
+
+	MaterialSettings = NewObject<UExistingMeshMaterialProperties>(this);
+	MaterialSettings->MaterialMode = ESetMeshMaterialMode::Checkerboard;
+	MaterialSettings->RestoreProperties(this, TEXT("ModelingUVTools"));
+	AddToolPropertySource(MaterialSettings);
+
+	UE::ToolTarget::HideSourceObject(Target);
+
+	// force update
+	MaterialSettings->UpdateMaterials();
+
+	RecomputeUVsOpFactory = NewObject<URecomputeUVsOpFactory>();
+	RecomputeUVsOpFactory->OriginalMesh = InputMesh;
+	RecomputeUVsOpFactory->Settings = Settings;
+	RecomputeUVsOpFactory->TargetTransform = UE::ToolTarget::GetLocalToWorldTransform(Target);
+	RecomputeUVsOpFactory->GetSelectedUVChannel = [this]() { return GetSelectedUVChannel(); };
+
+	Preview = NewObject<UMeshOpPreviewWithBackgroundCompute>(this);
+	Preview->Setup(this->TargetWorld, RecomputeUVsOpFactory);
+	ToolSetupUtil::ApplyRenderingConfigurationToPreview(Preview->PreviewMesh, Target);
+	Preview->PreviewMesh->SetTangentsMode(EDynamicMeshComponentTangentsMode::AutoCalculated);
+	Preview->PreviewMesh->ReplaceMesh(*InputMesh);
+	Preview->ConfigureMaterials(MaterialSet.Materials, ToolSetupUtil::GetDefaultWorkingMaterial(GetToolManager()));
+	Preview->PreviewMesh->SetTransform(TargetTransform);
+
+	Preview->OnMeshUpdated.AddLambda([this](UMeshOpPreviewWithBackgroundCompute* Op)
+		{
+			OnPreviewMeshUpdated();
+		});
+
+	Preview->OverrideMaterial = MaterialSettings->GetActiveOverrideMaterial();
+
+	if (bCreateUVLayoutViewOnSetup)
+	{
+		UVLayoutView = NewObject<UUVLayoutPreview>(this);
+		UVLayoutView->CreateInWorld(TargetWorld);
+		UVLayoutView->SetSourceMaterials(MaterialSet);
+		UVLayoutView->SetSourceWorldPosition(TargetTransform, UE::ToolTarget::GetTargetActor(Target)->GetComponentsBoundingBox());
+		UVLayoutView->Settings->bEnabled = false;
+		UVLayoutView->Settings->bShowWireframe = false;
+		UVLayoutView->Settings->RestoreProperties(this, TEXT("RecomputeUVsTool"));
+		AddToolPropertySource(UVLayoutView->Settings);
+	}
+
+	Preview->InvalidateResult();    // start compute
+
+	SetToolDisplayName(LOCTEXT("ToolNameLocal", "UV Unwrap"));
+	GetToolManager()->DisplayMessage(
+		LOCTEXT("OnStartTool_Regions", "Generate UVs for Polygroups or existing UV charts of the Mesh using various strategies."),
+		EToolMessageLevel::UserNotification);
+}
+
+
+void URecomputeUVsTool::OnPropertyModified(UObject* PropertySet, FProperty* Property)
+{
+	bool bForceMaterialUpdate = false;
+	if (PropertySet == Settings || PropertySet == UVChannelProperties)
+	{
+		// One of the UV generation properties must have changed.  Dirty the result to force a recompute
+		Preview->InvalidateResult();
+		bForceMaterialUpdate = true;
+	}
+
+	if (PropertySet == MaterialSettings || bForceMaterialUpdate)
+	{
+		MaterialSettings->UpdateMaterials();
+		Preview->OverrideMaterial = MaterialSettings->GetActiveOverrideMaterial();
+	}
+}
+
+
+void URecomputeUVsTool::Shutdown(EToolShutdownType ShutdownType)
+{
+	if (UVLayoutView)
+	{
+		UVLayoutView->Settings->SaveProperties(this, TEXT("RecomputeUVsTool"));
+		UVLayoutView->Disconnect();
+	}
+
+	UVChannelProperties->SaveProperties(this);
+	Settings->SaveProperties(this);
+	PolygroupLayerProperties->RestoreProperties(this, TEXT("RecomputeUVsTool"));
+	MaterialSettings->SaveProperties(this, TEXT("ModelingUVTools"));
+
+	UE::ToolTarget::ShowSourceObject(Target);
+
+	FDynamicMeshOpResult Result = Preview->Shutdown();
+	if (ShutdownType == EToolShutdownType::Accept)
+	{
+		GetToolManager()->BeginUndoTransaction(LOCTEXT("RecomputeUVs", "Recompute UVs"));
+		FDynamicMesh3* NewDynamicMesh = Result.Mesh.Get();
+		if (ensure(NewDynamicMesh))
+		{
+			UE::ToolTarget::CommitDynamicMeshUVUpdate(Target, NewDynamicMesh);
+		}
+		GetToolManager()->EndUndoTransaction();
+	}
+
+	if (RecomputeUVsOpFactory)
+	{
+		RecomputeUVsOpFactory = nullptr;
+	}
+
+}
+
+void URecomputeUVsTool::Render(IToolsContextRenderAPI* RenderAPI)
+{
+	if (UVLayoutView)
+	{
+		UVLayoutView->Render(RenderAPI);
+	}
+}
+
+
+void URecomputeUVsTool::OnTick(float DeltaTime)
+{
+	Preview->Tick(DeltaTime);
+
+	if (UVLayoutView)
+	{
+		UVLayoutView->OnTick(DeltaTime);
+	}
+}
+
+bool URecomputeUVsTool::CanAccept() const
+{
+	return Super::CanAccept() && Preview->HaveValidResult();
+}
+
+
+int32 URecomputeUVsTool::GetSelectedUVChannel() const
+{
+	return UVChannelProperties ? UVChannelProperties->GetSelectedChannelIndex(true) : 0;
+}
+
+void URecomputeUVsTool::OnSelectedGroupLayerChanged()
+{
+	UpdateActiveGroupLayer();
+	Preview->InvalidateResult();
+}
+
+
+void URecomputeUVsTool::UpdateActiveGroupLayer()
+{
+	if (PolygroupLayerProperties->HasSelectedPolygroup() == false)
+	{
+		ActiveGroupSet = MakeShared<UE::Geometry::FPolygroupSet, ESPMode::ThreadSafe>(InputMesh.Get());
+	}
+	else
+	{
+		FName SelectedName = PolygroupLayerProperties->ActiveGroupLayer;
+		FDynamicMeshPolygroupAttribute* FoundAttrib = UE::Geometry::FindPolygroupLayerByName(*InputMesh, SelectedName);
+		ensureMsgf(FoundAttrib, TEXT("Selected Attribute Not Found! Falling back to Default group layer."));
+		ActiveGroupSet = MakeShared<UE::Geometry::FPolygroupSet, ESPMode::ThreadSafe>(InputMesh.Get(), FoundAttrib);
+	}
+}
+
+
+void URecomputeUVsTool::OnPreviewMeshUpdated()
+{
+	if (UVLayoutView)
+	{
+		int32 UVChannel = UVChannelProperties ? UVChannelProperties->GetSelectedChannelIndex(true) : 0;
+		Preview->PreviewMesh->ProcessMesh([&](const FDynamicMesh3& NewMesh)
+		{
+			UVLayoutView->UpdateUVMesh(&NewMesh, UVChannel);
+		});
+	}
+
+	if (MaterialSettings)
+	{
+		MaterialSettings->UpdateMaterials();
+	}
+
+}
+
+
+
+#undef LOCTEXT_NAMESPACE
diff --git a/Engine/Plugins/Runtime/MeshModelingToolset/Source/MeshModelingTools/Public/Properties/MeshUVChannelProperties.h b/Engine/Plugins/Runtime/MeshModelingToolset/Source/MeshModelingTools/Public/Properties/MeshUVChannelProperties.h
new file mode 100644
--- /dev/null
+++ b/Engine/Plugins/Runtime/MeshModelingToolset/Source/MeshModelingTools/Public/Properties/MeshUVChannelProperties.h
@@ -0,0 +1,46 @@
+// Copyright Epic Games, Inc. All Rights Reserved.
+
+#pragma once
+
+#include "CoreMinimal.h"
+#include "InteractiveTool.h"
+#include "GeometryBase.h"
+#include "MeshUVChannelProperties.generated.h"
+
+struct FMeshDescription;
+PREDECLARE_USE_GEOMETRY_CLASS(FDynamicMesh3);
+
+
+UCLASS()
+class MESHMODELINGTOOLS_API UMeshUVChannelProperties : public UInteractiveToolPropertySet
+{
+	GENERATED_BODY()
+public:
+
+	/** Select UV channel in the mesh */
+	UPROPERTY(EditAnywhere, Category = "UV Channel", meta = (DisplayName = "UV Channel", GetOptions = GetUVChannelNamesFunc, NoResetToDefault))
+	FString UVChannel;
+
+	UFUNCTION()
+	const TArray<FString>& GetUVChannelNamesFunc() const;
+
+	UPROPERTY(meta = (TransientToolProperty))
+	TArray<FString> UVChannelNamesList;
+
+	void Initialize(int32 NumUVChannels, bool bInitializeSelection = true);
+	void Initialize(const FMeshDescription* MeshDescription, bool bInitializeSelection = true);
+	void Initialize(const FDynamicMesh3* Mesh, bool bInitializeSelection = true);
+
+	/**
+	 * Verify that the UV channel selection is valid
+	 * @param bUpdateIfInvalid if selection is not valid, reset UVChannel to UV0 or empty if no UV channels exist
+	 * @return true if selection in UVChannel is an entry in UVChannelNamesList.
+	 */
+	bool ValidateSelection(bool bUpdateIfInvalid);
+
+	/**
+	 * @param bForceToZeroOnFailure if true, then instead of returning -1 we return 0 so calling code can fallback to default UV paths
+	 * @return selected UV channel index, or -1 if invalid selection, or 0 if invalid selection and bool bForceToZeroOnFailure = true
+	 */
+	int32 GetSelectedChannelIndex(bool bForceToZeroOnFailure = false);
+};
diff --git a/Engine/Plugins/Runtime/MeshModelingToolset/Source/MeshModelingTools/Public/RecomputeUVsTool.h b/Engine/Plugins/Runtime/MeshModelingToolset/Source/MeshModelingTools/Public/RecomputeUVsTool.h
new file mode 100644
--- /dev/null
+++ b/Engine/Plugins/Runtime/MeshModelingToolset/Source/MeshModelingTools/Public/RecomputeUVsTool.h
@@ -0,0 +1,93 @@
+// Copyright Epic Games, Inc. All Rights Reserved.
+
+#pragma once
+
+#include "CoreMinimal.h"
+#include "BaseTools/SingleSelectionMeshEditingTool.h"
+#include "DynamicMesh/DynamicMesh3.h"
+#include "MeshOpPreviewHelpers.h"
+#include "Properties/MeshMaterialProperties.h"
+#include "Properties/MeshUVChannelProperties.h"
+#include "Properties/RecomputeUVsProperties.h"
+#include "PropertySets/PolygroupLayersProperties.h"
+#include "Polygroups/PolygroupSet.h"
+#include "Drawing/UVLayoutPreview.h"
+
+#include "RecomputeUVsTool.generated.h"
+
+
+// Forward declarations
+class UDynamicMeshComponent;
+class UMaterialInterface;
+class UMaterialInstanceDynamic;
+class RecomputeUVsOp;
+class URecomputeUVsOpFactory;
+
+
+/**
+ *
+ */
+UCLASS()
+class MESHMODELINGTOOLS_API URecomputeUVsToolBuilder : public USingleSelectionMeshEditingToolBuilder
+{
+	GENERATED_BODY()
+
+public:
+	virtual USingleSelectionMeshEditingTool* CreateNewTool(const FToolBuilderState& SceneState) const override;
+};
+
+
+/**
+ * URecomputeUVsTool Recomputes UVs based on existing segmentations of the mesh
+ */
+UCLASS()
+class MESHMODELINGTOOLS_API URecomputeUVsTool : public USingleSelectionMeshEditingTool
+{
+	GENERATED_BODY()
+
+public:
+	virtual void Setup() override;
+	virtual void Shutdown(EToolShutdownType ShutdownType) override;
+
+	virtual void Render(IToolsContextRenderAPI* RenderAPI) override;
+	virtual void OnTick(float DeltaTime) override;
+
+	virtual bool HasCancel() const override { return true; }
+	virtual bool HasAccept() const override { return true; }
+	virtual bool CanAccept() const override;
+
+	virtual void OnPropertyModified(UObject* PropertySet, FProperty* Property) override;
+
+protected:
+	UPROPERTY()
+	TObjectPtr<UMeshUVChannelProperties> UVChannelProperties = nullptr;
+
+	UPROPERTY()
+	TObjectPtr<URecomputeUVsToolProperties> Settings = nullptr;
+
+	UPROPERTY()
+	TObjectPtr<UPolygroupLayersProperties> PolygroupLayerProperties = nullptr;
+
+	UPROPERTY()
+	TObjectPtr<UExistingMeshMaterialProperties> MaterialSettings = nullptr;
+
+	UPROPERTY()
+	bool bCreateUVLayoutViewOnSetup = true;
+
+	UPROPERTY()
+	TObjectPtr<UUVLayoutPreview> UVLayoutView = nullptr;
+
+	UPROPERTY()
+	TObjectPtr<URecomputeUVsOpFactory> RecomputeUVsOpFactory;
+
+	UPROPERTY()
+	TObjectPtr<UMeshOpPreviewWithBackgroundCompute> Preview = nullptr;
+
+	TSharedPtr<UE::Geometry::FDynamicMesh3, ESPMode::ThreadSafe> InputMesh;
+
+	TSharedPtr<UE::Geometry::FPolygroupSet, ESPMode::ThreadSafe> ActiveGroupSet;
+	void OnSelectedGroupLayerChanged();
+	void UpdateActiveGroupLayer();
+	int32 GetSelectedUVChannel() const;
+	void OnPreviewMeshUpdated();
+};
diff --git a/Engine/Plugins/Runtime/MeshModelingToolset/Source/ModelingComponents/Public/PropertySets/PolygroupLayersProperties.h b/Engine/Plugins/Runtime/MeshModelingToolset/Source/ModelingComponents/Public/PropertySets/PolygroupLayersProperties.h
--- a/Engine/Plugins/Runtime/MeshModelingToolset/Source/ModelingComponents/Public/PropertySets/PolygroupLayersProperties.h
+++ b/Engine/Plugins/Runtime/MeshModelingToolset/Source/ModelingComponents/Public/PropertySets/PolygroupLayersProperties.h
@@ -20,10 +20,10 @@ class MODELINGCOMPONENTS_API UPolygroupLayersProperties : public UInteractiveToo
 public:
 
 	/** Select polygroup layer. */
-	UPROPERTY(EditAnywhere, Category = PolygroupLayers, meta = (GetOptions = GetGroupLayersFunc))
+	UPROPERTY(EditAnywhere, Category = "PolyGroup Layer", meta = (DisplayName = "Active PolyGroup", GetOptions = GetGroupLayersFunc))
 	FName ActiveGroupLayer = "Default";
 
-	// this function is called provide set of available group layers
+	// Provides set of available group layers
 	UFUNCTION()
 	TArray<FString> GetGroupLayersFunc() { return GroupLayersList; }
 
diff --git a/Engine/Plugins/Runtime/MeshModelingToolset/Source/ModelingOperators/Private/ParameterizationOps/RecomputeUVsOp.cpp b/Engine/Plugins/Runtime/MeshModelingToolset/Source/ModelingOperators/Private/ParameterizationOps/RecomputeUVsOp.cpp
--- a/Engine/Plugins/Runtime/MeshModelingToolset/Source/ModelingOperators/Private/ParameterizationOps/RecomputeUVsOp.cpp
+++ b/Engine/Plugins/Runtime/MeshModelingToolset/Source/ModelingOperators/Private/ParameterizationOps/RecomputeUVsOp.cpp
@@ -383,56 +383,61 @@ bool FRecomputeUVsOp::CalculateResult_RegionOptimization(FProgressCancel* Progre
 
 TUniquePtr<FDynamicMeshOperator> URecomputeUVsOpFactory::MakeNewOperator()
 {
-	FAxisAlignedBox3d MeshBounds = OriginalMesh->GetBounds();
 	TUniquePtr<FRecomputeUVsOp> RecomputeUVsOp = MakeUnique<FRecomputeUVsOp>();
+
 	RecomputeUVsOp->InputMesh = OriginalMesh;
 	RecomputeUVsOp->InputGroups = InputGroups;
 	RecomputeUVsOp->UVLayer = GetSelectedUVChannel();
 
-	RecomputeUVsOp->IslandMode = (UE::Geometry::ERecomputeUVsIslandMode)(int)Settings->IslandMode;
-	RecomputeUVsOp->UnwrapType = (UE::Geometry::ERecomputeUVsUnwrapType)(int)Settings->UnwrapType;
+	RecomputeUVsOp->IslandMode = static_cast<ERecomputeUVsIslandMode>(Settings->IslandGeneration);
+
+	switch (Settings->UnwrapType)
+	{
+	case ERecomputeUVsPropertiesUnwrapType::ExpMap:
+		RecomputeUVsOp->UnwrapType = ERecomputeUVsUnwrapType::ExpMap;
+		RecomputeUVsOp->bMergingOptimization = false;
+		break;
+	case ERecomputeUVsPropertiesUnwrapType::Conformal:
+		RecomputeUVsOp->UnwrapType = ERecomputeUVsUnwrapType::ConformalFreeBoundary;
+		RecomputeUVsOp->bMergingOptimization = false;
+		break;
+	case ERecomputeUVsPropertiesUnwrapType::IslandMerging:
+		RecomputeUVsOp->UnwrapType = ERecomputeUVsUnwrapType::ExpMap;
+		RecomputeUVsOp->bMergingOptimization = true;
+	}
 
-	RecomputeUVsOp->bAutoRotate = (Settings->AutoRotation == ERecomputeUVsToolOrientationMode::MinBoxBounds);
+	RecomputeUVsOp->bAutoRotate = (Settings->AutoRotation == ERecomputeUVsToolOrientationMode::MinBounds);
 
 	RecomputeUVsOp->NormalSmoothingRounds = Settings->SmoothingSteps;
 	RecomputeUVsOp->NormalSmoothingAlpha = Settings->SmoothingAlpha;
 
-	RecomputeUVsOp->bMergingOptimization = Settings->bIslandMerging;
-	RecomputeUVsOp->MergingThreshold = Settings->MergingThreshold;
-	RecomputeUVsOp->MaxNormalDeviationDeg = Settings->MaxAngleDeviation;
+	RecomputeUVsOp->MergingThreshold = Settings->MergingDistortionLimit;
+	RecomputeUVsOp->MaxNormalDeviationDeg = Settings->MergingAngleLimit;
+
+	RecomputeUVsOp->bPackUVs = false;
+	RecomputeUVsOp->bNormalizeAreas = false;
 
-	RecomputeUVsOp->bPackUVs = Settings->bAutoPack;
-	if (Settings->bAutoPack)
+	switch (Settings->LayoutType)
 	{
+	case ERecomputeUVsPropertiesLayoutType::None:
+		break;
+	case ERecomputeUVsPropertiesLayoutType::Repack:
+		RecomputeUVsOp->bPackUVs = true;
 		RecomputeUVsOp->PackingTextureResolution = Settings->TextureResolution;
-		RecomputeUVsOp->bNormalizeAreas = false;
-		RecomputeUVsOp->AreaScaling = 1.0;
-	}
-	else
-	{
-		switch (Settings->UVScaleMode)
-		{
-		case ERecomputeUVsToolUVScaleMode::NoScaling:
-			RecomputeUVsOp->bNormalizeAreas = false;
-			RecomputeUVsOp->AreaScaling = 1.0;
-			break;
-		case ERecomputeUVsToolUVScaleMode::NormalizeToBounds:
-			RecomputeUVsOp->bNormalizeAreas = true;
-			RecomputeUVsOp->AreaScaling = Settings->UVScale / MeshBounds.MaxDim();
-			break;
-		case ERecomputeUVsToolUVScaleMode::NormalizeToWorld:
-			RecomputeUVsOp->bNormalizeAreas = true;
-			RecomputeUVsOp->AreaScaling = Settings->UVScale;
-			break;
-		}
+		break;
+	case ERecomputeUVsPropertiesLayoutType::NormalizeToBounds:
+		RecomputeUVsOp->bNormalizeAreas = true;
+		RecomputeUVsOp->AreaScaling = Settings->NormalizeScale / OriginalMesh->GetBounds().MaxDim();
+		break;
+	case ERecomputeUVsPropertiesLayoutType::NormalizeToWorld:
+		RecomputeUVsOp->bNormalizeAreas = true;
+		RecomputeUVsOp->AreaScaling = Settings->NormalizeScale;
 	}
 
-
 	RecomputeUVsOp->SetTransform(TargetTransform);
 
 	return RecomputeUVsOp;
 }
 
 
-
-#undef LOCTEXT_NAMESPACE
\ No newline at end of file
+#undef LOCTEXT_NAMESPACE
diff --git a/Engine/Plugins/Runtime/MeshModelingToolset/Source/ModelingOperators/Public/Properties/RecomputeUVsProperties.h b/Engine/Plugins/Runtime/MeshModelingToolset/Source/ModelingOperators/Public/Properties/RecomputeUVsProperties.h
--- a/Engine/Plugins/Runtime/MeshModelingToolset/Source/ModelingOperators/Public/Properties/RecomputeUVsProperties.h
+++ b/Engine/Plugins/Runtime/MeshModelingToolset/Source/ModelingOperators/Public/Properties/RecomputeUVsProperties.h
@@ -14,22 +14,23 @@ enum class ERecomputeUVsPropertiesUnwrapType
 {
 	// Values must match UE::Geometry::ERecomputeUVsUnwrapType
 
-	/** ExpMap UV Flattening Algorithm is very fast but has limited ability to control stretching/distortion */
-	ExpMap = 0,
-	/** Conformal UV Flattening algorithm is increasingly-expensive on large islands but better at distortion control */
-	Conformal = 1
+	/** ExpMap UV flattening is very fast but has limited abilities to reduce stretching and distortion */
+	ExpMap = 0 UMETA(DisplayName = "ExpMap"),
+	/** Conformal UV flattening is increasingly expensive on large islands but reduces distortion */
+	Conformal = 1,
+	/** UV islands will be merged into larger islands if it does not increase stretching and distortion beyond defined limits */
+	IslandMerging = 2
 };
 
 
-
 UENUM()
 enum class ERecomputeUVsPropertiesIslandMode
 {
 	// Values must match UE::Geometry::ERecomputeUVsIslandMode
 
-	/** Use Active Mesh Polygroups Layer to define initial UV Islands */
-	PolyGroups = 0,
-	/** Use Existing UV Layer to define UV Islands (ie re-solve UV flattening based on existing UVs) */
+	/** Use Polygroups defined by Active PolyGroup property to define initial UV islands */
+	PolyGroups = 0 UMETA(DisplayName = "PolyGroups"),
+	/** Use existing UV Layer to define UV islands, i.e. re-solve UV flattening based on existing UVs */
 	ExistingUVs = 1
 };
 
@@ -37,21 +38,23 @@ enum class ERecomputeUVsPropertiesIslandMode
 UENUM()
 enum class ERecomputeUVsToolOrientationMode
 {
-	/**  */
+	/** Do not rotate UV islands */
 	None,
-	/**  */
-	MinBoxBounds
+	/** Automatically rotate each UV island to reduce its bounding box size */
+	MinBounds
 };
 
 
 UENUM()
-enum class ERecomputeUVsToolUVScaleMode
+enum class ERecomputeUVsPropertiesLayoutType
 {
-	/** No scaling is applied to UV islands */
-	NoScaling,
-	/** Scale UV islands such that they have constant relative area, relative to object bounds */
+	/** Do not apply additional layout options to unwrapped UVs */
+	None,
+	/** Uniformly scale and translate UV islands collectively to pack them into the unit square, i.e. fit between 0 and 1 with no overlap */
+	Repack,
+	/** Uniformly scale UV islands such that they have constant relative area, relative to object bounds */
 	NormalizeToBounds,
-	/** Scale UV islands such that they have constant relative area, relative to world space */
+	/** Uniformly scale UV islands such that they have constant relative area, relative to world space */
 	NormalizeToWorld
 };
 
@@ -62,52 +65,49 @@ class MODELINGOPERATORS_API URecomputeUVsToolProperties : public UInteractiveToo
 	GENERATED_BODY()
 public:
 
-	/** IslandMode determines which patches of triangles will be used as initial UV Islands */
-	UPROPERTY(EditAnywhere, Category = "UV Generation")
-	ERecomputeUVsPropertiesIslandMode IslandMode = ERecomputeUVsPropertiesIslandMode::ExistingUVs;
+	/** Generation method for initial UV islands */
+	UPROPERTY(EditAnywhere, Category = "UV Unwrap")
+	ERecomputeUVsPropertiesIslandMode IslandGeneration = ERecomputeUVsPropertiesIslandMode::ExistingUVs;
 
-	/** Which type of UV flattening algorithm to use */
-	UPROPERTY(EditAnywhere, Category = "UV Generation", meta = (EditCondition = "bIslandMerging == false"))
+	/** Type of UV flattening algorithm to use */
+	UPROPERTY(EditAnywhere, Category = "UV Unwrap")
 	ERecomputeUVsPropertiesUnwrapType UnwrapType = ERecomputeUVsPropertiesUnwrapType::Conformal;
 
-	/** Which type of automatic island orientation to use */
-	UPROPERTY(EditAnywhere, Category = "UV Generation")
-	ERecomputeUVsToolOrientationMode AutoRotation = ERecomputeUVsToolOrientationMode::MinBoxBounds;
-
-	/** If Island Merging is enabled, then the initial UV islands will be merged into larger islands if it does not increase distortion/stretching beyond the limits below */
-	UPROPERTY(EditAnywhere, Category = "UV Generation")
-	bool bIslandMerging = false;
-
-	/** Distortion/Stretching Threshold for island merging - larger values increase the allowable UV stretching */
-	UPROPERTY(EditAnywhere, Category = "UV Generation", meta = (UIMin = "1.0", UIMax = "5.0", ClampMin = "1.0", EditCondition = "bIslandMerging == true"))
-	float MergingThreshold = 1.5f;
-
-	/** UV islands will not be merged if their average face normals deviate by larger than this amount */
-	UPROPERTY(EditAnywhere, Category = "UV Generation", meta = (UIMin = "0.0", UIMax = "90.0", ClampMin = "0.0", ClampMax = "180.0", EditCondition = "bIslandMerging == true"))
-	float MaxAngleDeviation = 45.0f;
+	/** Type of automatic rotation applied to each UV island */
+	UPROPERTY(EditAnywhere, Category = "UV Unwrap")
+	ERecomputeUVsToolOrientationMode AutoRotation = ERecomputeUVsToolOrientationMode::MinBounds;
 
-	/** Number of smoothing steps to apply in the ExpMap UV Generation method (Smoothing slightly increases distortion but produces more stable results) */
-	UPROPERTY(EditAnywhere, Category = "ExpMap UV Generator", meta = (UIMin = "0", UIMax = "25", ClampMin = "0", ClampMax = "1000", EditCondition = "UnwrapType == ERecomputeUVsUnwrapType::ExpMap || bIslandMerging == true"))
+	/** Number of smoothing steps to apply; this slightly increases distortion but produces more stable results. This is only enabled if the Unwrap Type is set to ExpMap or Island Merging. */
+	UPROPERTY(EditAnywhere, Category = "UV Unwrap", meta = (UIMin = "0", UIMax = "25", ClampMin = "0", ClampMax = "1000",
+		EditCondition = "UnwrapType == ERecomputeUVsPropertiesUnwrapType::ExpMap || UnwrapType == ERecomputeUVsPropertiesUnwrapType::IslandMerging"))
 	int SmoothingSteps = 5;
 
-	/** Smoothing parameter, larger values result in faster smoothing in each step */
-	UPROPERTY(EditAnywhere, Category = "ExpMap UV Generator", meta = (UIMin = "0", UIMax = "1.0", ClampMin = "0", ClampMax = "1.0", EditCondition = "UnwrapType == ERecomputeUVsUnwrapType::ExpMap || bIslandMerging == true"))
+	/** Smoothing parameter, larger values result in faster smoothing in each step. This is only enabled if the Unwrap Type is set to ExpMap or Island Merging. */
+	UPROPERTY(EditAnywhere, Category = "UV Unwrap", meta = (UIMin = "0", UIMax = "1.0", ClampMin = "0", ClampMax = "1.0",
+		EditCondition = "UnwrapType == ERecomputeUVsPropertiesUnwrapType::ExpMap || UnwrapType == ERecomputeUVsPropertiesUnwrapType::IslandMerging"))
 	float SmoothingAlpha = 0.25f;
 
-	/** If enabled, result UVs are automatically packed into the standard UV 0-1 square */
+	/** Threshold for stretching and distortion below which island merging is allowed; larger values increase the allowable UV distortion. This is only enabled if the Unwrap Type is set to Island Merging. */
+	UPROPERTY(EditAnywhere, Category = "UV Unwrap", meta = (UIMin = "1.0", UIMax = "5.0", ClampMin = "1.0",
+		EditCondition = "UnwrapType == ERecomputeUVsPropertiesUnwrapType::IslandMerging"))
+	float MergingDistortionLimit = 1.5f;
+
+	/** Threshold for the average face normal deviation below  which island merging is allowed. This is only enabled if the Unwrap Type is set to Island Merging. */
+	UPROPERTY(EditAnywhere, Category = "UV Unwrap", meta = (UIMin = "0.0", UIMax = "90.0", ClampMin = "0.0", ClampMax = "180.0",
+		EditCondition = "UnwrapType == ERecomputeUVsPropertiesUnwrapType::IslandMerging"))
+	float MergingAngleLimit = 45.0f;
+
+	/** Uniformly scale and translate UV islands collectively to pack them into the unit square, i.e. fit between 0 and 1 with no overlap. */
 	UPROPERTY(EditAnywhere, Category = "UV Layout")
-	bool bAutoPack = true;
+	ERecomputeUVsPropertiesLayoutType LayoutType = ERecomputeUVsPropertiesLayoutType::Repack;
 
-	/** Target texture resolution used for UV packing, which determines gutter size */
-	UPROPERTY(EditAnywhere, Category = "UV Layout", meta = (UIMin = "64", UIMax = "2048", ClampMin = "2", ClampMax = "4096", EditCondition = bAutoPack))
+	/** Expected resolution of the output textures; this controls spacing left between UV islands to avoid interpolation artifacts. This is only enabled when the Layout Type is set to Repack. */
+	UPROPERTY(EditAnywhere, Category = "UV Layout", meta = (UIMin = "64", UIMax = "2048", ClampMin = "2", ClampMax = "4096",
+		EditCondition = "LayoutType == ERecomputeUVsPropertiesLayoutType::Repack"))
 	int TextureResolution = 1024;
 
-	/** Scaling applied to UV islands after solve */
-	UPROPERTY(EditAnywhere, Category = "Transform UVs", meta = (EditCondition = "bAutoPack == false"))
-	ERecomputeUVsToolUVScaleMode UVScaleMode = ERecomputeUVsToolUVScaleMode::NormalizeToBounds;
-
-	/** Scaling factor used for UV island normalization/scaling */
-	UPROPERTY(EditAnywhere, Category = "Transform UVs", meta = (EditCondition = "UVScaleMode!=ERecomputeUVsToolUVScaleMode::NoScaling && bAutoPack == false", UIMin = "0.001", UIMax = "10", ClampMin = "0.00001", ClampMax = "1000000.0"))
-	float UVScale = 1.0;
+	/** Scaling factor used for UV island normalization/scaling. This is only enabled when the Layout Type is set to Normalize to Bounds or Normalize to World. */
+	UPROPERTY(EditAnywhere, Category = "UV Layout",	meta = (UIMin = "0.001", UIMax = "10", ClampMin = "0.00001", ClampMax = "1000000.0",
+		EditCondition = "LayoutType == ERecomputeUVsPropertiesLayoutType::NormalizeToBounds || LayoutType == ERecomputeUVsPropertiesLayoutType::NormalizeToWorld"))
+	float NormalizeScale = 1.0f;
 };
-
