MeshModelingToolset: Move UV Projection tool out of Experimental

diff --git a/Engine/Plugins/Experimental/MeshModelingToolsetExp/Source/MeshModelingToolsExp/Private/UVProjectionTool.cpp b/Engine/Plugins/Experimental/MeshModelingToolsetExp/Source/MeshModelingToolsExp/Private/UVProjectionTool.cpp
deleted file mode 100644
diff --git a/Engine/Plugins/Experimental/MeshModelingToolsetExp/Source/MeshModelingToolsExp/Public/UVProjectionTool.h b/Engine/Plugins/Experimental/MeshModelingToolsetExp/Source/MeshModelingToolsExp/Public/UVProjectionTool.h
deleted file mode 100644
diff --git a/Engine/Plugins/Runtime/MeshModelingToolset/Source/MeshModelingTools/Private/UVProjectionTool.cpp b/Engine/Plugins/Runtime/MeshModelingToolset/Source/MeshModelingTools/Private/UVProjectionTool.cpp
new file mode 100644
--- /dev/null
+++ b/Engine/Plugins/Runtime/MeshModelingToolset/Source/MeshModelingTools/Private/UVProjectionTool.cpp
@@ -0,0 +1,598 @@
+// Copyright Epic Games, Inc. All Rights Reserved.
+
+#include "UVProjectionTool.h"
+#include "InteractiveToolManager.h"
+#include "ToolTargetManager.h"
+#include "ToolBuilderUtil.h"
+#include "ToolSetupUtil.h"
+#include "BaseBehaviors/SingleClickBehavior.h"
+
+#include "BaseGizmos/TransformGizmoUtil.h"
+#include "ModelingToolTargetUtil.h"
+#include "Selection/StoredMeshSelectionUtil.h"
+
+#include "DynamicMesh/DynamicMesh3.h"
+#include "DynamicMesh/DynamicMeshAABBTree3.h"
+#include "DynamicMesh/MeshIndexUtil.h"
+#include "MeshRegionBoundaryLoops.h"
+#include "Parameterization/DynamicMeshUVEditor.h"
+#include "Operations/MeshConvexHull.h"
+#include "MinVolumeBox3.h"
+
+using namespace UE::Geometry;
+
+#define LOCTEXT_NAMESPACE "UUVProjectionTool"
+
+/*
+ * ToolBuilder
+ */
+
+USingleSelectionMeshEditingTool* UUVProjectionToolBuilder::CreateNewTool(const FToolBuilderState& SceneState) const
+{
+	return NewObject<UUVProjectionTool>(SceneState.ToolManager);
+}
+
+/*
+ * Propert Sets
+ */
+
+void UUVProjectionToolEditActions::PostAction(EUVProjectionToolActions Action)
+{
+	if (ParentTool.IsValid())
+	{
+		ParentTool->RequestAction(Action);
+	}
+}
+
+
+
+
+/*
+ * Tool
+ */
+
+void UUVProjectionTool::Setup()
+{
+	UInteractiveTool::Setup();
+
+	UE::ToolTarget::HideSourceObject(Target);
+
+	BasicProperties = NewObject<UUVProjectionToolProperties>(this);
+	BasicProperties->RestoreProperties(this);
+	BasicProperties->GetOnModified().AddLambda([this](UObject*, FProperty*)	{ Preview->InvalidateResult(); });
+
+	// initialize input mesh
+	InitializeMesh();
+
+	// initialize the PreviewMesh+BackgroundCompute object
+	UpdateNumPreviews();
+
+	UVChannelProperties = NewObject<UMeshUVChannelProperties>(this);
+	UVChannelProperties->RestoreProperties(this);
+	UVChannelProperties->Initialize(InputMesh.Get(), false);
+	UVChannelProperties->ValidateSelection(true);
+	UVChannelProperties->WatchProperty(UVChannelProperties->UVChannel, [this](const FString& NewValue)
+	{
+		MaterialSettings->UpdateUVChannels(UVChannelProperties->UVChannelNamesList.IndexOfByKey(UVChannelProperties->UVChannel),
+		                                   UVChannelProperties->UVChannelNamesList);
+		Preview->InvalidateResult();
+		OnMaterialSettingsChanged();
+	});
+
+	MaterialSettings = NewObject<UExistingMeshMaterialProperties>(this);
+	MaterialSettings->RestoreProperties(this);
+	MaterialSettings->GetOnModified().AddLambda([this](UObject*, FProperty*)
+	{
+		OnMaterialSettingsChanged();
+	});
+
+	EditActions = NewObject<UUVProjectionToolEditActions>(this);
+	EditActions->Initialize(this);
+
+	AddToolPropertySource(UVChannelProperties);
+	AddToolPropertySource(BasicProperties);
+	AddToolPropertySource(EditActions);
+	AddToolPropertySource(MaterialSettings);
+	
+	OnMaterialSettingsChanged();
+
+	// set up visualizers
+	ProjectionShapeVisualizer.LineColor = FLinearColor::Red;
+	ProjectionShapeVisualizer.LineThickness = 2.0;
+
+	// initialize Gizmo
+	UE::Geometry::FTransform3d GizmoPositionWorld(WorldBounds.Center());
+	TransformProxy = NewObject<UTransformProxy>(this);
+	TransformProxy->SetTransform((FTransform)GizmoPositionWorld);
+	TransformProxy->OnTransformChanged.AddUObject(this, &UUVProjectionTool::TransformChanged);
+
+	TransformGizmo = UE::TransformGizmoUtil::CreateCustomTransformGizmo(
+		GetToolManager(), ETransformGizmoSubElements::StandardTranslateRotate, this);
+	TransformGizmo->SetActiveTarget(TransformProxy, GetToolManager());
+
+	InitialDimensions = BasicProperties->Dimensions;
+	bInitialUniformDimensions = BasicProperties->bUniformDimensions;
+	InitialTransform = TransformGizmo->GetGizmoTransform();
+	
+	ApplyInitializationMode();
+
+	// start watching for dimensions changes
+	DimensionsWatcher = BasicProperties->WatchProperty(BasicProperties->Dimensions, [this](FVector) { bTransformModified = true; });
+	DimensionsModeWatcher = BasicProperties->WatchProperty(BasicProperties->bUniformDimensions, [this](bool) { bTransformModified = true; });
+	BasicProperties->WatchProperty(BasicProperties->Initialization, [this](EUVProjectionToolInitializationMode NewMode) { OnInitializationModeChanged(); });
+	bTransformModified = false;
+	BasicProperties->SilentUpdateWatched();
+
+	// click to set plane behavior
+	SetPlaneCtrlClickBehaviorTarget = MakeUnique<FSelectClickedAction>();
+	SetPlaneCtrlClickBehaviorTarget->World = this->TargetWorld;
+	SetPlaneCtrlClickBehaviorTarget->OnClickedPositionFunc = [this](const FHitResult& Hit)
+	{
+		UpdatePlaneFromClick(FVector3d(Hit.ImpactPoint), FVector3d(Hit.ImpactNormal), SetPlaneCtrlClickBehaviorTarget->bShiftModifierToggle);
+	};
+	SetPlaneCtrlClickBehaviorTarget->InvisibleComponentsToHitTest.Add(UE::ToolTarget::GetTargetComponent(Target));
+	ClickToSetPlaneBehavior = NewObject<USingleClickInputBehavior>();
+	ClickToSetPlaneBehavior->ModifierCheckFunc = FInputDeviceState::IsCtrlKeyDown;
+	ClickToSetPlaneBehavior->Modifiers.RegisterModifier(FSelectClickedAction::ShiftModifier, FInputDeviceState::IsShiftKeyDown);
+	ClickToSetPlaneBehavior->Initialize(SetPlaneCtrlClickBehaviorTarget.Get());
+	AddInputBehavior(ClickToSetPlaneBehavior);
+
+	// probably already done
+	Preview->InvalidateResult();
+
+	SetToolDisplayName(LOCTEXT("ToolName", "UV Projection"));
+	GetToolManager()->DisplayMessage(
+		LOCTEXT("UVProjectionToolDescription", "Generate UVs for a Mesh by projecting onto simple geometric shapes. Ctrl+click to reposition shape. Face selections can be made in the PolyEdit and TriEdit Tools."),
+		EToolMessageLevel::UserNotification);
+}
+
+
+void UUVProjectionTool::UpdateNumPreviews()
+{
+	OperatorFactory = NewObject<UUVProjectionOperatorFactory>();
+	OperatorFactory->Tool = this;
+
+	Preview = NewObject<UMeshOpPreviewWithBackgroundCompute>(OperatorFactory);
+	Preview->Setup(this->TargetWorld, OperatorFactory);
+	ToolSetupUtil::ApplyRenderingConfigurationToPreview(Preview->PreviewMesh, Target); 
+	Preview->OnMeshUpdated.AddUObject(this, &UUVProjectionTool::OnMeshUpdated);
+	Preview->PreviewMesh->SetTangentsMode(EDynamicMeshComponentTangentsMode::AutoCalculated);
+
+	FComponentMaterialSet MaterialSet = UE::ToolTarget::GetMaterialSet(Target, false);
+	Preview->ConfigureMaterials(MaterialSet.Materials,
+		ToolSetupUtil::GetDefaultWorkingMaterial(GetToolManager())
+	);
+	Preview->PreviewMesh->UpdatePreview(InputMesh.Get());
+	Preview->PreviewMesh->SetTransform((FTransform)WorldTransform);
+
+	Preview->SetVisibility(true);
+
+	EdgeRenderer = NewObject<UPreviewGeometry>(this);
+	EdgeRenderer->CreateInWorld(this->TargetWorld, (FTransform)WorldTransform);
+
+	// if we have an ROI, show its borders
+	if (TriangleROI->Num() > 0)
+	{
+		FDynamicMesh3* UseMesh = InputMesh.Get();
+		FMeshRegionBoundaryLoops BoundaryLoops(UseMesh, *TriangleROI, true);
+		TSet<int32> BorderEdges;
+		for (const FEdgeLoop& Loop : BoundaryLoops.GetLoops())
+		{
+			BorderEdges.Append(Loop.Edges);
+		}
+
+		const FColor ROIBorderColor(240, 15, 240);
+		const float ROIBorderThickness = 4.0f;
+		const float ROIBorderDepthBias = 0.1f * (float)(WorldBounds.DiagonalLength() * 0.01);
+		EdgeRenderer->CreateOrUpdateLineSet(TEXT("ROIBorders"), UseMesh->MaxEdgeID(), [&](int32 eid, TArray<FRenderableLine>& LinesOut) {
+			if (BorderEdges.Contains(eid))
+			{
+				FVector3d A, B;
+				UseMesh->GetEdgeV(eid, A, B);
+				LinesOut.Add(FRenderableLine((FVector)A, (FVector)B, ROIBorderColor, ROIBorderThickness, ROIBorderDepthBias));
+			}
+		}, 1);
+	}
+}
+
+
+void UUVProjectionTool::Shutdown(EToolShutdownType ShutdownType)
+{
+	UVChannelProperties->SaveProperties(this);
+	BasicProperties->SavedDimensions = BasicProperties->Dimensions;
+	BasicProperties->bSavedUniformDimensions = BasicProperties->bUniformDimensions;
+	BasicProperties->SavedTransform = TransformGizmo->GetGizmoTransform();
+	BasicProperties->SaveProperties(this);
+	MaterialSettings->SaveProperties(this);
+
+	EdgeRenderer->Disconnect();
+
+	// Restore (unhide) the source meshes
+	UE::ToolTarget::ShowSourceObject(Target);
+
+	FDynamicMeshOpResult Result = Preview->Shutdown();
+	if (ShutdownType == EToolShutdownType::Accept)
+	{
+		GetToolManager()->BeginUndoTransaction(LOCTEXT("UVProjectionToolTransactionName", "Project UVs"));
+		FDynamicMesh3* NewDynamicMesh = Result.Mesh.Get();
+		if (ensure(NewDynamicMesh))
+		{
+			UE::ToolTarget::CommitDynamicMeshUVUpdate(Target, NewDynamicMesh);
+		}
+		GetToolManager()->EndUndoTransaction();
+	}
+
+	GetToolManager()->GetPairedGizmoManager()->DestroyAllGizmosByOwner(this);
+	TransformGizmo = nullptr;
+	TransformProxy = nullptr;
+}
+
+
+void UUVProjectionTool::RequestAction(EUVProjectionToolActions ActionType)
+{
+	if (bHavePendingAction)
+	{
+		return;
+	}
+	PendingAction = ActionType;
+	bHavePendingAction = true;
+}
+
+
+
+TUniquePtr<FDynamicMeshOperator> UUVProjectionOperatorFactory::MakeNewOperator()
+{
+	TUniquePtr<FUVProjectionOp> Op = MakeUnique<FUVProjectionOp>();
+	Op->ProjectionMethod = Tool->BasicProperties->ProjectionType;
+
+	Op->MeshToProjectionSpace = Tool->WorldTransform;
+	Op->ProjectionBox = Tool->GetProjectionBox();
+	Op->CylinderSplitAngle = Tool->BasicProperties->CylinderSplitAngle;
+	Op->BlendWeight = Tool->BasicProperties->ExpMapNormalBlending;
+	Op->SmoothingRounds = Tool->BasicProperties->ExpMapSmoothingSteps;
+	Op->SmoothingAlpha = Tool->BasicProperties->ExpMapSmoothingAlpha;
+
+	Op->UVRotationAngleDeg = Tool->BasicProperties->Rotation;
+	Op->UVScale = (FVector2f)Tool->BasicProperties->Scale;
+	Op->UVTranslate = (FVector2f)Tool->BasicProperties->Translation;
+
+	Op->OriginalMesh = Tool->InputMesh;
+	Op->TriangleROI = Tool->TriangleROI;
+	Op->UseUVLayer = Tool->UVChannelProperties->GetSelectedChannelIndex(true);
+		
+	Op->SetResultTransform(Tool->WorldTransform);
+
+	return Op;
+}
+
+
+
+
+void UUVProjectionTool::InitializeMesh()
+{
+	InputMesh = MakeShared<FDynamicMesh3, ESPMode::ThreadSafe>();
+	*InputMesh = UE::ToolTarget::GetDynamicMeshCopy(Target);
+	WorldTransform = UE::ToolTarget::GetLocalToWorldTransform(Target);
+
+	// initialize triangle ROI if one exists
+	TriangleROI = MakeShared<TArray<int32>, ESPMode::ThreadSafe>();
+	if (USingleSelectionMeshEditingTool::HasInputSelection())
+	{
+		UE::Geometry::GetStoredSelectionAsTriangles(GetInputSelection(), *InputMesh, *TriangleROI);
+		TriangleROISet.Append(*TriangleROI);
+	}
+	VertexROI = MakeShared<TArray<int32>, ESPMode::ThreadSafe>();
+	UE::Geometry::TriangleToVertexIDs(InputMesh.Get(), *TriangleROI, *VertexROI);
+
+	InputMeshROISpatial = MakeShared<FDynamicMeshAABBTree3, ESPMode::ThreadSafe>(InputMesh.Get(), false);
+	if (TriangleROI->Num() > 0)
+	{
+		InputMeshROISpatial->Build(*TriangleROI);
+	}
+	else
+	{
+		InputMeshROISpatial->Build();
+	}
+
+	WorldBounds = (VertexROI->Num() > 0) ?
+		FAxisAlignedBox3d::MakeBoundsFromIndices(*VertexROI, [this](int32 Index) { return WorldTransform.TransformPosition(InputMesh->GetVertex(Index)); })
+		: FAxisAlignedBox3d::MakeBoundsFromIndices(InputMesh->VertexIndicesItr(), [this](int32 Index) { return WorldTransform.TransformPosition(InputMesh->GetVertex(Index)); });
+
+	BasicProperties->Dimensions = (FVector)WorldBounds.Diagonal();
+}
+
+
+
+FOrientedBox3d UUVProjectionTool::GetProjectionBox() const
+{
+	const FFrame3d BoxFrame(TransformProxy->GetTransform());
+	const FVector3d BoxDimensions = 0.5 * (BasicProperties->bUniformDimensions
+		                                       ? FVector3d(BasicProperties->Dimensions.X, BasicProperties->Dimensions.X, BasicProperties->Dimensions.X)
+		                                       : static_cast<FVector3d>(BasicProperties->Dimensions));
+	return FOrientedBox3d(BoxFrame, BoxDimensions);
+}
+
+
+
+void UUVProjectionTool::Render(IToolsContextRenderAPI* RenderAPI)
+{
+	ProjectionShapeVisualizer.bDepthTested = false;
+	ProjectionShapeVisualizer.BeginFrame(RenderAPI, CameraState);
+
+	UE::Geometry::FOrientedBox3d PrimBox = GetProjectionBox();
+	FVector MinCorner = (FVector)-PrimBox.Extents;
+	FVector MaxCorner = (FVector)PrimBox.Extents;
+	float Width = MaxCorner.X - MinCorner.X;
+	float Depth = MaxCorner.Y - MinCorner.Y;
+	float Height = MaxCorner.Z - MinCorner.Z;
+
+	FTransform UseTransform = PrimBox.Frame.ToFTransform();
+	if (BasicProperties->ProjectionType == EUVProjectionMethod::Cylinder)
+	{
+		UseTransform.SetScale3D(FVector(Width, Depth, 1.0f));
+	}
+	ProjectionShapeVisualizer.SetTransform(UseTransform);
+
+	switch (BasicProperties->ProjectionType)
+	{
+	case EUVProjectionMethod::Box:
+		ProjectionShapeVisualizer.DrawWireBox(FBox(MinCorner, MaxCorner));
+		break;
+	case EUVProjectionMethod::Cylinder:
+		ProjectionShapeVisualizer.DrawWireCylinder(FVector(0, 0, -Height*0.5f), FVector(0, 0, 1), 0.5f, Height, 16);
+		break;
+	case EUVProjectionMethod::Plane:
+		ProjectionShapeVisualizer.DrawSquare(FVector(0, 0, 0), FVector(Width, 0, 0), FVector(0, Depth, 0));
+		break;
+	case EUVProjectionMethod::ExpMap:
+		ProjectionShapeVisualizer.DrawSquare(FVector(0, 0, 0), FVector(Width, 0, 0), FVector(0, Depth, 0));
+		break;
+	}
+	
+	ProjectionShapeVisualizer.EndFrame();
+}
+
+void UUVProjectionTool::OnTick(float DeltaTime)
+{
+	if (bHavePendingAction)
+	{
+		ApplyAction(PendingAction);
+		bHavePendingAction = false;
+		PendingAction = EUVProjectionToolActions::NoAction;
+	}
+
+	Preview->Tick(DeltaTime);
+}
+
+
+
+void UUVProjectionTool::OnMaterialSettingsChanged()
+{
+	MaterialSettings->UpdateMaterials();
+
+	UMaterialInterface* OverrideMaterial = MaterialSettings->GetActiveOverrideMaterial();
+	if (OverrideMaterial != nullptr)
+	{
+		Preview->OverrideMaterial = OverrideMaterial;
+	}
+	else
+	{
+		Preview->OverrideMaterial = nullptr;
+	}
+}
+
+
+void UUVProjectionTool::OnMeshUpdated(UMeshOpPreviewWithBackgroundCompute* PreviewCompute)
+{
+	const FColor UVSeamColor(15, 240, 15);
+	const float UVSeamThickness = 2.0f;
+	const float UVSeamDepthBias = 0.1f * (float)(WorldBounds.DiagonalLength() * 0.01);
+
+	const FDynamicMesh3* UseMesh = Preview->PreviewMesh->GetMesh();
+	const FDynamicMeshUVOverlay* UVOverlay = UseMesh->Attributes()->GetUVLayer(UVChannelProperties->GetSelectedChannelIndex(true));
+	auto AppendSeamEdge = [UseMesh, UVSeamColor, UVSeamThickness, UVSeamDepthBias](int32 eid, TArray<FRenderableLine>& LinesOut) {
+		FVector3d A, B;
+		UseMesh->GetEdgeV(eid, A, B);
+		LinesOut.Add(FRenderableLine((FVector)A, (FVector)B, UVSeamColor, UVSeamThickness, UVSeamDepthBias));
+	};
+	if (TriangleROI->Num() > 0)
+	{
+		EdgeRenderer->CreateOrUpdateLineSet(TEXT("UVSeams"), UseMesh->MaxEdgeID(), [&](int32 eid, TArray<FRenderableLine>& LinesOut) {
+			FIndex2i EdgeT = UseMesh->GetEdgeT(eid);
+			if (UseMesh->IsEdge(eid) && TriangleROISet.Contains(EdgeT.A) && TriangleROISet.Contains(EdgeT.B) && UVOverlay->IsSeamEdge(eid))
+			{
+				AppendSeamEdge(eid, LinesOut);
+			}
+		}, 1);
+	}
+	else
+	{
+		EdgeRenderer->CreateOrUpdateLineSet(TEXT("UVSeams"), UseMesh->MaxEdgeID(), [&](int32 eid, TArray<FRenderableLine>& LinesOut) {
+			if (UseMesh->IsEdge(eid) && UVOverlay->IsSeamEdge(eid))
+			{
+				AppendSeamEdge(eid, LinesOut);
+			}
+		}, 1);
+	}
+}
+
+
+void UUVProjectionTool::UpdatePlaneFromClick(const FVector3d& Position, const FVector3d& Normal, bool bTransationOnly)
+{
+	FFrame3d CurrentFrame(TransformProxy->GetTransform());
+	CurrentFrame.Origin = Position;
+	if (bTransationOnly == false)
+	{
+		CurrentFrame.AlignAxis(2, Normal);
+	}
+	TransformGizmo->SetNewGizmoTransform(CurrentFrame.ToFTransform());
+	Preview->InvalidateResult();
+	bTransformModified = true;
+}
+
+
+void UUVProjectionTool::TransformChanged(UTransformProxy* Proxy, FTransform Transform)
+{
+	Preview->InvalidateResult();
+	bTransformModified = true;
+}
+
+
+bool UUVProjectionTool::CanAccept() const
+{
+	if (!Preview->HaveValidResult())
+	{
+		return false;
+	}
+	return Super::CanAccept();
+}
+
+
+
+void UUVProjectionTool::OnInitializationModeChanged()
+{
+	bool bWasTransformModified = bTransformModified;
+	if (!bTransformModified)
+	{
+		ApplyInitializationMode();
+		bTransformModified = bWasTransformModified;
+	}
+}
+
+
+void UUVProjectionTool::ApplyInitializationMode()
+{
+
+	switch (BasicProperties->Initialization)
+	{
+	case EUVProjectionToolInitializationMode::Default:
+		BasicProperties->Dimensions = InitialDimensions;
+		TransformGizmo->SetNewGizmoTransform(InitialTransform);
+		break;
+
+	case EUVProjectionToolInitializationMode::UsePrevious:
+		{
+			bool bHavePrevious = (BasicProperties->SavedDimensions != FVector::ZeroVector);
+			BasicProperties->Dimensions = (bHavePrevious) ? BasicProperties->SavedDimensions : InitialDimensions;
+			TransformGizmo->SetNewGizmoTransform( (bHavePrevious) ? BasicProperties->SavedTransform : InitialTransform);
+		}
+		break;
+
+	case EUVProjectionToolInitializationMode::AutoFit:
+		ApplyAction_AutoFit(false);
+		break;
+
+	case EUVProjectionToolInitializationMode::AutoFitAlign:
+		ApplyAction_AutoFit(true);
+		break;
+	}
+}
+
+
+
+void UUVProjectionTool::ApplyAction(EUVProjectionToolActions ActionType)
+{
+	switch (ActionType)
+	{
+		case EUVProjectionToolActions::AutoFit:
+			ApplyAction_AutoFit(false);
+			break;
+		case EUVProjectionToolActions::AutoFitAlign:
+			ApplyAction_AutoFit(true);
+			break;
+		case EUVProjectionToolActions::Reset:
+			ApplyAction_Reset();
+			break;
+		case EUVProjectionToolActions::NoAction:
+			break;
+	}
+}
+
+
+void UUVProjectionTool::ApplyAction_AutoFit(bool bAlign)
+{
+	// get current transform
+	FFrame3d AlignFrame(TransformGizmo->GetGizmoTransform());
+
+	TArray<FVector3d> Points;
+	if (VertexROI->Num() > 0)
+	{
+		CollectVertexPositions(*InputMesh, *VertexROI, Points);
+	}
+	else
+	{
+		CollectVertexPositions(*InputMesh, InputMesh->VertexIndicesItr(), Points);
+	}
+
+	// auto-compute orientation if desired
+	if (bAlign)
+	{
+		if (BasicProperties->ProjectionType == EUVProjectionMethod::Plane || BasicProperties->ProjectionType == EUVProjectionMethod::Box || BasicProperties->ProjectionType == EUVProjectionMethod::Cylinder)
+		{
+			// compute min-volume box
+			UE::Geometry::FOrientedBox3d MinBoundingBox;
+			FMinVolumeBox3d MinBoxCalc;
+			bool bMinBoxOK = MinBoxCalc.SolveSubsample(Points.Num(), 1000,
+				[&](int32 Index) { return WorldTransform.TransformPosition(Points[Index]); }, false, nullptr);
+			if (bMinBoxOK && MinBoxCalc.IsSolutionAvailable())
+			{
+				MinBoxCalc.GetResult(MinBoundingBox);
+			}
+			else
+			{
+				MinBoundingBox = UE::Geometry::FOrientedBox3d(WorldBounds);
+			}
+
+			AlignFrame = MinBoundingBox.Frame;
+		}
+		else if (BasicProperties->ProjectionType == EUVProjectionMethod::ExpMap)
+		{
+			int32 VertexID;
+			if (TriangleROI->Num() > 0)
+			{
+				FDynamicMeshUVEditor::EstimateGeodesicCenterFrameVertex(*InputMesh, *TriangleROI, AlignFrame, VertexID);
+				AlignFrame.Transform(WorldTransform);
+			}
+			else
+			{
+				FDynamicMeshUVEditor::EstimateGeodesicCenterFrameVertex(*InputMesh, AlignFrame, VertexID);
+				AlignFrame.Transform(WorldTransform);
+			}
+		}
+	}
+
+	// fit bounds to current frame
+	FAxisAlignedBox3d FitBounds = FAxisAlignedBox3d::MakeBoundsFromIndices(Points.Num(), [this, &Points, &AlignFrame](int32 Index) 
+	{ 
+		return AlignFrame.ToFramePoint(WorldTransform.TransformPosition(Points[Index])); }
+	);
+	if (BasicProperties->ProjectionType == EUVProjectionMethod::Plane || BasicProperties->ProjectionType == EUVProjectionMethod::Box || BasicProperties->ProjectionType == EUVProjectionMethod::Cylinder)
+	{
+		FVector3d ShiftCenter = AlignFrame.FromFramePoint(FitBounds.Center());
+		AlignFrame.Origin = ShiftCenter;
+	}
+	BasicProperties->Dimensions = 2.0f * (FVector)FitBounds.Extents();
+	if (DimensionsWatcher >= 0)
+	{
+		BasicProperties->SilentUpdateWatcherAtIndex(DimensionsWatcher);
+	}
+	if (DimensionsModeWatcher >= 0)
+	{
+		BasicProperties->SilentUpdateWatcherAtIndex(DimensionsModeWatcher);
+	}
+
+	// update Gizmo
+	TransformGizmo->SetNewGizmoTransform(AlignFrame.ToFTransform());
+}
+
+
+
+void UUVProjectionTool::ApplyAction_Reset()
+{
+	bool bWasTransformModified = bTransformModified;
+	ApplyInitializationMode();
+	bTransformModified = bWasTransformModified;
+}
+
+
+#undef LOCTEXT_NAMESPACE
diff --git a/Engine/Plugins/Runtime/MeshModelingToolset/Source/MeshModelingTools/Public/UVProjectionTool.h b/Engine/Plugins/Runtime/MeshModelingToolset/Source/MeshModelingTools/Public/UVProjectionTool.h
new file mode 100644
--- /dev/null
+++ b/Engine/Plugins/Runtime/MeshModelingToolset/Source/MeshModelingTools/Public/UVProjectionTool.h
@@ -0,0 +1,310 @@
+// Copyright Epic Games, Inc. All Rights Reserved.
+
+#pragma once
+
+#include "CoreMinimal.h"
+#include "BaseTools/SingleSelectionMeshEditingTool.h"
+#include "MeshOpPreviewHelpers.h"
+#include "ToolDataVisualizer.h"
+#include "ParameterizationOps/UVProjectionOp.h"
+#include "Properties/MeshMaterialProperties.h"
+#include "Properties/MeshUVChannelProperties.h"
+#include "Drawing/PreviewGeometryActor.h"
+#include "Selection/SelectClickedAction.h"
+#include "OrientedBoxTypes.h"
+
+#include "UVProjectionTool.generated.h"
+
+
+// Forward declarations
+struct FMeshDescription;
+class UDynamicMeshComponent;
+class UCombinedTransformGizmo;
+class UTransformProxy;
+class USingleClickInputBehavior;
+class UUVProjectionTool;
+
+/**
+ *
+ */
+UCLASS()
+class MESHMODELINGTOOLS_API UUVProjectionToolBuilder : public USingleSelectionMeshEditingToolBuilder
+{
+	GENERATED_BODY()
+public:
+	virtual USingleSelectionMeshEditingTool* CreateNewTool(const FToolBuilderState& SceneState) const override;
+
+	virtual bool WantsInputSelectionIfAvailable() const override { return true; }
+};
+
+
+UENUM()
+enum class EUVProjectionToolActions
+{
+	NoAction,
+	AutoFit,
+	AutoFitAlign,
+	Reset
+};
+
+
+UCLASS()
+class MESHMODELINGTOOLS_API UUVProjectionToolEditActions : public UInteractiveToolPropertySet
+{
+	GENERATED_BODY()
+
+public:
+	TWeakObjectPtr<UUVProjectionTool> ParentTool;
+
+	void Initialize(UUVProjectionTool* ParentToolIn) { ParentTool = ParentToolIn; }
+	void PostAction(EUVProjectionToolActions Action);
+
+	/** Automatically fit the UV Projection Dimensions based on the current projection orientation */
+	UFUNCTION(CallInEditor, Category = Actions, meta = (DisplayName = "AutoFit", DisplayPriority = 1))
+	void AutoFit()
+	{
+		PostAction(EUVProjectionToolActions::AutoFit);
+	}
+
+	/** Automatically orient the projection and then automatically fit the UV Projection Dimensions */
+	UFUNCTION(CallInEditor, Category = Actions, meta = (DisplayName = "AutoFitAlign", DisplayPriority = 2))
+	void AutoFitAlign()
+	{
+		PostAction(EUVProjectionToolActions::AutoFitAlign);
+	}
+
+	/** Re-initialize the projection based on the UV Projection Initialization property */
+	UFUNCTION(CallInEditor, Category = Actions, meta = (DisplayName = "Reset", DisplayPriority = 3))
+	void Reset()
+	{
+		PostAction(EUVProjectionToolActions::Reset);
+	}
+};
+
+
+UENUM()
+enum class EUVProjectionToolInitializationMode
+{
+	/** Initialize projection to bounding box center */
+	Default,
+	/** Initialize projection based on previous usage of the Project tool */
+	UsePrevious,
+	/** Initialize projection using Auto Fitting for the initial projection type */
+	AutoFit,
+	/** Initialize projection using Auto Fitting with Alignment for the initial projection type */
+	AutoFitAlign
+};
+
+
+/**
+ * Standard properties
+ */
+UCLASS()
+class MESHMODELINGTOOLS_API UUVProjectionToolProperties : public UInteractiveToolPropertySet
+{
+	GENERATED_BODY()
+
+public:
+	/** Shape and/or algorithm to use for UV projection */
+	UPROPERTY(EditAnywhere, Category = "UV Projection")
+	EUVProjectionMethod ProjectionType = EUVProjectionMethod::Plane;
+
+	/** Width, length, and height of the projection shape before rotation */
+	UPROPERTY(EditAnywhere, Category = "UV Projection")
+	FVector Dimensions = FVector(100.0f, 100.0f, 100.0f);
+
+	/** Only use the Dimensions X value to uniformly define all projection shape dimensions */
+	UPROPERTY(EditAnywhere, Category = "UV Projection")
+	bool bUniformDimensions = false;
+
+	/** Determines how projection settings will be initialized; this only takes effect if the projection shape dimensions or position are unchanged */
+	UPROPERTY(EditAnywhere, Category = "UV Projection")
+	EUVProjectionToolInitializationMode Initialization = EUVProjectionToolInitializationMode::Default;
+
+	//
+	// Cylinder projection options
+	//
+
+	/** Angle in degrees to determine whether faces should be assigned to the cylinder or the flat end caps */
+	UPROPERTY(EditAnywhere, Category = CylinderProjection, meta = (DisplayName = "Split Angle", UIMin = "0", UIMax = "90",
+		EditCondition = "ProjectionType == EUVProjectionMethod::Cylinder", EditConditionHides))
+	float CylinderSplitAngle = 45.0f;
+
+	//
+	// ExpMap projection options
+	//
+
+	/** Blend between surface normals and projection normal; ExpMap projection becomes Plane projection when this value is 1 */
+	UPROPERTY(EditAnywhere, Category = "ExpMap Projection", meta = (DisplayName = "Normal Blending", UIMin = "0", UIMax = "1",
+		EditCondition = "ProjectionType == EUVProjectionMethod::ExpMap", EditConditionHides))
+	float ExpMapNormalBlending = 0.0f;
+
+	/** Number of smoothing steps to apply; this slightly increases distortion but produces more stable results. */
+	UPROPERTY(EditAnywhere, Category = "ExpMap Projection", meta = (DisplayName = "Smoothing Steps", UIMin = "0", UIMax = "100",
+		EditCondition = "ProjectionType == EUVProjectionMethod::ExpMap", EditConditionHides))
+	int ExpMapSmoothingSteps = 0;
+
+	/** Smoothing parameter, larger values result in faster smoothing in each step. */
+	UPROPERTY(EditAnywhere, Category = "ExpMap Projection", meta = (DisplayName = "Smoothing Alpha", UIMin = "0", UIMax = "1",
+		EditCondition = "ProjectionType == EUVProjectionMethod::ExpMap", EditConditionHides))
+	float ExpMapSmoothingAlpha = 0.25f;
+
+	//
+	// UV-space transform options
+	//
+
+	/** Rotation in degrees applied after computing projection */
+	UPROPERTY(EditAnywhere, Category = "UV Transform")
+	float Rotation = 0.0;
+
+	/** Scaling applied after computing projection */
+	UPROPERTY(EditAnywhere, Category = "UV Transform")
+	FVector2D Scale = FVector2D::UnitVector;
+
+	/** Translation applied after computing projection */
+	UPROPERTY(EditAnywhere, Category = "UV Transform")
+	FVector2D Translation = FVector2D::ZeroVector;
+
+	//
+	// Saved State. These are used internally to support UsePrevious initialization mode
+	//
+
+	UPROPERTY()
+	FVector SavedDimensions = FVector::ZeroVector;
+
+	UPROPERTY()
+	bool bSavedUniformDimensions = false;
+
+	UPROPERTY()
+	FTransform SavedTransform;
+};
+
+
+/**
+ * Factory with enough info to spawn the background-thread Operator to do a chunk of work for the tool
+ *  stores a pointer to the tool and enough info to know which specific operator it should spawn
+ */
+UCLASS()
+class MESHMODELINGTOOLS_API UUVProjectionOperatorFactory : public UObject, public UE::Geometry::IDynamicMeshOperatorFactory
+{
+	GENERATED_BODY()
+public:
+	// IDynamicMeshOperatorFactory API
+	virtual TUniquePtr<UE::Geometry::FDynamicMeshOperator> MakeNewOperator() override;
+
+	UPROPERTY()
+	TObjectPtr<UUVProjectionTool> Tool;
+};
+
+
+/**
+ * UV projection tool
+ */
+UCLASS()
+class MESHMODELINGTOOLS_API UUVProjectionTool : public USingleSelectionMeshEditingTool
+{
+	GENERATED_BODY()
+
+public:
+
+	friend UUVProjectionOperatorFactory;
+
+	virtual void Setup() override;
+	virtual void Shutdown(EToolShutdownType ShutdownType) override;
+
+	virtual void OnTick(float DeltaTime) override;
+	virtual void Render(IToolsContextRenderAPI* RenderAPI) override;
+
+	virtual bool HasCancel() const override { return true; }
+	virtual bool HasAccept() const override { return true; }
+	virtual bool CanAccept() const override;
+
+	virtual void RequestAction(EUVProjectionToolActions ActionType);
+
+protected:
+
+	UPROPERTY()
+	TObjectPtr<UMeshUVChannelProperties> UVChannelProperties = nullptr;
+
+	UPROPERTY()
+	TObjectPtr<UUVProjectionToolProperties> BasicProperties = nullptr;
+
+	UPROPERTY()
+	TObjectPtr<UUVProjectionToolEditActions> EditActions = nullptr;
+
+	UPROPERTY()
+	TObjectPtr<UExistingMeshMaterialProperties> MaterialSettings = nullptr;
+
+	UPROPERTY()
+	TObjectPtr<UMeshOpPreviewWithBackgroundCompute> Preview = nullptr;
+
+	UPROPERTY()
+	TObjectPtr<UMaterialInstanceDynamic> CheckerMaterial = nullptr;
+
+	UPROPERTY()
+	TObjectPtr<UCombinedTransformGizmo> TransformGizmo = nullptr;
+	
+	UPROPERTY()
+	TObjectPtr<UTransformProxy> TransformProxy = nullptr;
+
+	UPROPERTY()
+	TObjectPtr<UUVProjectionOperatorFactory> OperatorFactory = nullptr;
+
+	UPROPERTY()
+	TObjectPtr<UPreviewGeometry> EdgeRenderer = nullptr;
+
+	TSharedPtr<UE::Geometry::FDynamicMesh3, ESPMode::ThreadSafe> InputMesh;
+	TSharedPtr<TArray<int32>, ESPMode::ThreadSafe> TriangleROI;
+	TSharedPtr<TArray<int32>, ESPMode::ThreadSafe> VertexROI;
+	TSharedPtr<UE::Geometry::FDynamicMeshAABBTree3, ESPMode::ThreadSafe> InputMeshROISpatial;
+	TSet<int32> TriangleROISet;
+
+	UE::Geometry::FTransform3d WorldTransform;
+	UE::Geometry::FAxisAlignedBox3d WorldBounds;
+
+	FVector InitialDimensions;
+	bool bInitialUniformDimensions;
+	FTransform InitialTransform;
+	int32 DimensionsWatcher = -1;
+	int32 DimensionsModeWatcher = -1;
+	bool bTransformModified = false;
+	void OnInitializationModeChanged();
+	void ApplyInitializationMode();
+
+	FViewCameraState CameraState;
+
+	FToolDataVisualizer ProjectionShapeVisualizer;
+
+	void InitializeMesh();
+	void UpdateNumPreviews();
+
+	void TransformChanged(UTransformProxy* Proxy, FTransform Transform);
+
+	void OnMaterialSettingsChanged();
+	void OnMeshUpdated(UMeshOpPreviewWithBackgroundCompute* PreviewCompute);
+
+	UE::Geometry::FOrientedBox3d GetProjectionBox() const;
+
+
+	//
+	// Support for ctrl+click to set plane from hit point
+	//
+
+	TUniquePtr<FSelectClickedAction> SetPlaneCtrlClickBehaviorTarget;
+
+	UPROPERTY()
+	TObjectPtr<USingleClickInputBehavior> ClickToSetPlaneBehavior;
+
+	void UpdatePlaneFromClick(const FVector3d& Position, const FVector3d& Normal, bool bTransationOnly);
+
+	//
+	// Support for Action Buttons
+	//
+
+	bool bHavePendingAction = false;
+	EUVProjectionToolActions PendingAction;
+	virtual void ApplyAction(EUVProjectionToolActions ActionType);
+	void ApplyAction_AutoFit(bool bAlign);
+	void ApplyAction_Reset();
+
+};
