MeshModelingToolset: Move Layout tool out of Experimental

diff --git a/Engine/Plugins/Experimental/MeshModelingToolsetExp/Source/MeshModelingToolsExp/Private/UVLayoutTool.cpp b/Engine/Plugins/Experimental/MeshModelingToolsetExp/Source/MeshModelingToolsExp/Private/UVLayoutTool.cpp
deleted file mode 100644
diff --git a/Engine/Plugins/Experimental/MeshModelingToolsetExp/Source/MeshModelingToolsExp/Public/UVLayoutTool.h b/Engine/Plugins/Experimental/MeshModelingToolsetExp/Source/MeshModelingToolsExp/Public/UVLayoutTool.h
deleted file mode 100644
diff --git a/Engine/Plugins/Runtime/MeshModelingToolset/Source/MeshModelingTools/Private/UVLayoutTool.cpp b/Engine/Plugins/Runtime/MeshModelingToolset/Source/MeshModelingTools/Private/UVLayoutTool.cpp
new file mode 100644
--- /dev/null
+++ b/Engine/Plugins/Runtime/MeshModelingToolset/Source/MeshModelingTools/Private/UVLayoutTool.cpp
@@ -0,0 +1,304 @@
+// Copyright Epic Games, Inc. All Rights Reserved.
+
+#include "UVLayoutTool.h"
+#include "InteractiveToolManager.h"
+#include "ToolBuilderUtil.h"
+
+#include "ToolSetupUtil.h"
+#include "ModelingToolTargetUtil.h"
+
+#include "DynamicMesh/DynamicMesh3.h"
+
+#include "MeshDescriptionToDynamicMesh.h"
+#include "DynamicMeshToMeshDescription.h"
+
+#include "ParameterizationOps/UVLayoutOp.h"
+#include "Properties/UVLayoutProperties.h"
+
+#include "TargetInterfaces/MaterialProvider.h"
+#include "TargetInterfaces/PrimitiveComponentBackedTarget.h"
+#include "ModelingToolTargetUtil.h"
+
+using namespace UE::Geometry;
+
+#define LOCTEXT_NAMESPACE "UUVLayoutTool"
+
+/*
+ * ToolBuilder
+ */
+
+UMultiSelectionMeshEditingTool* UUVLayoutToolBuilder::CreateNewTool(const FToolBuilderState& SceneState) const
+{
+	return NewObject<UUVLayoutTool>(SceneState.ToolManager);
+}
+
+
+
+/*
+ * Tool
+ */
+
+UUVLayoutTool::UUVLayoutTool()
+{
+}
+
+void UUVLayoutTool::Setup()
+{
+	UInteractiveTool::Setup();
+
+	// hide input StaticMeshComponent
+	for (int32 ComponentIdx = 0; ComponentIdx < Targets.Num(); ComponentIdx++)
+	{
+		UE::ToolTarget::HideSourceObject(Targets[ComponentIdx]);
+	}
+
+	// if we only have one object, add ability to set UV channel
+	if (Targets.Num() == 1)
+	{
+		UVChannelProperties = NewObject<UMeshUVChannelProperties>(this);
+		UVChannelProperties->RestoreProperties(this);
+		UVChannelProperties->Initialize(UE::ToolTarget::GetMeshDescription(Targets[0]), false);
+		UVChannelProperties->ValidateSelection(true);
+		AddToolPropertySource(UVChannelProperties);
+		UVChannelProperties->WatchProperty(UVChannelProperties->UVChannel, [this](const FString& NewValue)
+		{
+			MaterialSettings->UpdateUVChannels(UVChannelProperties->UVChannelNamesList.IndexOfByKey(UVChannelProperties->UVChannel),
+			                                   UVChannelProperties->UVChannelNamesList);
+			UpdateVisualization();
+		});
+	}
+
+	BasicProperties = NewObject<UUVLayoutProperties>(this);
+	BasicProperties->RestoreProperties(this);
+	AddToolPropertySource(BasicProperties);
+
+	MaterialSettings = NewObject<UExistingMeshMaterialProperties>(this);
+	MaterialSettings->RestoreProperties(this);
+	AddToolPropertySource(MaterialSettings);
+
+	// if we only have one object, add optional UV layout view
+	if (Targets.Num() == 1)
+	{
+		UVLayoutView = NewObject<UUVLayoutPreview>(this);
+		UVLayoutView->CreateInWorld(TargetWorld);
+
+		const FComponentMaterialSet MaterialSet = UE::ToolTarget::GetMaterialSet(Targets[0]);
+		UVLayoutView->SetSourceMaterials(MaterialSet);
+
+		const AActor* Actor = UE::ToolTarget::GetTargetActor(Targets[0]);
+		UVLayoutView->SetSourceWorldPosition(
+			Actor->GetTransform(),
+			Actor->GetComponentsBoundingBox());
+
+		UVLayoutView->Settings->RestoreProperties(this);
+		AddToolPropertySource(UVLayoutView->Settings);
+	}
+
+	UpdateVisualization();
+
+	SetToolDisplayName(LOCTEXT("ToolName", "UV Layout"));
+	GetToolManager()->DisplayMessage(LOCTEXT("OnStartUVLayoutTool", "Transform/Rotate/Scale existing UV Charts using various strategies"),
+		EToolMessageLevel::UserNotification);
+}
+
+
+void UUVLayoutTool::UpdateNumPreviews()
+{
+	const int32 CurrentNumPreview = Previews.Num();
+	const int32 TargetNumPreview = Targets.Num();
+	if (TargetNumPreview < CurrentNumPreview)
+	{
+		for (int32 PreviewIdx = CurrentNumPreview - 1; PreviewIdx >= TargetNumPreview; PreviewIdx--)
+		{
+			Previews[PreviewIdx]->Cancel();
+		}
+		Previews.SetNum(TargetNumPreview);
+		OriginalDynamicMeshes.SetNum(TargetNumPreview);
+		Factories.SetNum(TargetNumPreview);
+	}
+	else
+	{
+		OriginalDynamicMeshes.SetNum(TargetNumPreview);
+		Factories.SetNum(TargetNumPreview);
+		for (int32 PreviewIdx = CurrentNumPreview; PreviewIdx < TargetNumPreview; PreviewIdx++)
+		{
+			OriginalDynamicMeshes[PreviewIdx] = MakeShared<FDynamicMesh3, ESPMode::ThreadSafe>();
+			FMeshDescriptionToDynamicMesh Converter;
+			Converter.Convert(UE::ToolTarget::GetMeshDescription(Targets[PreviewIdx]), *OriginalDynamicMeshes[PreviewIdx]);
+
+			Factories[PreviewIdx]= NewObject<UUVLayoutOperatorFactory>();
+			Factories[PreviewIdx]->OriginalMesh = OriginalDynamicMeshes[PreviewIdx];
+			Factories[PreviewIdx]->Settings = BasicProperties;
+			Factories[PreviewIdx]->TargetTransform = (FTransform) UE::ToolTarget::GetLocalToWorldTransform(Targets[PreviewIdx]);
+			Factories[PreviewIdx]->GetSelectedUVChannel = [this]() { return GetSelectedUVChannel(); };
+
+			UMeshOpPreviewWithBackgroundCompute* Preview = Previews.Add_GetRef(NewObject<UMeshOpPreviewWithBackgroundCompute>(Factories[PreviewIdx], "Preview"));
+			Preview->Setup(this->TargetWorld, Factories[PreviewIdx]);
+			ToolSetupUtil::ApplyRenderingConfigurationToPreview(Preview->PreviewMesh, Targets[PreviewIdx]); 
+
+			const FComponentMaterialSet MaterialSet = UE::ToolTarget::GetMaterialSet(Targets[PreviewIdx]);
+			Preview->ConfigureMaterials(MaterialSet.Materials,
+				ToolSetupUtil::GetDefaultWorkingMaterial(GetToolManager())
+			);
+			Preview->PreviewMesh->UpdatePreview(OriginalDynamicMeshes[PreviewIdx].Get());
+			Preview->PreviewMesh->SetTransform((FTransform) UE::ToolTarget::GetLocalToWorldTransform(Targets[PreviewIdx]));
+
+			Preview->OnMeshUpdated.AddLambda([this](UMeshOpPreviewWithBackgroundCompute* Compute)
+			{
+				OnPreviewMeshUpdated(Compute);
+			});
+
+			Preview->SetVisibility(true);
+		}
+	}
+}
+
+
+void UUVLayoutTool::Shutdown(EToolShutdownType ShutdownType)
+{
+	if (UVLayoutView)
+	{
+		UVLayoutView->Settings->SaveProperties(this);
+		UVLayoutView->Disconnect();
+	}
+
+	BasicProperties->SaveProperties(this);
+	MaterialSettings->SaveProperties(this);
+
+	// Restore (unhide) the source meshes
+	for (int32 ComponentIdx = 0; ComponentIdx < Targets.Num(); ComponentIdx++)
+	{
+		UE::ToolTarget::ShowSourceObject(Targets[ComponentIdx]);
+	}
+
+	TArray<FDynamicMeshOpResult> Results;
+	for (UMeshOpPreviewWithBackgroundCompute* Preview : Previews)
+	{
+		Results.Emplace(Preview->Shutdown());
+	}
+	if (ShutdownType == EToolShutdownType::Accept)
+	{
+		GenerateAsset(Results);
+	}
+	for (int32 TargetIndex = 0; TargetIndex < Targets.Num(); ++TargetIndex)
+	{
+		Factories[TargetIndex] = nullptr;
+	}
+}
+
+int32 UUVLayoutTool::GetSelectedUVChannel() const
+{
+	return UVChannelProperties ? UVChannelProperties->GetSelectedChannelIndex(true) : 0;
+}
+
+
+void UUVLayoutTool::Render(IToolsContextRenderAPI* RenderAPI)
+{
+	GetToolManager()->GetContextQueriesAPI()->GetCurrentViewState(CameraState);
+
+	if (UVLayoutView)
+	{
+		UVLayoutView->Render(RenderAPI);
+	}
+
+}
+
+void UUVLayoutTool::OnTick(float DeltaTime)
+{
+	for (UMeshOpPreviewWithBackgroundCompute* Preview : Previews)
+	{
+		Preview->Tick(DeltaTime);
+	}
+
+	if (UVLayoutView)
+	{
+		UVLayoutView->OnTick(DeltaTime);
+	}
+
+
+}
+
+
+
+void UUVLayoutTool::OnPropertyModified(UObject* PropertySet, FProperty* Property)
+{
+	if (PropertySet == BasicProperties || PropertySet == UVChannelProperties)
+	{
+		UpdateNumPreviews();
+		for (UMeshOpPreviewWithBackgroundCompute* Preview : Previews)
+		{
+			Preview->InvalidateResult();
+		}
+	}
+	else if (PropertySet == MaterialSettings)
+	{
+		// if we don't know what changed, or we know checker density changed, update checker material
+		UpdateVisualization();
+	}
+}
+
+
+void UUVLayoutTool::OnPreviewMeshUpdated(UMeshOpPreviewWithBackgroundCompute* Compute)
+{
+	if (UVLayoutView)
+	{
+		FDynamicMesh3 ResultMesh;
+		if (Compute->GetCurrentResultCopy(ResultMesh, false) == false)
+		{
+			return;
+		}
+		UVLayoutView->UpdateUVMesh(&ResultMesh, GetSelectedUVChannel());
+	}
+
+}
+
+
+void UUVLayoutTool::UpdateVisualization()
+{
+	MaterialSettings->UpdateMaterials();
+	UpdateNumPreviews();
+	for (int PreviewIdx = 0; PreviewIdx < Previews.Num(); PreviewIdx++)
+	{
+		UMeshOpPreviewWithBackgroundCompute* Preview = Previews[PreviewIdx];
+		Preview->OverrideMaterial = MaterialSettings->GetActiveOverrideMaterial();
+	}
+
+	for (UMeshOpPreviewWithBackgroundCompute* Preview : Previews)
+	{
+		Preview->InvalidateResult();
+	}
+}
+
+bool UUVLayoutTool::CanAccept() const
+{
+	for (const UMeshOpPreviewWithBackgroundCompute* Preview : Previews)
+	{
+		if (!Preview->HaveValidResult())
+		{
+			return false;
+		}
+	}
+	return Super::CanAccept();
+}
+
+
+void UUVLayoutTool::GenerateAsset(const TArray<FDynamicMeshOpResult>& Results)
+{
+	GetToolManager()->BeginUndoTransaction(LOCTEXT("UVLayoutToolTransactionName", "UV Layout Tool"));
+
+	check(Results.Num() == Targets.Num());
+	
+	for (int32 ComponentIdx = 0; ComponentIdx < Targets.Num(); ComponentIdx++)
+	{
+		const FDynamicMesh3* DynamicMesh = Results[ComponentIdx].Mesh.Get();
+		check(DynamicMesh != nullptr);
+		UE::ToolTarget::CommitDynamicMeshUVUpdate(Targets[ComponentIdx], DynamicMesh);
+	}
+
+	GetToolManager()->EndUndoTransaction();
+}
+
+
+
+
+#undef LOCTEXT_NAMESPACE
diff --git a/Engine/Plugins/Runtime/MeshModelingToolset/Source/MeshModelingTools/Public/UVLayoutTool.h b/Engine/Plugins/Runtime/MeshModelingToolset/Source/MeshModelingTools/Public/UVLayoutTool.h
new file mode 100644
--- /dev/null
+++ b/Engine/Plugins/Runtime/MeshModelingToolset/Source/MeshModelingTools/Public/UVLayoutTool.h
@@ -0,0 +1,97 @@
+// Copyright Epic Games, Inc. All Rights Reserved.
+
+#pragma once
+
+#include "CoreMinimal.h"
+#include "UObject/NoExportTypes.h"
+#include "BaseTools/MultiSelectionMeshEditingTool.h"
+#include "GeometryBase.h"
+#include "InteractiveToolBuilder.h"
+#include "MeshOpPreviewHelpers.h"
+#include "ToolDataVisualizer.h"
+#include "DynamicMesh/DynamicMesh3.h"
+#include "Properties/MeshMaterialProperties.h"
+#include "Properties/MeshUVChannelProperties.h"
+#include "Drawing/UVLayoutPreview.h"
+
+#include "UVLayoutTool.generated.h"
+
+
+// Forward declarations
+struct FMeshDescription;
+class UDynamicMeshComponent;
+class UUVLayoutProperties;
+class UUVLayoutOperatorFactory;
+PREDECLARE_GEOMETRY(class FDynamicMesh3);
+
+/**
+ *
+ */
+UCLASS()
+class MESHMODELINGTOOLS_API UUVLayoutToolBuilder : public UMultiSelectionMeshEditingToolBuilder
+{
+	GENERATED_BODY()
+
+public:
+	virtual UMultiSelectionMeshEditingTool* CreateNewTool(const FToolBuilderState& SceneState) const override;
+};
+
+
+/**
+ * The level editor version of the UV layout tool.
+ */
+UCLASS()
+class MESHMODELINGTOOLS_API UUVLayoutTool : public UMultiSelectionMeshEditingTool
+{
+	GENERATED_BODY()
+
+public:
+
+	UUVLayoutTool();
+
+	virtual void Setup() override;
+	virtual void Shutdown(EToolShutdownType ShutdownType) override;
+
+	virtual void OnTick(float DeltaTime) override;
+	virtual void Render(IToolsContextRenderAPI* RenderAPI) override;
+
+	virtual bool HasCancel() const override { return true; }
+	virtual bool HasAccept() const override { return true; }
+	virtual bool CanAccept() const override;
+
+	virtual void OnPropertyModified(UObject* PropertySet, FProperty* Property) override;
+
+	int32 GetSelectedUVChannel() const;
+
+protected:
+
+	UPROPERTY()
+	TObjectPtr<UMeshUVChannelProperties> UVChannelProperties = nullptr;
+
+	UPROPERTY()
+	TObjectPtr<UUVLayoutProperties> BasicProperties = nullptr;
+
+	UPROPERTY()
+	TObjectPtr<UExistingMeshMaterialProperties> MaterialSettings = nullptr;
+
+	UPROPERTY()
+	TArray<TObjectPtr<UMeshOpPreviewWithBackgroundCompute>> Previews;
+
+	UPROPERTY()
+	TArray<TObjectPtr<UUVLayoutOperatorFactory>> Factories;
+
+	TArray<TSharedPtr<UE::Geometry::FDynamicMesh3, ESPMode::ThreadSafe>> OriginalDynamicMeshes;
+
+	FViewCameraState CameraState;
+
+	void UpdateNumPreviews();
+
+	void UpdateVisualization();
+
+	void OnPreviewMeshUpdated(UMeshOpPreviewWithBackgroundCompute* Compute);
+
+	void GenerateAsset(const TArray<FDynamicMeshOpResult>& Results);
+
+	UPROPERTY()
+	TObjectPtr<UUVLayoutPreview> UVLayoutView = nullptr;
+};
