Add missing predeclaration
#rnx

diff --git a/Engine/Plugins/Runtime/MeshModelingToolset/Source/ModelingComponents/Public/ConversionUtils/DynamicMeshViaMeshDescriptionUtil.h b/Engine/Plugins/Runtime/MeshModelingToolset/Source/ModelingComponents/Public/ConversionUtils/DynamicMeshViaMeshDescriptionUtil.h
--- a/Engine/Plugins/Runtime/MeshModelingToolset/Source/ModelingComponents/Public/ConversionUtils/DynamicMeshViaMeshDescriptionUtil.h
+++ b/Engine/Plugins/Runtime/MeshModelingToolset/Source/ModelingComponents/Public/ConversionUtils/DynamicMeshViaMeshDescriptionUtil.h
@@ -11,6 +11,7 @@
 PREDECLARE_GEOMETRY(class FDynamicMesh3);
 class IMeshDescriptionCommitter;
 class IMeshDescriptionProvider;
+struct FMeshDescription;
 
 namespace UE {
 namespace Geometry {
