UV Editor: Add hover preview to Edit mode

#rnx

diff --git a/Engine/Plugins/Experimental/UVEditor/Source/UVEditorTools/Private/UVEditorUXSettings.cpp b/Engine/Plugins/Experimental/UVEditor/Source/UVEditorTools/Private/UVEditorUXSettings.cpp
--- a/Engine/Plugins/Experimental/UVEditor/Source/UVEditorTools/Private/UVEditorUXSettings.cpp
+++ b/Engine/Plugins/Experimental/UVEditor/Source/UVEditorTools/Private/UVEditorUXSettings.cpp
@@ -10,6 +10,8 @@ const float FUVEditorUXSettings::CameraFarPlaneWorldZ(-10.0);
 // sorting order.
 const float FUVEditorUXSettings::CameraNearPlaneProportionZ(8); // Top layer
 const float FUVEditorUXSettings::SewLineDepthOffset(5.0f);
+const float FUVEditorUXSettings::SelectionHoverWireframeDepthBias(6);
+const float FUVEditorUXSettings::SelectionHoverTriangleDepthBias(5);
 const float FUVEditorUXSettings::SelectionWireframeDepthBias(4.0);
 const float FUVEditorUXSettings::SelectionTriangleDepthBias(3.0);
 const float FUVEditorUXSettings::WireframeDepthOffset(2.0);
@@ -28,6 +30,7 @@ const float FUVEditorUXSettings::LivePreviewHighlightDepthOffset(0.5);
 const float FUVEditorUXSettings::UnwrapTriangleOpacity(1.0);
 const float FUVEditorUXSettings::UnwrapTriangleOpacityWithBackground(0.25);
 const float FUVEditorUXSettings::SelectionTriangleOpacity(1.0f);
+const float FUVEditorUXSettings::SelectionHoverTriangleOpacity(1.0f);
 
 // Per Asset Shifts
 const float FUVEditorUXSettings::UnwrapBoundaryHueShift(30);
diff --git a/Engine/Plugins/Experimental/UVEditor/Source/UVEditorTools/Private/UVSelectTool.cpp b/Engine/Plugins/Experimental/UVEditor/Source/UVEditorTools/Private/UVSelectTool.cpp
--- a/Engine/Plugins/Experimental/UVEditor/Source/UVEditorTools/Private/UVSelectTool.cpp
+++ b/Engine/Plugins/Experimental/UVEditor/Source/UVEditorTools/Private/UVSelectTool.cpp
@@ -315,7 +315,9 @@ void UUVSelectTool::Setup()
 	SelectionMechanic->Setup(this);
 	SelectionMechanic->SetWorld(Targets[0]->UnwrapPreview->GetWorld());
 	SelectionMechanic->OnSelectionChanged.AddUObject(this, &UUVSelectTool::OnSelectionChanged);
+	
 	FMeshSelectionMechanicStyle SelectionStyle;
+	
 	SelectionStyle.TriangleColor = FUVEditorUXSettings::SelectionTriangleFillColor;
 	SelectionStyle.LineColor = FUVEditorUXSettings::SelectionTriangleWireframeColor;
 	SelectionStyle.PointColor = FUVEditorUXSettings::SelectionTriangleWireframeColor;
@@ -324,6 +326,15 @@ void UUVSelectTool::Setup()
 	SelectionStyle.PointThickness = FUVEditorUXSettings::SelectionPointThickness;
 	SelectionStyle.LineAndPointDepthBias = FUVEditorUXSettings::SelectionWireframeDepthBias;
 	SelectionStyle.TriangleDepthBias = FUVEditorUXSettings::SelectionTriangleDepthBias;
+	
+	SelectionStyle.HoverLineAndPointDepthBias = FUVEditorUXSettings::SelectionHoverWireframeDepthBias;
+	SelectionStyle.HoverTriangleDepthBias = FUVEditorUXSettings::SelectionHoverTriangleDepthBias;
+	SelectionStyle.HoverPointColor = FUVEditorUXSettings::SelectionHoverTriangleWireframeColor;
+	SelectionStyle.HoverEdgeColor = FUVEditorUXSettings::SelectionHoverTriangleWireframeColor;
+	SelectionStyle.HoverTriangleEdgeColor = FUVEditorUXSettings::SelectionHoverTriangleWireframeColor;
+	SelectionStyle.HoverTriangleFillColor = FUVEditorUXSettings::SelectionHoverTriangleFillColor;
+	SelectionStyle.HoverTriangleOpacity = FUVEditorUXSettings::SelectionHoverTriangleOpacity;
+	
 	SelectionMechanic->SetVisualizationStyle(SelectionStyle);
 		
 	// Make it so that our selection mechanic creates undo/redo transactions that go to a selection
diff --git a/Engine/Plugins/Experimental/UVEditor/Source/UVEditorTools/Public/UVEditorUXSettings.h b/Engine/Plugins/Experimental/UVEditor/Source/UVEditorTools/Public/UVEditorUXSettings.h
--- a/Engine/Plugins/Experimental/UVEditor/Source/UVEditorTools/Public/UVEditorUXSettings.h
+++ b/Engine/Plugins/Experimental/UVEditor/Source/UVEditorTools/Public/UVEditorUXSettings.h
@@ -41,7 +41,7 @@ public:
 	static const FColor SelectionTriangleFillColor;
 	static const FColor SelectionTriangleWireframeColor;
 
-	// Todo: Not used yet, for the upcoming pre-selection hover feature
+	static const float SelectionHoverTriangleOpacity;
 	static const FColor SelectionHoverTriangleFillColor;
 	static const FColor SelectionHoverTriangleWireframeColor;
 
@@ -53,6 +53,8 @@ public:
 	static const float SelectionPointThickness;
 	static const float SelectionWireframeDepthBias;
 	static const float SelectionTriangleDepthBias;
+	static const float SelectionHoverWireframeDepthBias;
+	static const float SelectionHoverTriangleDepthBias;
 
 	// Sew Action styling
 	static const float SewLineHighlightThickness;
diff --git a/Engine/Plugins/Runtime/MeshModelingToolset/Source/ModelingComponents/Private/Selection/MeshSelectionMechanic.cpp b/Engine/Plugins/Runtime/MeshModelingToolset/Source/ModelingComponents/Private/Selection/MeshSelectionMechanic.cpp
--- a/Engine/Plugins/Runtime/MeshModelingToolset/Source/ModelingComponents/Private/Selection/MeshSelectionMechanic.cpp
+++ b/Engine/Plugins/Runtime/MeshModelingToolset/Source/ModelingComponents/Private/Selection/MeshSelectionMechanic.cpp
@@ -4,6 +4,7 @@
 
 #include "BaseBehaviors/BehaviorTargetInterfaces.h"
 #include "BaseBehaviors/SingleClickOrDragBehavior.h"
+#include "BaseBehaviors/MouseHoverBehavior.h"
 #include "Drawing/TriangleSetComponent.h"
 #include "Drawing/LineSetComponent.h"
 #include "Drawing/PointSetComponent.h"
@@ -290,9 +291,92 @@ void UMeshSelectionMechanic::Setup(UInteractiveTool* ParentToolIn)
 	ClickOrDragBehavior->Modifiers.RegisterModifier(ShiftModifierID, FInputDeviceState::IsShiftKeyDown);
 	ClickOrDragBehavior->Modifiers.RegisterModifier(CtrlModifierID, FInputDeviceState::IsCtrlKeyDown);
 	ParentTool->AddInputBehavior(ClickOrDragBehavior);
+	
+	Settings = NewObject<UMeshSelectionMechanicProperties>(ParentToolIn);
+	Settings->RestoreProperties(ParentToolIn);
+	AddToolPropertySource(Settings);
+
+	ULocalMouseHoverBehavior* HoverBehavior = NewObject<ULocalMouseHoverBehavior>();
+	HoverBehavior->Initialize();
+	HoverBehavior->BeginHitTestFunc = [this](const FInputDeviceRay& Pos)
+	{
+		FInputRayHit Hit;
+		if (Settings->bShowHoveredElements)
+		{
+			EMeshSelectionMechanicMode Mode = SelectionMode;
+			if (Mode != EMeshSelectionMechanicMode::Vertex && Mode != EMeshSelectionMechanicMode::Edge)
+			{
+				Mode = EMeshSelectionMechanicMode::Triangle;
+			}
+
+			// We don't bother with the depth since there aren't competing hover behaviors
+			Hit.bHit = RayCast(Pos, Mode).IsEmpty() == false;
+		}
+		else
+		{
+			// Disable hover behavior
+			Hit.bHit = false;
+		}
+		return Hit;
+	};
+	HoverBehavior->OnUpdateHoverFunc = [this](const FInputDeviceRay& Pos)
+	{
+		EMeshSelectionMechanicMode Mode = SelectionMode;
+		if (SelectionMode != EMeshSelectionMechanicMode::Vertex && SelectionMode != EMeshSelectionMechanicMode::Edge)
+		{
+			Mode = EMeshSelectionMechanicMode::Triangle;
+		}
+		
+		HoverPointSet->Clear();
+		HoverLineSet->Clear();
+		HoverTriangleSet->Clear();
+		
+		TSet<int32> Hovered = RayCast(Pos, Mode);
+		for (int32 Id : Hovered)
+		{
+			if (SelectionMode == EMeshSelectionMechanicMode::Vertex)
+			{
+				const FVector& P = CurrentSelection.Mesh->GetVertexRef(Id);
+				const FRenderablePoint PointToRender(P, VisualizationStyle.HoverPointColor, VisualizationStyle.PointThickness);
+				
+				HoverPointSet->AddPoint(PointToRender);
+			}
+			else if (SelectionMode == EMeshSelectionMechanicMode::Edge)
+			{
+				const FIndex2i EdgeVids = CurrentSelection.Mesh->GetEdgeV(Id);
+				const FVector& A = CurrentSelection.Mesh->GetVertexRef(EdgeVids.A);
+				const FVector& B = CurrentSelection.Mesh->GetVertexRef(EdgeVids.B);
+				
+				HoverLineSet->AddLine(A, B, VisualizationStyle.HoverEdgeColor, VisualizationStyle.LineThickness, VisualizationStyle.HoverLineAndPointDepthBias);
+			}
+			else
+			{
+				const FIndex3i Vids = CurrentSelection.Mesh->GetTriangle(Id);
+				const FVector& A = CurrentSelection.Mesh->GetVertex(Vids[0]);
+				const FVector& B = CurrentSelection.Mesh->GetVertex(Vids[1]);
+				const FVector& C = CurrentSelection.Mesh->GetVertex(Vids[2]);
+
+				HoverLineSet->AddLine(A, B, VisualizationStyle.HoverEdgeColor, VisualizationStyle.LineThickness, VisualizationStyle.HoverLineAndPointDepthBias);
+				HoverLineSet->AddLine(B, C, VisualizationStyle.HoverEdgeColor, VisualizationStyle.LineThickness, VisualizationStyle.HoverLineAndPointDepthBias);
+				HoverLineSet->AddLine(C, A, VisualizationStyle.HoverEdgeColor, VisualizationStyle.LineThickness, VisualizationStyle.HoverLineAndPointDepthBias);
+				HoverTriangleSet->AddTriangle(A, B, C, FVector::ZAxisVector, VisualizationStyle.HoverTriangleFillColor, HoverTriangleSetMaterial);
+			}
+
+			return true;
+		}
+		return false;
+	};
+	HoverBehavior->OnEndHoverFunc = [this]()
+	{
+		HoverPointSet->Clear();
+		HoverLineSet->Clear();
+		HoverTriangleSet->Clear();
+	};
+	ParentTool->AddInputBehavior(HoverBehavior);
 
 	ClearCurrentSelection();
 
+	// Setup selection visualizations
 	TriangleSet = NewObject<UTriangleSetComponent>();
 	// We are setting the TranslucencySortPriority here to handle the UV editor's use case in 2D
 	// where multiple translucent layers are drawn on top of each other but still need depth sorting.
@@ -300,13 +384,19 @@ void UMeshSelectionMechanic::Setup(UInteractiveTool* ParentToolIn)
 	TriangleSetMaterial = ToolSetupUtil::GetCustomTwoSidedDepthOffsetMaterial(GetParentTool()->GetToolManager(),
 		VisualizationStyle.TriangleColor, VisualizationStyle.TriangleDepthBias, VisualizationStyle.TriangleOpacity);
 	
+	PointSet = NewObject<UPointSetComponent>();
+	PointSet->SetPointMaterial(ToolSetupUtil::GetDefaultPointComponentMaterial( GetParentTool()->GetToolManager(), true));
 	LineSet = NewObject<ULineSetComponent>();
-	LineSet->SetLineMaterial(ToolSetupUtil::GetDefaultLineComponentMaterial(
-		GetParentTool()->GetToolManager(), /*bDepthTested*/ true));
+	LineSet->SetLineMaterial(ToolSetupUtil::GetDefaultLineComponentMaterial( GetParentTool()->GetToolManager(), true));
 
-	PointSet = NewObject<UPointSetComponent>();
-	PointSet->SetPointMaterial(ToolSetupUtil::GetDefaultPointComponentMaterial(
-		GetParentTool()->GetToolManager(), /*bDepthTested*/ true));
+	// Setup hover visualizations
+	HoverPointSet = NewObject<UPointSetComponent>();
+	HoverPointSet->SetPointMaterial(ToolSetupUtil::GetDefaultPointComponentMaterial(GetParentTool()->GetToolManager(), false));
+	HoverLineSet = NewObject<ULineSetComponent>();
+	HoverLineSet->SetLineMaterial(ToolSetupUtil::GetDefaultLineComponentMaterial(GetParentTool()->GetToolManager(), false));
+	HoverTriangleSet = NewObject<UTriangleSetComponent>();
+	HoverTriangleSetMaterial = ToolSetupUtil::GetCustomTwoSidedDepthOffsetMaterial(GetParentTool()->GetToolManager(),
+		VisualizationStyle.HoverTriangleFillColor, VisualizationStyle.HoverTriangleDepthBias, VisualizationStyle.HoverTriangleOpacity);
 
 	// Add default selection change emitter if one was not provided.
 	if (!EmitSelectionChange)
@@ -322,6 +412,18 @@ void UMeshSelectionMechanic::Setup(UInteractiveTool* ParentToolIn)
 
 void UMeshSelectionMechanic::SetWorld(UWorld* World)
 {
+	auto RegisterComponent = [this](UMeshComponent* Component)
+	{
+		if (Component->IsRegistered())
+		{
+			Component->ReregisterComponent();
+		}
+		else
+		{
+			Component->RegisterComponent();
+		}
+	};
+	
 	// It may be unreasonable to worry about SetWorld being called more than once, but let's be safe anyway
 	if (PreviewGeometryActor)
 	{
@@ -337,36 +439,35 @@ void UMeshSelectionMechanic::SetWorld(UWorld* World)
 	// Attach the rendering component to the actor		
 	TriangleSet->Rename(nullptr, PreviewGeometryActor); // Changes the "outer"
 	PreviewGeometryActor->SetRootComponent(TriangleSet);
-	if (TriangleSet->IsRegistered())
-	{
-		TriangleSet->ReregisterComponent();
-	}
-	else
-	{
-		TriangleSet->RegisterComponent();
-	}
+	RegisterComponent(TriangleSet);
 
 	LineSet->Rename(nullptr, PreviewGeometryActor); // Changes the "outer"
 	LineSet->AttachToComponent(TriangleSet, FAttachmentTransformRules::KeepWorldTransform);
-	if (LineSet->IsRegistered())
-	{
-		LineSet->ReregisterComponent();
-	}
-	else
-	{
-		LineSet->RegisterComponent();
-	}
-
+	RegisterComponent(LineSet);
+	
 	PointSet->Rename(nullptr, PreviewGeometryActor); // Changes the "outer"
 	PointSet->AttachToComponent(TriangleSet, FAttachmentTransformRules::KeepWorldTransform);
-	if (PointSet->IsRegistered())
-	{
-		PointSet->ReregisterComponent();
-	}
-	else
+	RegisterComponent(PointSet);
+
+	if (HoverGeometryActor)
 	{
-		PointSet->RegisterComponent();
+		HoverGeometryActor->Destroy();
 	}
+	
+	HoverGeometryActor = World->SpawnActor<APreviewGeometryActor>(FVector::ZeroVector, Rotation, SpawnInfo);
+	
+	// Attach the rendering component to the actor		
+	HoverTriangleSet->Rename(nullptr, HoverGeometryActor); // Changes the "outer"
+	HoverGeometryActor->SetRootComponent(HoverTriangleSet);
+	RegisterComponent(HoverTriangleSet);
+
+	HoverLineSet->Rename(nullptr, HoverGeometryActor); // Changes the "outer"
+	HoverLineSet->AttachToComponent(HoverTriangleSet, FAttachmentTransformRules::KeepWorldTransform);
+	RegisterComponent(HoverLineSet);
+	
+	HoverPointSet->Rename(nullptr, HoverGeometryActor); // Changes the "outer"
+	HoverPointSet->AttachToComponent(HoverTriangleSet, FAttachmentTransformRules::KeepWorldTransform);
+	RegisterComponent(HoverPointSet);
 }
 
 void UMeshSelectionMechanic::Shutdown()
@@ -404,6 +505,18 @@ void UMeshSelectionMechanic::SetVisualizationStyle(const FMeshSelectionMechanicS
 		TriangleSetMaterial->SetScalarParameterValue(TEXT("Opacity"), VisualizationStyle.TriangleOpacity);
 		TriangleSetMaterial->SetScalarParameterValue(TEXT("PercentDepthOffset"), VisualizationStyle.TriangleDepthBias);
 	}
+	
+	if (HoverTriangleSet) {
+		// We are setting the TranslucencySortPriority here to handle the UV editor's use case in 2D
+		// where multiple translucent layers are drawn on top of each other but still need depth sorting.
+		HoverTriangleSet->TranslucencySortPriority = VisualizationStyle.HoverTriangleDepthBias;
+	}
+	
+	if (HoverTriangleSetMaterial) {
+		HoverTriangleSetMaterial->SetVectorParameterValue(TEXT("Color"), VisualizationStyle.HoverTriangleFillColor);
+		HoverTriangleSetMaterial->SetScalarParameterValue(TEXT("Opacity"), VisualizationStyle.HoverTriangleOpacity);
+		HoverTriangleSetMaterial->SetScalarParameterValue(TEXT("PercentDepthOffset"), VisualizationStyle.HoverTriangleDepthBias);
+	}
 }
 
 void UMeshSelectionMechanic::SetSelection(const FDynamicMeshSelection& Selection, bool bBroadcast, bool bEmitChange)
@@ -891,120 +1004,133 @@ void UMeshSelectionMechanic::OnClicked(const FInputDeviceRay& ClickPos)
 		ClearCurrentSelection();
 	}
 
-	TSet<int32> ClickSelectedIDs; // TODO Maybe using a TVariant would have been better here...
+	TSet<int32> ClickSelectedIDs = RayCast(ClickPos, SelectionMode); // TODO Maybe using a TVariant would have been better here...
+
+	// TODO Perhaps selection clearing should happen only if the click occurs further than some threshold from the meshes
+	UpdateCurrentSelection(ClickSelectedIDs);
+
+	if (OriginalSelection != CurrentSelection)
+	{
+		UpdateCentroid();
+		EmitSelectionChange(OriginalSelection, CurrentSelection);
+		OnSelectionChanged.Broadcast();
+
+		// Rebuild after broadcast in case the outside world wants to adjust things like color...
+		RebuildDrawnElements(FTransform(GetCurrentSelectionCentroid()));
+	}
+}
+
+TSet<int32> UMeshSelectionMechanic::RayCast(const FInputDeviceRay& Pos, EMeshSelectionMechanicMode Mode)
+{
+	// TODO if the selection mode is Vertex, Edge or Triangle we'll return at most one index, if the selection mode is
+	// Component or Mesh we'll return a set of components. Perhaps it would be better to use TVariant here (and maybe
+	// also a TArray rather than a TSet)
+	TSet<int32> Result;
+	
+	int32 HitTid = IndexConstants::InvalidID;
+	for (int32 MeshIndex = 0; MeshIndex < MeshSpatials.Num(); ++MeshIndex)
 	{
-		int32 HitTid = IndexConstants::InvalidID;
-		for (int32 MeshIndex = 0; MeshIndex < MeshSpatials.Num(); ++MeshIndex)
-		{
-			FRay3d LocalRay(
-				(FVector3d)MeshTransforms[MeshIndex].InverseTransformPosition(ClickPos.WorldRay.Origin),
-				(FVector3d)MeshTransforms[MeshIndex].InverseTransformVector(ClickPos.WorldRay.Direction));
+		FRay3d LocalRay(
+			(FVector3d)MeshTransforms[MeshIndex].InverseTransformPosition(Pos.WorldRay.Origin),
+			(FVector3d)MeshTransforms[MeshIndex].InverseTransformVector(Pos.WorldRay.Direction));
 
-			double RayT = 0; 
-			if (MeshSpatials[MeshIndex]->FindNearestHitTriangle(LocalRay, RayT, HitTid))
-			{
-				CurrentSelection.Mesh = MeshSpatials[MeshIndex]->GetMesh();
-				CurrentSelection.TopologyTimestamp = CurrentSelection.Mesh->GetTopologyChangeStamp();
-				CurrentSelectionIndex = MeshIndex;
-				break;
-			}
+		double RayT = 0;
+		if (MeshSpatials[MeshIndex]->FindNearestHitTriangle(LocalRay, RayT, HitTid))
+		{
+			CurrentSelection.Mesh = MeshSpatials[MeshIndex]->GetMesh();
+			CurrentSelection.TopologyTimestamp = CurrentSelection.Mesh->GetTopologyChangeStamp();
+			CurrentSelectionIndex = MeshIndex;
+			break;
 		}
+	}
 
-		if (HitTid != IndexConstants::InvalidID)
+	if (HitTid != IndexConstants::InvalidID)
+	{
+		if (Mode == EMeshSelectionMechanicMode::Component)
 		{
-			if (SelectionMode == EMeshSelectionMechanicMode::Component)
-			{
-				TRACE_CPUPROFILER_EVENT_SCOPE(MeshSelectionMechanic_OnClicked_Component);
-
-				FMeshConnectedComponents MeshSelectedComponent(CurrentSelection.Mesh);
-				TArray<int32> SeedTriangles;
-				SeedTriangles.Add(HitTid);
-				MeshSelectedComponent.FindTrianglesConnectedToSeeds(SeedTriangles);
-				ensure(MeshSelectedComponent.Components.Num() == 1); // Expect each triangle to only be in a single component
-				ClickSelectedIDs.Append(MoveTemp(MeshSelectedComponent.Components[0].Indices));
-			}
-			else if (SelectionMode == EMeshSelectionMechanicMode::Edge)
-			{
-				TRACE_CPUPROFILER_EVENT_SCOPE(MeshSelectionMechanic_OnClicked_Edge);
-				// TODO: We'll need the ability to hit occluded triangles to see if there is a better edge to snap to.
+			TRACE_CPUPROFILER_EVENT_SCOPE(Component);
+
+			FMeshConnectedComponents MeshSelectedComponent(CurrentSelection.Mesh);
+			TArray<int32> SeedTriangles;
+			SeedTriangles.Add(HitTid);
+			MeshSelectedComponent.FindTrianglesConnectedToSeeds(SeedTriangles);
+			ensure(MeshSelectedComponent.Components.Num() == 1); // Expect each triangle to only be in a single component
+			Result.Append(MoveTemp(MeshSelectedComponent.Components[0].Indices));
+		}
+		else if (Mode == EMeshSelectionMechanicMode::Edge)
+		{
+			TRACE_CPUPROFILER_EVENT_SCOPE(Edge);
+			// TODO: We'll need the ability to hit occluded triangles to see if there is a better edge to snap to.
 
-				// Try to snap to one of the edges.
-				FIndex3i Eids = CurrentSelection.Mesh->GetTriEdges(HitTid);
+			// Try to snap to one of the edges.
+			FIndex3i Eids = CurrentSelection.Mesh->GetTriEdges(HitTid);
 
-				FGeometrySet3 GeometrySet;
-				for (int i = 0; i < 3; ++i)
-				{
-					FIndex2i Vids = CurrentSelection.Mesh->GetEdgeV(Eids[i]);
-					FPolyline3d Polyline(CurrentSelection.Mesh->GetVertex(Vids.A), CurrentSelection.Mesh->GetVertex(Vids.B));
-					GeometrySet.AddCurve(Eids[i], Polyline);
-				}
-				FGeometrySet3::FNearest Result;
-				if (GeometrySet.FindNearestCurveToRay((FRay3d)ClickPos.WorldRay, Result, 
-					[this](const FVector3d& Position1, const FVector3d& Position2) {
-						return ToolSceneQueriesUtil::PointSnapQuery(CameraState,
-							Position1, Position2,
-							ToolSceneQueriesUtil::GetDefaultVisualAngleSnapThreshD()); }))
-				{
-					ClickSelectedIDs = TSet<int32>{Result.ID};
-				}
+			FGeometrySet3 GeometrySet;
+			for (int i = 0; i < 3; ++i)
+			{
+				FIndex2i Vids = CurrentSelection.Mesh->GetEdgeV(Eids[i]);
+				FPolyline3d Polyline(CurrentSelection.Mesh->GetVertex(Vids.A), CurrentSelection.Mesh->GetVertex(Vids.B));
+				GeometrySet.AddCurve(Eids[i], Polyline);
 			}
-			else if (SelectionMode == EMeshSelectionMechanicMode::Vertex)
+
+			FRay3d WorldRay((FVector3d)Pos.WorldRay.Origin, (FVector3d)Pos.WorldRay.Direction);
+			FGeometrySet3::FNearest Nearest;
+			if (GeometrySet.FindNearestCurveToRay(WorldRay, Nearest, 
+				[this](const FVector3d& Position1, const FVector3d& Position2) {
+					return ToolSceneQueriesUtil::PointSnapQuery(CameraState,
+						Position1, Position2,
+						ToolSceneQueriesUtil::GetDefaultVisualAngleSnapThreshD()); }))
 			{
-				TRACE_CPUPROFILER_EVENT_SCOPE(MeshSelectionMechanic_OnClicked_Vertex);
-				// TODO: Improve this to handle super narrow, sliver triangles better, where testing near vertices can be difficult.
+				Result = TSet<int32>{Nearest.ID};
+			}
+		}
+		else if (Mode == EMeshSelectionMechanicMode::Vertex)
+		{
+			TRACE_CPUPROFILER_EVENT_SCOPE(Vertex);
+			// TODO: Improve this to handle super narrow, sliver triangles better, where testing near vertices can be difficult.
 
-				// Try to snap to one of the vertices
-				FIndex3i Vids = CurrentSelection.Mesh->GetTriangle(HitTid);
+			// Try to snap to one of the vertices
+			FIndex3i Vids = CurrentSelection.Mesh->GetTriangle(HitTid);
 
-				FGeometrySet3 GeometrySet;
-				for (int i = 0; i < 3; ++i)
-				{
-					GeometrySet.AddPoint(Vids[i], CurrentSelection.Mesh->GetTriVertex(HitTid, i));
-				}
-				FGeometrySet3::FNearest Result;
-				if (GeometrySet.FindNearestPointToRay((FRay3d)ClickPos.WorldRay, Result,
-					[this](const FVector3d& Position1, const FVector3d& Position2) {
-						return ToolSceneQueriesUtil::PointSnapQuery(CameraState,
-							Position1, Position2,
-							ToolSceneQueriesUtil::GetDefaultVisualAngleSnapThreshD()); }))
-				{
-					ClickSelectedIDs = TSet<int32>{Result.ID};
-				}
-			}
-			else if (SelectionMode == EMeshSelectionMechanicMode::Triangle)
+			FGeometrySet3 GeometrySet;
+			for (int i = 0; i < 3; ++i)
 			{
-				TRACE_CPUPROFILER_EVENT_SCOPE(MeshSelectionMechanic_OnClicked_Triangle);
-
-				ClickSelectedIDs = TSet<int32>{HitTid};
+				GeometrySet.AddPoint(Vids[i], CurrentSelection.Mesh->GetTriVertex(HitTid, i));
 			}
-			else if (SelectionMode == EMeshSelectionMechanicMode::Mesh)
+			
+			FRay3d WorldRay((FVector3d)Pos.WorldRay.Origin, (FVector3d)Pos.WorldRay.Direction);
+			FGeometrySet3::FNearest Nearest;
+			if (GeometrySet.FindNearestPointToRay(WorldRay, Nearest,
+				[this](const FVector3d& Position1, const FVector3d& Position2) {
+					return ToolSceneQueriesUtil::PointSnapQuery(CameraState,
+						Position1, Position2,
+						ToolSceneQueriesUtil::GetDefaultVisualAngleSnapThreshD()); }))
 			{
-				TRACE_CPUPROFILER_EVENT_SCOPE(MeshSelectionMechanic_OnClicked_Mesh);
-
-				for (int32 Tid : CurrentSelection.Mesh->TriangleIndicesItr())
-				{
-					ClickSelectedIDs.Add(Tid);
-				}
+				Result = TSet<int32>{Nearest.ID};
 			}
-			else
+		}
+		else if (Mode == EMeshSelectionMechanicMode::Triangle)
+		{
+			TRACE_CPUPROFILER_EVENT_SCOPE(Triangle);
+
+			Result = TSet<int32>{HitTid};
+		}
+		else if (Mode == EMeshSelectionMechanicMode::Mesh)
+		{
+			TRACE_CPUPROFILER_EVENT_SCOPE(Mesh);
+
+			for (int32 Tid : CurrentSelection.Mesh->TriangleIndicesItr())
 			{
-				checkNoEntry();
+				Result.Add(Tid);
 			}
 		}
+		else
+		{
+			checkNoEntry();
+		}
 	}
 
-	// TODO Perhaps selection clearing should happen only if the click occurs further than some threshold from the meshes
-	UpdateCurrentSelection(ClickSelectedIDs);
-
-	if (OriginalSelection != CurrentSelection)
-	{
-		UpdateCentroid();
-		EmitSelectionChange(OriginalSelection, CurrentSelection);
-		OnSelectionChanged.Broadcast();
-
-		// Rebuild after broadcast in case the outside world wants to adjust things like color...
-		RebuildDrawnElements(FTransform(GetCurrentSelectionCentroid()));
-	}
+	return Result;
 }
 
 void UMeshSelectionMechanic::OnDragRectangleStarted()
diff --git a/Engine/Plugins/Runtime/MeshModelingToolset/Source/ModelingComponents/Public/Selection/MeshSelectionMechanic.h b/Engine/Plugins/Runtime/MeshModelingToolset/Source/ModelingComponents/Public/Selection/MeshSelectionMechanic.h
--- a/Engine/Plugins/Runtime/MeshModelingToolset/Source/ModelingComponents/Public/Selection/MeshSelectionMechanic.h
+++ b/Engine/Plugins/Runtime/MeshModelingToolset/Source/ModelingComponents/Public/Selection/MeshSelectionMechanic.h
@@ -42,8 +42,14 @@ class MODELINGCOMPONENTS_API UMeshSelectionMechanicProperties : public UInteract
 	GENERATED_BODY()
 
 public:
+	//~ TODO Make this show up with the AdvancedDisplay settings in UUVSelectToolProoperties? Maybe it needs to move there and get passed it to UMeshSelectionMechanic
+	UPROPERTY(EditAnywhere, Category = Options, AdvancedDisplay)
+	bool bShowHoveredElements = true;
 };
 
+// These values are all overwritten
+// TODO Directly use the values in FUVEditorUXSettings when MeshSelectionMechanic is moved into UV Editor module, we can
+//  also remove the SetVisualizationStyle
 struct MODELINGCOMPONENTS_API FMeshSelectionMechanicStyle
 {
 	FColor TriangleColor = FColor::Yellow;
@@ -52,8 +58,16 @@ struct MODELINGCOMPONENTS_API FMeshSelectionMechanicStyle
 	float TriangleOpacity = 0.3;
 	float LineThickness = 1.5;
 	float PointThickness = 6;
-	float LineAndPointDepthBias = 0.7;
-	float TriangleDepthBias = 0.65;
+	float LineAndPointDepthBias = 4;
+	float TriangleDepthBias = 3;
+	
+	FColor HoverPointColor = FColor::FromHex("0E86FF");
+	FColor HoverEdgeColor = FColor::FromHex("0E86FF");
+	FColor HoverTriangleEdgeColor = FColor::FromHex("0E86FF");
+	FColor HoverTriangleFillColor = FColor::FromHex("4E719B");
+	float HoverTriangleOpacity = 1;
+	float HoverLineAndPointDepthBias = 6;
+	float HoverTriangleDepthBias = 5;
 };
 
 /**
@@ -108,6 +122,8 @@ public:
 
 	// IModifierToggleBehaviorTarget implementation
 	virtual void OnUpdateModifierState(int ModifierID, bool bIsOn) override;
+	
+	TSet<int32> RayCast(const FInputDeviceRay& ClickPos, EMeshSelectionMechanicMode Mode);
 
 	FSimpleMulticastDelegate OnSelectionChanged;
 
@@ -130,6 +146,9 @@ protected:
 	bool ShouldToggleFromSelection() const { return bCtrlToggle && bShiftToggle; }
 	bool ShouldRestartSelection() const { return !bCtrlToggle && !bShiftToggle; }
 
+	UPROPERTY()
+	TObjectPtr<UMeshSelectionMechanicProperties> Settings = nullptr;
+
 	UPROPERTY()
 	TObjectPtr<URectangleMarqueeMechanic> MarqueeMechanic;
 
@@ -138,17 +157,26 @@ protected:
 
 	/** The material being displayed for selected triangles */
 	UPROPERTY()
-	TObjectPtr<UMaterialInstanceDynamic> TriangleSetMaterial;
-
+	TObjectPtr<UMaterialInstanceDynamic> TriangleSetMaterial = nullptr;
 	UPROPERTY()
 	TObjectPtr<UTriangleSetComponent> TriangleSet = nullptr;
-
 	UPROPERTY()
 	TObjectPtr<ULineSetComponent> LineSet = nullptr;
-
 	UPROPERTY()
 	TObjectPtr<UPointSetComponent> PointSet = nullptr;
-
+	
+	UPROPERTY()
+	TObjectPtr<APreviewGeometryActor> HoverGeometryActor = nullptr;
+	
+	UPROPERTY()
+	TObjectPtr<UMaterialInstanceDynamic> HoverTriangleSetMaterial = nullptr;
+	UPROPERTY()
+	TObjectPtr<UTriangleSetComponent> HoverTriangleSet = nullptr;
+	UPROPERTY()
+	TObjectPtr<ULineSetComponent> HoverLineSet = nullptr;
+	UPROPERTY()
+	TObjectPtr<UPointSetComponent> HoverPointSet = nullptr;
+	
 
 	TArray<TSharedPtr<FDynamicMeshAABBTree3>> MeshSpatials;
 	TArray<FTransform> MeshTransforms;
