UV Editor: Move the checkbox to enable/disable rendering of hovered elements into the main details panel

#rnx

diff --git a/Engine/Plugins/Runtime/MeshModelingToolset/Source/ModelingComponents/Public/Selection/MeshSelectionMechanic.h b/Engine/Plugins/Runtime/MeshModelingToolset/Source/ModelingComponents/Public/Selection/MeshSelectionMechanic.h
--- a/Engine/Plugins/Runtime/MeshModelingToolset/Source/ModelingComponents/Public/Selection/MeshSelectionMechanic.h
+++ b/Engine/Plugins/Runtime/MeshModelingToolset/Source/ModelingComponents/Public/Selection/MeshSelectionMechanic.h
@@ -42,8 +42,8 @@ class MODELINGCOMPONENTS_API UMeshSelectionMechanicProperties : public UInteract
 	GENERATED_BODY()
 
 public:
-	//~ TODO Make this show up with the AdvancedDisplay settings in UUVSelectToolProoperties? Maybe it needs to move there and get passed it to UMeshSelectionMechanic
-	UPROPERTY(EditAnywhere, Category = Options, AdvancedDisplay)
+	/** If true, show the currently hovered selectable primitive */
+	UPROPERTY(EditAnywhere, Category = Options)
 	bool bShowHoveredElements = true;
 };
 
