ModelingComponents: register StaticMeshDescription attributes in UE::ToolTarget::CommitMeshDescriptionUpdateViaDynamicMesh for full-update path, and read back a copy of MeshDescription before updating for partial-update path
#rnx

diff --git a/Engine/Plugins/Runtime/MeshModelingToolset/Source/ModelingComponents/Private/ModelingToolTargetUtil.cpp b/Engine/Plugins/Runtime/MeshModelingToolset/Source/ModelingComponents/Private/ModelingToolTargetUtil.cpp
--- a/Engine/Plugins/Runtime/MeshModelingToolset/Source/ModelingComponents/Private/ModelingToolTargetUtil.cpp
+++ b/Engine/Plugins/Runtime/MeshModelingToolset/Source/ModelingComponents/Private/ModelingToolTargetUtil.cpp
@@ -23,6 +23,7 @@
 #include "MeshConversionOptions.h"
 #include "MeshDescriptionToDynamicMesh.h"
 #include "DynamicMeshToMeshDescription.h"
+#include "StaticMeshAttributes.h"
 
 #define LOCTEXT_NAMESPACE "ModelingToolTargetUtil"
 
@@ -265,13 +266,21 @@ UE::ToolTarget::EDynamicMeshUpdateResult UE::ToolTarget::CommitMeshDescriptionUp
 	}
 
 	FMeshDescription ConvertedMesh;
-	FDynamicMeshToMeshDescription Converter;
 	if (bHaveModifiedTopology)
 	{
+		FStaticMeshAttributes Attributes(ConvertedMesh);
+		Attributes.Register();
+		FDynamicMeshToMeshDescription Converter;
 		Converter.Convert(&UpdatedMesh, ConvertedMesh);
 	}
 	else
 	{
+		// This should never need to be true because we are either (1) ignoring the tangents, in which
+		// case, we don't want to update them anyway (2) replacing them, in which case, we don't need
+		// to compute them
+		const bool bComputeAutoGeneratedTangents = false;
+		ConvertedMesh = UE::ToolTarget::GetMeshDescriptionCopy(Target, bComputeAutoGeneratedTangents);
+		FDynamicMeshToMeshDescription Converter;
 		Converter.Update(&UpdatedMesh, ConvertedMesh);
 	}
 
