ModelingTools: convert various tools/mechanics/etc to use new SnappingManager functionality for scene hit-tests and snapping
#rnx

diff --git a/Engine/Plugins/Editor/ModelingToolsEditorMode/Source/ModelingToolsEditorMode/Private/ModelingToolsEditorMode.cpp b/Engine/Plugins/Editor/ModelingToolsEditorMode/Source/ModelingToolsEditorMode/Private/ModelingToolsEditorMode.cpp
--- a/Engine/Plugins/Editor/ModelingToolsEditorMode/Source/ModelingToolsEditorMode/Private/ModelingToolsEditorMode.cpp
+++ b/Engine/Plugins/Editor/ModelingToolsEditorMode/Source/ModelingToolsEditorMode/Private/ModelingToolsEditorMode.cpp
@@ -995,7 +995,7 @@ void UModelingToolsEditorMode::FocusCameraAtCursorHotkey()
 
 	// cast ray against visible objects
 	FHitResult WorldHitResult;
-	if (ToolSceneQueriesUtil::FindNearestVisibleObjectHit(GetWorld(), WorldHitResult, Ray.Origin, Ray.PointAt(HALF_WORLD_MAX)))
+	if (ToolSceneQueriesUtil::FindNearestVisibleObjectHit( USceneSnappingManager::Find(GetToolManager()), WorldHitResult, Ray) )
 	{
 		HitPoint = WorldHitResult.ImpactPoint;
 		NearestHitDist = (double)Ray.GetParameter(HitPoint);
diff --git a/Engine/Plugins/Experimental/MeshModelingToolsetExp/Source/MeshModelingToolsExp/Private/AddPatchTool.cpp b/Engine/Plugins/Experimental/MeshModelingToolsetExp/Source/MeshModelingToolsExp/Private/AddPatchTool.cpp
--- a/Engine/Plugins/Experimental/MeshModelingToolsetExp/Source/MeshModelingToolsExp/Private/AddPatchTool.cpp
+++ b/Engine/Plugins/Experimental/MeshModelingToolsetExp/Source/MeshModelingToolsExp/Private/AddPatchTool.cpp
@@ -172,7 +172,7 @@ void UAddPatchTool::UpdatePreviewPosition(const FInputDeviceRay& DeviceClickPos)
 
 	// cast ray into scene
 	FHitResult Result;
-	bool bHit = ToolSceneQueriesUtil::FindNearestVisibleObjectHit(TargetWorld, Result, ClickPosWorldRay);
+	bool bHit = ToolSceneQueriesUtil::FindNearestVisibleObjectHit(this, Result, ClickPosWorldRay);
 	if (bHit)
 	{
 		ShapeFrame = FFrame3f((FVector3f)Result.ImpactPoint, (FVector3f)Result.ImpactNormal);
@@ -226,7 +226,7 @@ void UAddPatchTool::UpdatePreviewMesh()
 		FVector RayEnd = RayStart; RayEnd.Z = WorldMinHeight;
 
 		FHitResult Result;
-		bool bHit = ToolSceneQueriesUtil::FindNearestVisibleObjectHit(TargetWorld, Result, RayStart, RayEnd);
+		bool bHit = ToolSceneQueriesUtil::FindNearestVisibleObjectHit(this, Result, RayStart, RayEnd);
 		if (bHit)
 		{
 			FVector3d HitPoint = (FVector3d)Result.ImpactPoint + (double)ShapeSettings->Shift * Direction;
diff --git a/Engine/Plugins/Experimental/MeshModelingToolsetExp/Source/MeshModelingToolsExp/Private/CubeGridTool.cpp b/Engine/Plugins/Experimental/MeshModelingToolsetExp/Source/MeshModelingToolsExp/Private/CubeGridTool.cpp
--- a/Engine/Plugins/Experimental/MeshModelingToolsetExp/Source/MeshModelingToolsExp/Private/CubeGridTool.cpp
+++ b/Engine/Plugins/Experimental/MeshModelingToolsetExp/Source/MeshModelingToolsExp/Private/CubeGridTool.cpp
@@ -1245,7 +1245,7 @@ bool UCubeGridTool::GetHitGridFace(const FRay& WorldRay, FCubeGrid::FCubeFace& F
 	if (Settings->bHitUnrelatedGeometry)
 	{
 		FHitResult HitResult;
-		if (ToolSceneQueriesUtil::FindNearestVisibleObjectHit(TargetWorld, HitResult, WorldRay)
+		if (ToolSceneQueriesUtil::FindNearestVisibleObjectHit(this, HitResult, WorldRay)
 			&& HitResult.Distance < BestHitT)
 		{
 			BestHitT = HitResult.Distance;
diff --git a/Engine/Plugins/Experimental/MeshModelingToolsetExp/Source/MeshModelingToolsExp/Private/DrawPolyPathTool.cpp b/Engine/Plugins/Experimental/MeshModelingToolsetExp/Source/MeshModelingToolsExp/Private/DrawPolyPathTool.cpp
--- a/Engine/Plugins/Experimental/MeshModelingToolsetExp/Source/MeshModelingToolsExp/Private/DrawPolyPathTool.cpp
+++ b/Engine/Plugins/Experimental/MeshModelingToolsetExp/Source/MeshModelingToolsExp/Private/DrawPolyPathTool.cpp
@@ -645,7 +645,7 @@ void UDrawPolyPathTool::BeginInteractiveExtrudeHeight()
 
 	ExtrudeHeightMechanic->WorldHitQueryFunc = [this](const FRay& WorldRay, FHitResult& HitResult)
 	{
-		return ToolSceneQueriesUtil::FindNearestVisibleObjectHit(TargetWorld, HitResult, WorldRay);
+		return ToolSceneQueriesUtil::FindNearestVisibleObjectHit(this, HitResult, WorldRay);
 	};
 	ExtrudeHeightMechanic->WorldPointSnapFunc = [this](const FVector3d& WorldPos, FVector3d& SnapPos)
 	{
diff --git a/Engine/Plugins/Experimental/MeshModelingToolsetExp/Source/MeshModelingToolsExp/Private/DynamicMeshSculptTool.cpp b/Engine/Plugins/Experimental/MeshModelingToolsetExp/Source/MeshModelingToolsExp/Private/DynamicMeshSculptTool.cpp
--- a/Engine/Plugins/Experimental/MeshModelingToolsetExp/Source/MeshModelingToolsExp/Private/DynamicMeshSculptTool.cpp
+++ b/Engine/Plugins/Experimental/MeshModelingToolsetExp/Source/MeshModelingToolsExp/Private/DynamicMeshSculptTool.cpp
@@ -1517,7 +1517,7 @@ void UDynamicMeshSculptTool::OnTick(float DeltaTime)
 		// raycast into scene and current sculpt and place plane at closest hit point
 		FRay CursorWorldRay = UMeshSurfacePointTool::LastWorldRay;
 		FHitResult Result;
-		bool bWorldHit = ToolSceneQueriesUtil::FindNearestVisibleObjectHit(TargetWorld, Result, CursorWorldRay);
+		bool bWorldHit = ToolSceneQueriesUtil::FindNearestVisibleObjectHit(this, Result, CursorWorldRay);
 		FRay3d LocalRay(CurTargetTransform.InverseTransformPosition((FVector3d)CursorWorldRay.Origin),
 			CurTargetTransform.InverseTransformVector((FVector3d)CursorWorldRay.Direction));
 		UE::Geometry::Normalize(LocalRay.Direction);
diff --git a/Engine/Plugins/Experimental/MeshModelingToolsetExp/Source/MeshModelingToolsExp/Private/EditPivotTool.cpp b/Engine/Plugins/Experimental/MeshModelingToolsetExp/Source/MeshModelingToolsExp/Private/EditPivotTool.cpp
--- a/Engine/Plugins/Experimental/MeshModelingToolsetExp/Source/MeshModelingToolsExp/Private/EditPivotTool.cpp
+++ b/Engine/Plugins/Experimental/MeshModelingToolsetExp/Source/MeshModelingToolsExp/Private/EditPivotTool.cpp
@@ -325,7 +325,7 @@ FInputRayHit UEditPivotTool::CanBeginClickDragSequence(const FInputDeviceRay& Pr
 	}
 
 	FHitResult Result;
-	bool bWorldHit = ToolSceneQueriesUtil::FindNearestVisibleObjectHit(TargetWorld, Result, PressPos.WorldRay);
+	bool bWorldHit = ToolSceneQueriesUtil::FindNearestVisibleObjectHit(this, Result, PressPos.WorldRay);
 
 	if (!bWorldHit)
 	{
@@ -353,7 +353,7 @@ void UEditPivotTool::OnClickDrag(const FInputDeviceRay& DragPos)
 	float NormalSign = (TransformProps->RotationMode == EEditPivotSnapDragRotationMode::AlignFlipped) ? -1.0f : 1.0f;
 
 	FHitResult Result;
-	bool bWorldHit = ToolSceneQueriesUtil::FindNearestVisibleObjectHit(TargetWorld, Result, DragPos.WorldRay);
+	bool bWorldHit = ToolSceneQueriesUtil::FindNearestVisibleObjectHit(this, Result, DragPos.WorldRay);
 	if (bWorldHit == false)
 	{
 		return;
diff --git a/Engine/Plugins/Experimental/MeshModelingToolsetExp/Source/MeshModelingToolsExp/Private/MeshSpaceDeformerTool.cpp b/Engine/Plugins/Experimental/MeshModelingToolsetExp/Source/MeshModelingToolsExp/Private/MeshSpaceDeformerTool.cpp
--- a/Engine/Plugins/Experimental/MeshModelingToolsetExp/Source/MeshModelingToolsExp/Private/MeshSpaceDeformerTool.cpp
+++ b/Engine/Plugins/Experimental/MeshModelingToolsetExp/Source/MeshModelingToolsExp/Private/MeshSpaceDeformerTool.cpp
@@ -3,6 +3,7 @@
 #include "MeshSpaceDeformerTool.h"
 #include "InteractiveToolManager.h"
 #include "InteractiveGizmoManager.h"
+#include "SceneQueries/SceneSnappingManager.h"
 #include "ToolBuilderUtil.h"
 
 #include "SegmentTypes.h"
@@ -192,7 +193,7 @@ void UMeshSpaceDeformerTool::Setup()
 
 	// add click to set plane behavior
 	SetPointInWorldConnector = MakePimpl<FSelectClickedAction>();
-	SetPointInWorldConnector->World = this->TargetWorld;
+	SetPointInWorldConnector->SnapManager = USceneSnappingManager::Find(GetToolManager());
 	SetPointInWorldConnector->InvisibleComponentsToHitTest.Add(TargetComponent->GetOwnerComponent());
 	SetPointInWorldConnector->OnClickedPositionFunc = [this](const FHitResult& Hit)
 	{
diff --git a/Engine/Plugins/Experimental/MeshModelingToolsetExp/Source/MeshModelingToolsExp/Private/Sculpting/MeshSculptToolBase.cpp b/Engine/Plugins/Experimental/MeshModelingToolsetExp/Source/MeshModelingToolsExp/Private/Sculpting/MeshSculptToolBase.cpp
--- a/Engine/Plugins/Experimental/MeshModelingToolsetExp/Source/MeshModelingToolsExp/Private/Sculpting/MeshSculptToolBase.cpp
+++ b/Engine/Plugins/Experimental/MeshModelingToolsetExp/Source/MeshModelingToolsExp/Private/Sculpting/MeshSculptToolBase.cpp
@@ -1213,7 +1213,7 @@ void UMeshSculptToolBase::UpdateWorkPlane()
 		// raycast into scene and current sculpt and place plane at closest hit point
 		FRay CursorWorldRay = UMeshSurfacePointTool::LastWorldRay;
 		FHitResult Result;
-		bool bWorldHit = ToolSceneQueriesUtil::FindNearestVisibleObjectHit(TargetWorld, Result, CursorWorldRay);
+		bool bWorldHit = ToolSceneQueriesUtil::FindNearestVisibleObjectHit(this, Result, CursorWorldRay);
 		FRay3d LocalRay = GetLocalRay(CursorWorldRay);
 		bool bObjectHit = ( FindHitSculptMeshTriangle(LocalRay) != IndexConstants::InvalidID );
 		if (bWorldHit && 
diff --git a/Engine/Plugins/Experimental/MeshModelingToolsetExp/Source/MeshModelingToolsExp/Private/TransformMeshesTool.cpp b/Engine/Plugins/Experimental/MeshModelingToolsetExp/Source/MeshModelingToolsExp/Private/TransformMeshesTool.cpp
--- a/Engine/Plugins/Experimental/MeshModelingToolsetExp/Source/MeshModelingToolsExp/Private/TransformMeshesTool.cpp
+++ b/Engine/Plugins/Experimental/MeshModelingToolsetExp/Source/MeshModelingToolsExp/Private/TransformMeshesTool.cpp
@@ -401,7 +401,7 @@ void UTransformMeshesTool::OnClickDrag(const FInputDeviceRay& DragPos)
 	float NormalSign = (TransformProps->RotationMode == ETransformMeshesSnapDragRotationMode::AlignFlipped) ? -1.0f : 1.0f;
 
 	FHitResult Result;
-	bool bWorldHit = ToolSceneQueriesUtil::FindNearestVisibleObjectHit(TargetWorld, Result, DragPos.WorldRay, &IgnoreComponents);
+	bool bWorldHit = ToolSceneQueriesUtil::FindNearestVisibleObjectHit(this, Result, DragPos.WorldRay, &IgnoreComponents);
 	if (bWorldHit == false)
 	{
 		return;
diff --git a/Engine/Plugins/Runtime/MeshModelingToolset/Source/MeshModelingTools/Private/AddPrimitiveTool.cpp b/Engine/Plugins/Runtime/MeshModelingToolset/Source/MeshModelingTools/Private/AddPrimitiveTool.cpp
--- a/Engine/Plugins/Runtime/MeshModelingToolset/Source/MeshModelingTools/Private/AddPrimitiveTool.cpp
+++ b/Engine/Plugins/Runtime/MeshModelingToolset/Source/MeshModelingTools/Private/AddPrimitiveTool.cpp
@@ -225,7 +225,7 @@ void UAddPrimitiveTool::UpdatePreviewPosition(const FInputDeviceRay& DeviceClick
 	{
 		// cast ray into scene
 		FHitResult Result;
-		bHit = ToolSceneQueriesUtil::FindNearestVisibleObjectHit(TargetWorld, Result, ClickPosWorldRay);
+		bHit = ToolSceneQueriesUtil::FindNearestVisibleObjectHit(this, Result, ClickPosWorldRay);
 		if (bHit)
 		{
 			FVector3f Normal = Result.ImpactNormal;
diff --git a/Engine/Plugins/Runtime/MeshModelingToolset/Source/MeshModelingTools/Private/DrawPolygonTool.cpp b/Engine/Plugins/Runtime/MeshModelingToolset/Source/MeshModelingTools/Private/DrawPolygonTool.cpp
--- a/Engine/Plugins/Runtime/MeshModelingToolset/Source/MeshModelingTools/Private/DrawPolygonTool.cpp
+++ b/Engine/Plugins/Runtime/MeshModelingToolset/Source/MeshModelingTools/Private/DrawPolygonTool.cpp
@@ -20,6 +20,7 @@
 #include "Distance/DistLine3Ray3.h"
 #include "Intersection/IntrSegment2Segment2.h"
 #include "ToolSceneQueriesUtil.h"
+#include "SceneQueries/SceneSnappingManager.h"
 #include "ConstrainedDelaunay2.h"
 #include "Arrangement2d.h"
 
@@ -91,7 +92,7 @@ void UDrawPolygonTool::Setup()
 
 	// Register a click behavior/action pair, that sets the draw plane to the clicked world position
 	FSelectClickedAction* SetPlaneAction = new FSelectClickedAction();
-	SetPlaneAction->World = this->TargetWorld;
+	SetPlaneAction->SnapManager = USceneSnappingManager::Find(GetToolManager());
 	SetPlaneAction->OnClickedPositionFunc = [this](const FHitResult& Hit) {
 		SetDrawPlaneFromWorldPos((FVector3d)Hit.ImpactPoint, (FVector3d)Hit.ImpactNormal);
 	};
@@ -610,7 +611,7 @@ bool UDrawPolygonTool::FindDrawPlaneHitPoint(const FInputDeviceRay& ClickPos, FV
 	if (SnapProperties->bSnapToSurfaces)
 	{
 		FHitResult Result;
-		bool bWorldHit = ToolSceneQueriesUtil::FindNearestVisibleObjectHit(TargetWorld, Result, ClickPos.WorldRay);
+		bool bWorldHit = ToolSceneQueriesUtil::FindNearestVisibleObjectHit(this, Result, ClickPos.WorldRay);
 		if (bWorldHit)
 		{
 			bHaveSurfaceHit = true;
@@ -969,7 +970,7 @@ void UDrawPolygonTool::BeginInteractiveExtrude()
 	{
 		if (this->bIgnoreSnappingToggle == false)
 		{
-			return ToolSceneQueriesUtil::FindNearestVisibleObjectHit(TargetWorld, HitResult, WorldRay);
+			return ToolSceneQueriesUtil::FindNearestVisibleObjectHit(this, HitResult, WorldRay);
 		}
 		return false;
 	};
diff --git a/Engine/Plugins/Runtime/MeshModelingToolset/Source/MeshModelingTools/Private/ToolActivities/PolyEditExtrudeActivity.cpp b/Engine/Plugins/Runtime/MeshModelingToolset/Source/MeshModelingTools/Private/ToolActivities/PolyEditExtrudeActivity.cpp
--- a/Engine/Plugins/Runtime/MeshModelingToolset/Source/MeshModelingTools/Private/ToolActivities/PolyEditExtrudeActivity.cpp
+++ b/Engine/Plugins/Runtime/MeshModelingToolset/Source/MeshModelingTools/Private/ToolActivities/PolyEditExtrudeActivity.cpp
@@ -234,8 +234,7 @@ EToolActivityStartResult UPolyEditExtrudeActivity::Start()
 	ExtrudeHeightMechanic->Setup(ParentTool.Get());
 	ExtrudeHeightMechanic->WorldHitQueryFunc = [this](const FRay& WorldRay, FHitResult& HitResult)
 	{
-		return ToolSceneQueriesUtil::FindNearestVisibleObjectHit(ActivityContext->Preview->GetWorld(),
-			HitResult, WorldRay);
+		return ToolSceneQueriesUtil::FindNearestVisibleObjectHit(GetParentTool(), HitResult, WorldRay);
 	};
 	ExtrudeHeightMechanic->WorldPointSnapFunc = [this](const FVector3d& WorldPos, FVector3d& SnapPos)
 	{
diff --git a/Engine/Plugins/Runtime/MeshModelingToolset/Source/MeshModelingTools/Private/UVProjectionTool.cpp b/Engine/Plugins/Runtime/MeshModelingToolset/Source/MeshModelingTools/Private/UVProjectionTool.cpp
--- a/Engine/Plugins/Runtime/MeshModelingToolset/Source/MeshModelingTools/Private/UVProjectionTool.cpp
+++ b/Engine/Plugins/Runtime/MeshModelingToolset/Source/MeshModelingTools/Private/UVProjectionTool.cpp
@@ -3,6 +3,7 @@
 #include "UVProjectionTool.h"
 #include "InteractiveToolManager.h"
 #include "ToolTargetManager.h"
+#include "SceneQueries/SceneSnappingManager.h"
 #include "ToolBuilderUtil.h"
 #include "ToolSetupUtil.h"
 #include "BaseBehaviors/SingleClickBehavior.h"
@@ -125,7 +126,7 @@ void UUVProjectionTool::Setup()
 
 	// click to set plane behavior
 	SetPlaneCtrlClickBehaviorTarget = MakeUnique<FSelectClickedAction>();
-	SetPlaneCtrlClickBehaviorTarget->World = this->TargetWorld;
+	SetPlaneCtrlClickBehaviorTarget->SnapManager = USceneSnappingManager::Find(GetToolManager());
 	SetPlaneCtrlClickBehaviorTarget->OnClickedPositionFunc = [this](const FHitResult& Hit)
 	{
 		UpdatePlaneFromClick(FVector3d(Hit.ImpactPoint), FVector3d(Hit.ImpactNormal), SetPlaneCtrlClickBehaviorTarget->bShiftModifierToggle);
diff --git a/Engine/Plugins/Runtime/MeshModelingToolset/Source/ModelingComponents/Private/Mechanics/ConstructionPlaneMechanic.cpp b/Engine/Plugins/Runtime/MeshModelingToolset/Source/ModelingComponents/Private/Mechanics/ConstructionPlaneMechanic.cpp
--- a/Engine/Plugins/Runtime/MeshModelingToolset/Source/ModelingComponents/Private/Mechanics/ConstructionPlaneMechanic.cpp
+++ b/Engine/Plugins/Runtime/MeshModelingToolset/Source/ModelingComponents/Private/Mechanics/ConstructionPlaneMechanic.cpp
@@ -7,6 +7,7 @@
 #include "BaseGizmos/TransformGizmoUtil.h"
 #include "Drawing/MeshDebugDrawing.h"
 #include "ToolSceneQueriesUtil.h"
+#include "SceneQueries/SceneSnappingManager.h"
 
 using namespace UE::Geometry;
 
@@ -39,7 +40,7 @@ void UConstructionPlaneMechanic::Initialize(UWorld* TargetWorld, const FFrame3d&
 
 	// click to set plane behavior
 	SetPlaneCtrlClickBehaviorTarget = MakeUnique<FSelectClickedAction>();
-	SetPlaneCtrlClickBehaviorTarget->World = TargetWorld;
+	SetPlaneCtrlClickBehaviorTarget->SnapManager = USceneSnappingManager::Find(GetParentTool()->GetToolManager());
 	SetPlaneCtrlClickBehaviorTarget->OnClickedPositionFunc = [this](const FHitResult& Hit)
 	{
 		SetDrawPlaneFromWorldPos(FVector3d(Hit.ImpactPoint), FVector3d(Hit.ImpactNormal), SetPlaneCtrlClickBehaviorTarget->bShiftModifierToggle);
diff --git a/Engine/Plugins/Runtime/MeshModelingToolset/Source/ModelingComponents/Private/Mechanics/DragAlignmentMechanic.cpp b/Engine/Plugins/Runtime/MeshModelingToolset/Source/ModelingComponents/Private/Mechanics/DragAlignmentMechanic.cpp
--- a/Engine/Plugins/Runtime/MeshModelingToolset/Source/ModelingComponents/Private/Mechanics/DragAlignmentMechanic.cpp
+++ b/Engine/Plugins/Runtime/MeshModelingToolset/Source/ModelingComponents/Private/Mechanics/DragAlignmentMechanic.cpp
@@ -54,7 +54,7 @@ void UDragAlignmentMechanic::Setup(UInteractiveTool* ParentToolIn)
 			bSuccess = true;
 		}
 		// If we failed, try to hit a simple triangle
-		else if (ToolSceneQueriesUtil::FindNearestVisibleObjectHit(ParentToolIn->GetWorld(), HitResult, WorldRay, ComponentsToIgnore))
+		else if (ToolSceneQueriesUtil::FindNearestVisibleObjectHit(ParentToolIn, HitResult, WorldRay, ComponentsToIgnore))
 		{
 			HitResult.Item = TRIANGLE_HIT_TYPE_ID;
 			bSuccess = true;
diff --git a/Engine/Plugins/Runtime/MeshModelingToolset/Source/ModelingComponents/Private/ToolSceneQueriesUtil.cpp b/Engine/Plugins/Runtime/MeshModelingToolset/Source/ModelingComponents/Private/ToolSceneQueriesUtil.cpp
--- a/Engine/Plugins/Runtime/MeshModelingToolset/Source/ModelingComponents/Private/ToolSceneQueriesUtil.cpp
+++ b/Engine/Plugins/Runtime/MeshModelingToolset/Source/ModelingComponents/Private/ToolSceneQueriesUtil.cpp
@@ -350,6 +350,8 @@ bool ToolSceneQueriesUtil::FindNearestVisibleObjectHit(USceneSnappingManager* Sn
 	FSceneHitQueryRequest Request;
 	Request.WorldRay = (FRay3d)Ray;
 	Request.bWantHitGeometryInfo = false;
+	Request.VisibilityFilter.ComponentsToIgnore = IgnoreComponents;
+	Request.VisibilityFilter.InvisibleComponentsToInclude = InvisibleComponentsToInclude;
 
 	FSceneHitQueryResult Result;
 	if ( SnapManager->ExecuteSceneHitQuery(Request, Result) )
diff --git a/Engine/Plugins/Runtime/MeshModelingToolset/Source/ModelingComponents/Public/Selection/SelectClickedAction.h b/Engine/Plugins/Runtime/MeshModelingToolset/Source/ModelingComponents/Public/Selection/SelectClickedAction.h
--- a/Engine/Plugins/Runtime/MeshModelingToolset/Source/ModelingComponents/Public/Selection/SelectClickedAction.h
+++ b/Engine/Plugins/Runtime/MeshModelingToolset/Source/ModelingComponents/Public/Selection/SelectClickedAction.h
@@ -15,11 +15,11 @@ class FSelectClickedAction : public IClickBehaviorTarget
 	FInputRayHit DoRayCast(const FInputDeviceRay& ClickPos, bool callbackOnHit)
 	{
 		FHitResult Result;
-
-		bool bHitWorld = ToolSceneQueriesUtil::FindNearestVisibleObjectHit(
-			World, Result, ClickPos.WorldRay, 
-			VisibleComponentsToIgnore.Num() == 0 ? nullptr : &VisibleComponentsToIgnore,
-			InvisibleComponentsToHitTest.Num() == 0 ? nullptr : &InvisibleComponentsToHitTest);
+		const TArray<const UPrimitiveComponent*>* IgnoreComponents = VisibleComponentsToIgnore.Num() == 0 ? nullptr : &VisibleComponentsToIgnore;
+		const TArray<const UPrimitiveComponent*>* InvisibleComponentsToInclude = InvisibleComponentsToHitTest.Num() == 0 ? nullptr : &InvisibleComponentsToHitTest;
+		bool bHitWorld = (SnapManager != nullptr) ?
+			ToolSceneQueriesUtil::FindNearestVisibleObjectHit(SnapManager, Result, ClickPos.WorldRay, IgnoreComponents, InvisibleComponentsToInclude) :
+			ToolSceneQueriesUtil::FindNearestVisibleObjectHit(World, Result, ClickPos.WorldRay, IgnoreComponents, InvisibleComponentsToInclude);
 
 		if (callbackOnHit && bHitWorld && OnClickedPositionFunc != nullptr)
 		{
@@ -29,7 +29,8 @@ class FSelectClickedAction : public IClickBehaviorTarget
 	}
 
 public:
-	UWorld* World;
+	USceneSnappingManager* SnapManager = nullptr;
+	UWorld* World = nullptr;
 	TFunction<void(const FHitResult&)> OnClickedPositionFunc = nullptr;
 	TUniqueFunction<bool()> ExternalCanClickPredicate = nullptr;
 
