UVLayoutTool: Fix UV layout being recomputed when preview settings change

diff --git a/Engine/Plugins/Runtime/MeshModelingToolset/Source/MeshModelingTools/Private/UVLayoutTool.cpp b/Engine/Plugins/Runtime/MeshModelingToolset/Source/MeshModelingTools/Private/UVLayoutTool.cpp
--- a/Engine/Plugins/Runtime/MeshModelingToolset/Source/MeshModelingTools/Private/UVLayoutTool.cpp
+++ b/Engine/Plugins/Runtime/MeshModelingToolset/Source/MeshModelingTools/Private/UVLayoutTool.cpp
@@ -233,7 +233,7 @@ void UUVLayoutTool::OnPropertyModified(UObject* PropertySet, FProperty* Property
 	else if (PropertySet == MaterialSettings)
 	{
 		// if we don't know what changed, or we know checker density changed, update checker material
-		UpdateVisualization();
+		UpdatePreviewMaterial();
 	}
 }
 
@@ -252,8 +252,7 @@ void UUVLayoutTool::OnPreviewMeshUpdated(UMeshOpPreviewWithBackgroundCompute* Co
 
 }
 
-
-void UUVLayoutTool::UpdateVisualization()
+void UUVLayoutTool::UpdatePreviewMaterial()
 {
 	MaterialSettings->UpdateMaterials();
 	UpdateNumPreviews();
@@ -262,6 +261,11 @@ void UUVLayoutTool::UpdateVisualization()
 		UMeshOpPreviewWithBackgroundCompute* Preview = Previews[PreviewIdx];
 		Preview->OverrideMaterial = MaterialSettings->GetActiveOverrideMaterial();
 	}
+}
+
+void UUVLayoutTool::UpdateVisualization()
+{
+	UpdatePreviewMaterial();
 
 	for (UMeshOpPreviewWithBackgroundCompute* Preview : Previews)
 	{
diff --git a/Engine/Plugins/Runtime/MeshModelingToolset/Source/MeshModelingTools/Public/UVLayoutTool.h b/Engine/Plugins/Runtime/MeshModelingToolset/Source/MeshModelingTools/Public/UVLayoutTool.h
--- a/Engine/Plugins/Runtime/MeshModelingToolset/Source/MeshModelingTools/Public/UVLayoutTool.h
+++ b/Engine/Plugins/Runtime/MeshModelingToolset/Source/MeshModelingTools/Public/UVLayoutTool.h
@@ -87,6 +87,8 @@ protected:
 	void UpdateNumPreviews();
 
 	void UpdateVisualization();
+	
+	void UpdatePreviewMaterial();
 
 	void OnPreviewMeshUpdated(UMeshOpPreviewWithBackgroundCompute* Compute);
 
