MeshModelingToolset: Add some missing progress cancel checks to Voxel tools

#rnx

diff --git a/Engine/Plugins/Runtime/MeshModelingToolset/Source/ModelingOperators/Private/CompositionOps/VoxelBlendMeshesOp.cpp b/Engine/Plugins/Runtime/MeshModelingToolset/Source/ModelingOperators/Private/CompositionOps/VoxelBlendMeshesOp.cpp
--- a/Engine/Plugins/Runtime/MeshModelingToolset/Source/ModelingOperators/Private/CompositionOps/VoxelBlendMeshesOp.cpp
+++ b/Engine/Plugins/Runtime/MeshModelingToolset/Source/ModelingOperators/Private/CompositionOps/VoxelBlendMeshesOp.cpp
@@ -69,7 +69,15 @@ void FVoxelBlendMeshesOp::CalculateResult(FProgressCancel* Progress)
 			TFastWindingTree<FDynamicMesh3> Winding(&Spatial);
 			TImplicitSolidify<FDynamicMesh3> Solidify(&TransformedMeshes[MeshIdx], &Spatial, &Winding);
 			Solidify.SetCellSizeAndExtendBounds(Spatial.GetBoundingBox(), 0, InputVoxelCount);
+			Solidify.CancelF = [&Progress]()
+			{
+				return Progress && Progress->Cancelled();
+			};
 			TransformedMeshes[MeshIdx].Copy(&Solidify.Generate());
+			
+			if (Progress && Progress->Cancelled()) {
+				return;
+			};
 
 			if (bRemoveInternalsAfterVoxWrap)
 			{
@@ -93,7 +101,6 @@ void FVoxelBlendMeshesOp::CalculateResult(FProgressCancel* Progress)
 
 	ImplicitBlend.SetCellSizesAndFalloff(CombinedBounds, BlendFalloff, InputVoxelCount, OutputVoxelCount);
 	ImplicitBlend.BlendPower = BlendPower;
-
 	ImplicitBlend.CancelF = [&Progress]()
 	{
 		return Progress && Progress->Cancelled();
diff --git a/Engine/Plugins/Runtime/MeshModelingToolset/Source/ModelingOperators/Private/CompositionOps/VoxelMorphologyMeshesOp.cpp b/Engine/Plugins/Runtime/MeshModelingToolset/Source/ModelingOperators/Private/CompositionOps/VoxelMorphologyMeshesOp.cpp
--- a/Engine/Plugins/Runtime/MeshModelingToolset/Source/ModelingOperators/Private/CompositionOps/VoxelMorphologyMeshesOp.cpp
+++ b/Engine/Plugins/Runtime/MeshModelingToolset/Source/ModelingOperators/Private/CompositionOps/VoxelMorphologyMeshesOp.cpp
@@ -104,7 +104,16 @@ void FVoxelMorphologyMeshesOp::CalculateResult(FProgressCancel* Progress)
 		TFastWindingTree<FDynamicMesh3> Winding(&Spatial);
 		TImplicitSolidify<FDynamicMesh3> Solidify(&CombinedMesh, &Spatial, &Winding);
 		Solidify.SetCellSizeAndExtendBounds(Spatial.GetBoundingBox(), 0, InputVoxelCount);
+		Solidify.CancelF = [&Progress]()
+		{
+			return Progress && Progress->Cancelled();
+		};
 		CombinedMesh.Copy(&Solidify.Generate());
+		
+		if (Progress && Progress->Cancelled())
+		{
+			return;
+		}
 
 		if (bRemoveInternalsAfterVoxWrap)
 		{
@@ -121,14 +130,11 @@ void FVoxelMorphologyMeshesOp::CalculateResult(FProgressCancel* Progress)
 
 	ImplicitMorphology.SourceSpatial = &Spatial;
 	ImplicitMorphology.SetCellSizesAndDistance(CombinedMesh.GetBounds(true), Distance, InputVoxelCount, OutputVoxelCount);
-	
 	ImplicitMorphology.CancelF = [&Progress]()
 	{
 		return Progress && Progress->Cancelled();
 	};
 
-	;
-
 	if (Progress && Progress->Cancelled())
 	{
 		return;
diff --git a/Engine/Plugins/Runtime/MeshModelingToolset/Source/ModelingOperators/Private/CompositionOps/VoxelSolidifyMeshesOp.cpp b/Engine/Plugins/Runtime/MeshModelingToolset/Source/ModelingOperators/Private/CompositionOps/VoxelSolidifyMeshesOp.cpp
--- a/Engine/Plugins/Runtime/MeshModelingToolset/Source/ModelingOperators/Private/CompositionOps/VoxelSolidifyMeshesOp.cpp
+++ b/Engine/Plugins/Runtime/MeshModelingToolset/Source/ModelingOperators/Private/CompositionOps/VoxelSolidifyMeshesOp.cpp
@@ -82,6 +82,10 @@ void FVoxelSolidifyMeshesOp::CalculateResult(FProgressCancel* Progress)
 	Solidify.SurfaceSearchSteps = SurfaceSearchSteps;
 	Solidify.bSolidAtBoundaries = bSolidAtBoundaries;
 	Solidify.ExtendBounds = ExtendBounds;
+	Solidify.CancelF = [&Progress]()
+	{
+		return Progress && Progress->Cancelled();
+	};
 
 	if (Progress && Progress->Cancelled())
 	{
diff --git a/Engine/Source/Runtime/GeometryCore/Public/Implicit/Blend.h b/Engine/Source/Runtime/GeometryCore/Public/Implicit/Blend.h
--- a/Engine/Source/Runtime/GeometryCore/Public/Implicit/Blend.h
+++ b/Engine/Source/Runtime/GeometryCore/Public/Implicit/Blend.h
@@ -174,6 +174,11 @@ protected:
 
 	void GenerateBlendAnalytic(bool bReuseComputed)
 	{
+		if (CancelF())
+		{
+			return;
+		}
+		
 		ComputeSpatials(bReuseComputed);
 
 		if (CancelF())
