ModelingTools: Add support for reading and writing to a specific LOD via IMeshDescriptionProvider/Committer. Refactor various other aspects of MeshDescription access. Update MeshTransferTool to be able to transfer from and to specific LODs.

Add new MeshTargetInterfaceTypes.h file, move EStaticMeshEditingLOD there and rename to EMeshLODIdentifier. Add FGetMeshParameters and FCommitMeshParameters types.

IMeshDescriptionProvider::GetMeshDescription() now takes FGetMeshParameters to optionally specify LOD.
Added IMeshDescriptionProvider::GetMeshDescriptionCopy() function, to handle copy-case.
removed IMeshDescriptionProvider::CalculateAutoGeneratedAttributes(). This was only being used to force Tangents computation, which can now be done via GetMeshDescriptionCopy() and FGetMeshParameters.bWantMeshTangents option

IMeshDescriptionCommitter commit functions now take a FCommitMeshParameters to optionally specify target LOD.

StaticMeshComponentToolTarget::GetMeshDescriptionCopy() added, optionally computes auto-generated MeshDescription attributes on the copy

StaticMesh(Component)ToolTarget updated to support FCommitMeshParameters target-LOD.

SkeletalMesh, Volume, and DynamicMesh ToolTargets updated for new APIs but do not currently support any of the new LOD parameters. These should never be called w/ LOD parameters in the current codebase, ensures added to catch any errors (they would be non-fatal).

UE::ToolTarget::GetMeshDescription() and GetMeshDescriptionCopy() in now take FGetMeshParameters argument.
Removed bWantMeshTangents param to GetMeshDescriptionCopy(), now done via FGetMeshParameters. Updated call sites.

TransferMeshesTool now supports specifying read and write LOD (numbered or HiRes) for StaticMeshComponent source/target, via above functions.
TransferMeshesTool can now also target a new LOD in a StaticMeshComponent, index is 1 larger than current maximum source LOD.

#rnx

diff --git a/Engine/Plugins/Editor/ModelingToolsEditorMode/Source/ModelingToolsEditorMode/Private/ModelingToolsEditorModeToolkit.cpp b/Engine/Plugins/Editor/ModelingToolsEditorMode/Source/ModelingToolsEditorMode/Private/ModelingToolsEditorModeToolkit.cpp
--- a/Engine/Plugins/Editor/ModelingToolsEditorMode/Source/ModelingToolsEditorMode/Private/ModelingToolsEditorModeToolkit.cpp
+++ b/Engine/Plugins/Editor/ModelingToolsEditorMode/Source/ModelingToolsEditorMode/Private/ModelingToolsEditorModeToolkit.cpp
@@ -265,14 +265,14 @@ TSharedPtr<SWidget> FModelingToolsEditorModeToolkit::MakeAssetConfigPanel()
 		.OptionsSource(&AssetLODModes)
 		.OnSelectionChanged_Lambda([&](TSharedPtr<FString> String, ESelectInfo::Type)
 	{
-		EStaticMeshEditingLOD NewSelectedLOD = EStaticMeshEditingLOD::LOD0;
+		EMeshLODIdentifier NewSelectedLOD = EMeshLODIdentifier::LOD0;
 		if (*String == *AssetLODModes[0])
 		{
-			NewSelectedLOD = EStaticMeshEditingLOD::MaxQuality;
+			NewSelectedLOD = EMeshLODIdentifier::MaxQuality;
 		}
 		else if (*String == *AssetLODModes[1])
 		{
-			NewSelectedLOD = EStaticMeshEditingLOD::HiResSource;
+			NewSelectedLOD = EMeshLODIdentifier::HiResSource;
 		}
 		else
 		{
@@ -280,7 +280,7 @@ TSharedPtr<SWidget> FModelingToolsEditorModeToolkit::MakeAssetConfigPanel()
 			{
 				if (*String == *AssetLODModes[k])
 				{
-					NewSelectedLOD = (EStaticMeshEditingLOD)(k - 2);
+					NewSelectedLOD = (EMeshLODIdentifier)(k - 2);
 					break;
 				}
 			}
diff --git a/Engine/Plugins/Experimental/MeshModelingToolsetExp/Source/MeshModelingToolsEditorOnly/Private/SimplifyMeshTool.cpp b/Engine/Plugins/Experimental/MeshModelingToolsetExp/Source/MeshModelingToolsEditorOnly/Private/SimplifyMeshTool.cpp
--- a/Engine/Plugins/Experimental/MeshModelingToolsetExp/Source/MeshModelingToolsEditorOnly/Private/SimplifyMeshTool.cpp
+++ b/Engine/Plugins/Experimental/MeshModelingToolsetExp/Source/MeshModelingToolsEditorOnly/Private/SimplifyMeshTool.cpp
@@ -92,7 +92,9 @@ void USimplifyMeshTool::Setup()
 #else
 		auto EnterProgressFrame = [](int Progress) {};
 #endif
-		OriginalMeshDescription = MakeShared<FMeshDescription, ESPMode::ThreadSafe>(UE::ToolTarget::GetMeshDescriptionCopy(Target, true));
+		FGetMeshParameters GetMeshParams;
+		GetMeshParams.bWantMeshTangents = true;
+		OriginalMeshDescription = MakeShared<FMeshDescription, ESPMode::ThreadSafe>(UE::ToolTarget::GetMeshDescriptionCopy(Target, GetMeshParams));
 
 		EnterProgressFrame(1);
 		// UE::ToolTarget::GetDynamicMeshCopy() would recompute the tangents a second time here
diff --git a/Engine/Plugins/Experimental/MeshModelingToolsetExp/Source/MeshModelingToolsExp/Private/BakeTransformTool.cpp b/Engine/Plugins/Experimental/MeshModelingToolsetExp/Source/MeshModelingToolsExp/Private/BakeTransformTool.cpp
--- a/Engine/Plugins/Experimental/MeshModelingToolsetExp/Source/MeshModelingToolsExp/Private/BakeTransformTool.cpp
+++ b/Engine/Plugins/Experimental/MeshModelingToolsetExp/Source/MeshModelingToolsExp/Private/BakeTransformTool.cpp
@@ -193,7 +193,7 @@ void UBakeTransformTool::UpdateAssets()
 			}
 
 			// apply edit
-			FMeshDescription SourceMesh(UE::ToolTarget::GetMeshDescriptionCopy(Targets[ComponentIdx], false));
+			FMeshDescription SourceMesh(UE::ToolTarget::GetMeshDescriptionCopy(Targets[ComponentIdx]));
 			FMeshDescriptionEditableTriangleMeshAdapter EditableMeshDescAdapter(&SourceMesh);
 
 			// do this part within the commit because we have the MeshDescription already computed
diff --git a/Engine/Plugins/Experimental/MeshModelingToolsetExp/Source/MeshModelingToolsExp/Private/EditPivotTool.cpp b/Engine/Plugins/Experimental/MeshModelingToolsetExp/Source/MeshModelingToolsExp/Private/EditPivotTool.cpp
--- a/Engine/Plugins/Experimental/MeshModelingToolsetExp/Source/MeshModelingToolsExp/Private/EditPivotTool.cpp
+++ b/Engine/Plugins/Experimental/MeshModelingToolsetExp/Source/MeshModelingToolsExp/Private/EditPivotTool.cpp
@@ -498,7 +498,7 @@ void UEditPivotTool::UpdateAssets(const FFrame3d& NewPivotWorldFrame)
 				}
 			}
 
-			FMeshDescription SourceMesh(UE::ToolTarget::GetMeshDescriptionCopy(Targets[ComponentIdx], false));
+			FMeshDescription SourceMesh(UE::ToolTarget::GetMeshDescriptionCopy(Targets[ComponentIdx]));
 			FMeshDescriptionEditableTriangleMeshAdapter EditableMeshDescAdapter(&SourceMesh);
 			MeshAdapterTransforms::ApplyTransform(EditableMeshDescAdapter, ToBake);
 			if (bNeedSeparateScale)
diff --git a/Engine/Plugins/Experimental/MeshModelingToolsetExp/Source/MeshModelingToolsExp/Private/TransferMeshTool.cpp b/Engine/Plugins/Experimental/MeshModelingToolsetExp/Source/MeshModelingToolsExp/Private/TransferMeshTool.cpp
--- a/Engine/Plugins/Experimental/MeshModelingToolsetExp/Source/MeshModelingToolsExp/Private/TransferMeshTool.cpp
+++ b/Engine/Plugins/Experimental/MeshModelingToolsetExp/Source/MeshModelingToolsExp/Private/TransferMeshTool.cpp
@@ -6,7 +6,11 @@
 #include "ToolTargetManager.h"
 #include "ToolBuilderUtil.h"
 #include "ToolSetupUtil.h"
+#include "TargetInterfaces/StaticMeshBackedTarget.h"
+#include "TargetInterfaces/MeshDescriptionProvider.h"
+#include "TargetInterfaces/MeshDescriptionCommitter.h"
 #include "ModelingToolTargetUtil.h"
+#include "Engine/StaticMesh.h"
 
 using namespace UE::Geometry;
 
@@ -42,6 +46,63 @@ void UTransferMeshTool::Setup()
 	GetToolManager()->DisplayMessage(
 		LOCTEXT("OnStartTool", "Copy Mesh from Source object to Target object"),
 		EToolMessageLevel::UserNotification);
+
+	// currently we can only query available LODs via the UStaticMesh, improvements TBD
+#if WITH_EDITOR
+	if (IStaticMeshBackedTarget* StaticMeshTarget = Cast<IStaticMeshBackedTarget>(Targets[0]))
+	{
+		if (UStaticMesh* StaticMesh = StaticMeshTarget->GetStaticMesh())
+		{
+			BasicProperties->bIsStaticMeshSource = true;
+			int32 NumSourceModels = StaticMesh->GetNumSourceModels();
+			for (int32 si = 0; si < NumSourceModels; ++si)
+			{
+				if (StaticMesh->IsMeshDescriptionValid(si))
+				{
+					BasicProperties->SourceLODNamesList.Add(FString::Printf(TEXT("LOD %d"), si));
+					BasicProperties->SourceLODEnums.Add(static_cast<EMeshLODIdentifier>(si));
+				}
+			}
+
+			// hires LOD slot
+			if (StaticMesh->IsHiResMeshDescriptionValid())
+			{
+				BasicProperties->SourceLODNamesList.Add(TEXT("HiRes"));
+				BasicProperties->SourceLODEnums.Add(EMeshLODIdentifier::HiResSource);
+			}
+
+			BasicProperties->SourceLOD = BasicProperties->SourceLODNamesList[0];
+		}
+	}	
+
+	if (IStaticMeshBackedTarget* StaticMeshTarget = Cast<IStaticMeshBackedTarget>(Targets[1]))
+	{
+		if (UStaticMesh* StaticMesh = StaticMeshTarget->GetStaticMesh())
+		{
+			BasicProperties->bIsStaticMeshTarget = true;
+			int32 NumSourceModels = StaticMesh->GetNumSourceModels();
+			for (int32 si = 0; si < NumSourceModels; ++si)
+			{
+				BasicProperties->TargetLODNamesList.Add( FString::Printf(TEXT("LOD %d"), si) );
+				BasicProperties->TargetLODEnums.Add(static_cast<EMeshLODIdentifier>(si));
+			}
+
+			// option to add one additional lod
+			if (NumSourceModels <= 7)
+			{
+				BasicProperties->TargetLODNamesList.Add(FString::Printf(TEXT("LOD %d (New)"), NumSourceModels));
+				BasicProperties->TargetLODEnums.Add(static_cast<EMeshLODIdentifier>(NumSourceModels));
+			}
+
+			// hires LOD slot
+			BasicProperties->TargetLODNamesList.Add(TEXT("HiRes"));
+			BasicProperties->TargetLODEnums.Add(EMeshLODIdentifier::HiResSource);
+
+			BasicProperties->TargetLOD = BasicProperties->TargetLODNamesList[0];
+		}
+	}
+#endif
+
 }
 
 bool UTransferMeshTool::CanAccept() const
@@ -57,14 +118,65 @@ void UTransferMeshTool::Shutdown(EToolShutdownType ShutdownType)
 	{
 		GetToolManager()->BeginUndoTransaction(LOCTEXT("TransferMeshToolTransactionName", "Transfer Mesh"));
 
-		const FMeshDescription* SourceMesh = UE::ToolTarget::GetMeshDescription(Targets[0]);
-
 		FComponentMaterialSet Materials = UE::ToolTarget::GetMaterialSet(Targets[0]);
 		const FComponentMaterialSet* TransferMaterials = (BasicProperties->bTransferMaterials) ? &Materials : nullptr;
 
-		if (ensure(SourceMesh))
+		FMeshDescription SourceMesh;
+
+		FGetMeshParameters GetMeshParams;
+		GetMeshParams.bWantMeshTangents = true;
+
+		if (IMeshDescriptionProvider* MeshDescriptionProvider = Cast<IMeshDescriptionProvider>(Targets[0]))
+		{
+			FString SelectedLOD = BasicProperties->SourceLOD;
+			if (SelectedLOD.StartsWith(TEXT("HiRes")))
+			{
+				GetMeshParams.bHaveRequestLOD = true;
+				GetMeshParams.RequestLOD = EMeshLODIdentifier::HiResSource;
+			}
+			else
+			{
+				int32 FoundIndex = BasicProperties->SourceLODNamesList.IndexOfByKey(BasicProperties->SourceLOD);
+				if (FoundIndex != INDEX_NONE)
+				{
+					GetMeshParams.bHaveRequestLOD = true;
+					GetMeshParams.RequestLOD =  BasicProperties->SourceLODEnums[FoundIndex];
+				}
+			}
+		}
+		SourceMesh = UE::ToolTarget::GetMeshDescriptionCopy(Targets[0], GetMeshParams);
+
+
+		IMeshDescriptionCommitter* MeshDescriptionCommitter = Cast<IMeshDescriptionCommitter>(Targets[1]);
+		if (BasicProperties->bIsStaticMeshTarget && MeshDescriptionCommitter)
+		{
+			FCommitMeshParameters CommitParams;
+			FString SelectedLOD = BasicProperties->TargetLOD;
+			if (SelectedLOD.StartsWith(TEXT("HiRes")))
+			{
+				CommitParams.bHaveTargetLOD = true;
+				CommitParams.TargetLOD = EMeshLODIdentifier::HiResSource;
+			}
+			else
+			{
+				int32 FoundIndex = BasicProperties->TargetLODNamesList.IndexOfByKey(BasicProperties->TargetLOD);
+				if (FoundIndex != INDEX_NONE)
+				{
+					CommitParams.bHaveTargetLOD = true;
+					CommitParams.TargetLOD =  BasicProperties->TargetLODEnums[FoundIndex];
+				}
+			}
+
+			if (TransferMaterials)
+			{
+				UE::ToolTarget::CommitMaterialSetUpdate(Targets[1], *TransferMaterials, true);
+			}
+
+			MeshDescriptionCommitter->CommitMeshDescription(MoveTemp(SourceMesh), CommitParams);
+		}
+		else 
 		{
-			UE::ToolTarget::CommitMeshDescriptionUpdate(Targets[1], SourceMesh, TransferMaterials);
+			UE::ToolTarget::CommitMeshDescriptionUpdate(Targets[1], &SourceMesh, TransferMaterials);
 		}
 
 		GetToolManager()->EndUndoTransaction();
diff --git a/Engine/Plugins/Experimental/MeshModelingToolsetExp/Source/MeshModelingToolsExp/Public/TransferMeshTool.h b/Engine/Plugins/Experimental/MeshModelingToolsetExp/Source/MeshModelingToolsExp/Public/TransferMeshTool.h
--- a/Engine/Plugins/Experimental/MeshModelingToolsetExp/Source/MeshModelingToolsExp/Public/TransferMeshTool.h
+++ b/Engine/Plugins/Experimental/MeshModelingToolsetExp/Source/MeshModelingToolsExp/Public/TransferMeshTool.h
@@ -5,6 +5,7 @@
 #include "CoreMinimal.h"
 #include "BaseTools/MultiSelectionMeshEditingTool.h"
 #include "InteractiveToolBuilder.h"
+#include "TargetInterfaces/MeshTargetInterfaceTypes.h"
 #include "TransferMeshTool.generated.h"
 
 /**
@@ -19,6 +20,7 @@ public:
 	virtual UMultiSelectionMeshEditingTool* CreateNewTool(const FToolBuilderState& SceneState) const override;
 };
 
+
 /**
  * Standard properties of the Transfer operation
  */
@@ -29,6 +31,42 @@ class MESHMODELINGTOOLSEXP_API UTransferMeshToolProperties : public UInteractive
 public:
 	UPROPERTY(EditAnywhere, Category = Options)
 	bool bTransferMaterials = true;
+
+	/** Specify which LOD from the Source (First) mesh to copy from */
+	UPROPERTY(EditAnywhere, Category = SourceMesh, meta = (DisplayName = "Source LOD", TransientToolProperty, NoResetToDefault, GetOptions = GetSourceLODNamesFunc, HideEditConditionToggle, EditCondition="bIsStaticMeshSource", EditConditionHides))
+	FString SourceLOD;
+
+	/** Specify which LOD on the Target (Second) mesh to copy to */
+	UPROPERTY(EditAnywhere, Category = TargetMesh, meta = (DisplayName = "Target LOD", TransientToolProperty, NoResetToDefault, GetOptions = GetTargetLODNamesFunc, HideEditConditionToggle, EditCondition="bIsStaticMeshTarget", EditConditionHides))
+	FString TargetLOD;
+
+
+	//
+	// non-exposed properties used to provide custom lists to SourceLOD/TargetLOD
+	//
+
+	UPROPERTY(meta = (TransientToolProperty))
+	bool bIsStaticMeshSource  = false;
+
+	UFUNCTION()
+	const TArray<FString>& GetSourceLODNamesFunc() const { return SourceLODNamesList; }
+
+	UPROPERTY(meta = (TransientToolProperty))
+	TArray<FString> SourceLODNamesList;
+
+	TArray<EMeshLODIdentifier> SourceLODEnums;
+
+
+	UFUNCTION()
+	const TArray<FString>& GetTargetLODNamesFunc() const { return TargetLODNamesList; }
+
+	UPROPERTY(meta = (TransientToolProperty))
+	TArray<FString> TargetLODNamesList;
+
+	TArray<EMeshLODIdentifier> TargetLODEnums;
+
+	UPROPERTY(meta = (TransientToolProperty))
+	bool bIsStaticMeshTarget = false;
 };
 
 
diff --git a/Engine/Plugins/Runtime/MeshModelingToolset/Source/ModelingComponents/Private/ModelingToolTargetUtil.cpp b/Engine/Plugins/Runtime/MeshModelingToolset/Source/ModelingComponents/Private/ModelingToolTargetUtil.cpp
--- a/Engine/Plugins/Runtime/MeshModelingToolset/Source/ModelingComponents/Private/ModelingToolTargetUtil.cpp
+++ b/Engine/Plugins/Runtime/MeshModelingToolset/Source/ModelingComponents/Private/ModelingToolTargetUtil.cpp
@@ -144,31 +144,26 @@ bool UE::ToolTarget::CommitMaterialSetUpdate(
 
 
 
-const FMeshDescription* UE::ToolTarget::GetMeshDescription(UToolTarget* Target)
+const FMeshDescription* UE::ToolTarget::GetMeshDescription(UToolTarget* Target, const FGetMeshParameters& GetMeshParams)
 {
 	static FMeshDescription EmptyMeshDescription;
 
 	IMeshDescriptionProvider* MeshDescriptionProvider = Cast<IMeshDescriptionProvider>(Target);
 	if (MeshDescriptionProvider)
 	{
-		return MeshDescriptionProvider->GetMeshDescription();
+		return MeshDescriptionProvider->GetMeshDescription(GetMeshParams);
 	}
 	ensure(false);
 	return &EmptyMeshDescription;
 }
 
 
-FMeshDescription UE::ToolTarget::GetMeshDescriptionCopy(UToolTarget* Target, bool bWantMeshTangents)
+FMeshDescription UE::ToolTarget::GetMeshDescriptionCopy(UToolTarget* Target, const FGetMeshParameters& GetMeshParams)
 {
 	IMeshDescriptionProvider* MeshDescriptionProvider = Cast<IMeshDescriptionProvider>(Target);
 	if (MeshDescriptionProvider)
 	{
-		FMeshDescription NewMeshDescription = *MeshDescriptionProvider->GetMeshDescription();
-		if (bWantMeshTangents)
-		{
-			MeshDescriptionProvider->CalculateAutoGeneratedAttributes(NewMeshDescription);
-		}
-		return NewMeshDescription;
+		return MeshDescriptionProvider->GetMeshDescriptionCopy(GetMeshParams);
 	}
 	ensure(false);
 	return FMeshDescription();
@@ -201,10 +196,9 @@ FDynamicMesh3 UE::ToolTarget::GetDynamicMeshCopy(UToolTarget* Target, bool bWant
 		FMeshDescriptionToDynamicMesh Converter;
 		if (bWantMeshTangents)
 		{
-			// Currently to auto-calculate attributes we must make a copy, because they are written
-			// directly to the input mesh
-			FMeshDescription MeshDescriptionCopy( *MeshDescriptionProvider->GetMeshDescription() );
-			MeshDescriptionProvider->CalculateAutoGeneratedAttributes(MeshDescriptionCopy);
+			FGetMeshParameters GetMeshParams;
+			GetMeshParams.bWantMeshTangents = true;
+			FMeshDescription MeshDescriptionCopy = MeshDescriptionProvider->GetMeshDescriptionCopy(GetMeshParams);
 			Converter.Convert(&MeshDescriptionCopy, Mesh, bWantMeshTangents);
 		}
 		else
@@ -278,8 +272,7 @@ UE::ToolTarget::EDynamicMeshUpdateResult UE::ToolTarget::CommitMeshDescriptionUp
 		// This should never need to be true because we are either (1) ignoring the tangents, in which
 		// case, we don't want to update them anyway (2) replacing them, in which case, we don't need
 		// to compute them
-		const bool bComputeAutoGeneratedTangents = false;
-		ConvertedMesh = UE::ToolTarget::GetMeshDescriptionCopy(Target, bComputeAutoGeneratedTangents);
+		ConvertedMesh = UE::ToolTarget::GetMeshDescriptionCopy(Target, FGetMeshParameters());
 		FDynamicMeshToMeshDescription Converter;
 		Converter.Update(&UpdatedMesh, ConvertedMesh);
 	}
diff --git a/Engine/Plugins/Runtime/MeshModelingToolset/Source/ModelingComponents/Public/ModelingToolTargetUtil.h b/Engine/Plugins/Runtime/MeshModelingToolset/Source/ModelingComponents/Public/ModelingToolTargetUtil.h
--- a/Engine/Plugins/Runtime/MeshModelingToolset/Source/ModelingComponents/Public/ModelingToolTargetUtil.h
+++ b/Engine/Plugins/Runtime/MeshModelingToolset/Source/ModelingComponents/Public/ModelingToolTargetUtil.h
@@ -4,6 +4,7 @@
 
 #include "CoreMinimal.h"
 #include "DynamicMesh/DynamicMesh3.h"
+#include "TargetInterfaces/MeshTargetInterfaceTypes.h"
 #include "TargetInterfaces/MaterialProvider.h" // FComponentMaterialSet
 #include "MeshConversionOptions.h"
 
@@ -90,14 +91,17 @@ MODELINGCOMPONENTS_API bool CommitMaterialSetUpdate(
 /**
  * @return the MeshDescription underlying a ToolTarget, if it has such a mesh. May be generated internally by the ToolTarget. May be nullptr if the Target does not have a mesh.
  */
-MODELINGCOMPONENTS_API const FMeshDescription* GetMeshDescription(UToolTarget* Target);
+MODELINGCOMPONENTS_API const FMeshDescription* GetMeshDescription(
+	UToolTarget* Target,
+	const FGetMeshParameters& GetMeshParams = FGetMeshParameters());
 
 /**
- * @return a copy of the MeshDescription underlying a ToolTarget
- * @param bWantMeshTangents if true, tangents will be returned if the target has them available. This may require that they be auto-calculated in some cases (which may be expensive)
+ * Return a copy of the MeshDescription underlying a ToolTarget
  * @return a new MeshDescription, which may be empty if the Target doesn't have a mesh  
  */
-MODELINGCOMPONENTS_API FMeshDescription GetMeshDescriptionCopy(UToolTarget* Target, bool bWantMeshTangents = false);
+MODELINGCOMPONENTS_API FMeshDescription GetMeshDescriptionCopy(
+	UToolTarget* Target, 
+	const FGetMeshParameters& GetMeshParams = FGetMeshParameters());
 
 /**
  * Fetch a DynamicMesh3 representing the given ToolTarget. This may be a conversion of the output of GetMeshDescription().
diff --git a/Engine/Plugins/Runtime/MeshModelingToolset/Source/ModelingComponentsEditorOnly/Private/ToolTargets/DynamicMeshComponentToolTarget.cpp b/Engine/Plugins/Runtime/MeshModelingToolset/Source/ModelingComponentsEditorOnly/Private/ToolTargets/DynamicMeshComponentToolTarget.cpp
--- a/Engine/Plugins/Runtime/MeshModelingToolset/Source/ModelingComponentsEditorOnly/Private/ToolTargets/DynamicMeshComponentToolTarget.cpp
+++ b/Engine/Plugins/Runtime/MeshModelingToolset/Source/ModelingComponentsEditorOnly/Private/ToolTargets/DynamicMeshComponentToolTarget.cpp
@@ -93,10 +93,12 @@ bool UDynamicMeshComponentToolTarget::CommitMaterialSetUpdate(const FComponentMa
 }
 
 
-const FMeshDescription* UDynamicMeshComponentToolTarget::GetMeshDescription()
+const FMeshDescription* UDynamicMeshComponentToolTarget::GetMeshDescription(const FGetMeshParameters& GetMeshParams)
 {
 	if (ensure(IsValid()))
 	{
+		ensure(GetMeshParams.bHaveRequestLOD == false);	// not supported yet, just returning default LOD
+
 		if (bHaveMeshDescription)
 		{
 			return ConvertedMeshDescription.Get();
@@ -118,10 +120,12 @@ const FMeshDescription* UDynamicMeshComponentToolTarget::GetMeshDescription()
 }
 
 
-void UDynamicMeshComponentToolTarget::CommitMeshDescription(const FCommitter& Committer)
+void UDynamicMeshComponentToolTarget::CommitMeshDescription(const FCommitter& Committer, const FCommitMeshParameters& CommitMeshParams)
 {
 	if (ensure(IsValid()) == false) return;
 
+	ensure(CommitMeshParams.bHaveTargetLOD == false);		// not supporting this yet
+
 	UDynamicMesh* DynamicMesh = GetDynamicMeshContainer();
 	TUniquePtr<FDynamicMesh3> CurrentMesh = DynamicMesh->ExtractMesh();
 	TSharedPtr<FDynamicMesh3> CurrentMeshShared(CurrentMesh.Release());
diff --git a/Engine/Plugins/Runtime/MeshModelingToolset/Source/ModelingComponentsEditorOnly/Private/ToolTargets/SkeletalMeshComponentToolTarget.cpp b/Engine/Plugins/Runtime/MeshModelingToolset/Source/ModelingComponentsEditorOnly/Private/ToolTargets/SkeletalMeshComponentToolTarget.cpp
--- a/Engine/Plugins/Runtime/MeshModelingToolset/Source/ModelingComponentsEditorOnly/Private/ToolTargets/SkeletalMeshComponentToolTarget.cpp
+++ b/Engine/Plugins/Runtime/MeshModelingToolset/Source/ModelingComponentsEditorOnly/Private/ToolTargets/SkeletalMeshComponentToolTarget.cpp
@@ -88,12 +88,13 @@ bool USkeletalMeshComponentReadOnlyToolTarget::CommitMaterialSetUpdate(const FCo
 	return true;
 }
 
-const FMeshDescription* USkeletalMeshComponentReadOnlyToolTarget::GetMeshDescription()
+const FMeshDescription* USkeletalMeshComponentReadOnlyToolTarget::GetMeshDescription(const FGetMeshParameters& GetMeshParams)
 {
 	if (!ensure(IsValid()))
 	{
 		return nullptr;
 	}
+	ensure(GetMeshParams.bHaveRequestLOD == false);	// not supported yet, just returning default LOD
 
 	if (!CachedMeshDescription.IsValid())
 	{
@@ -123,7 +124,7 @@ USkeletalMesh* USkeletalMeshComponentReadOnlyToolTarget::GetSkeletalMesh() const
 //
 
 
-void USkeletalMeshComponentToolTarget::CommitMeshDescription(const FCommitter& Committer)
+void USkeletalMeshComponentToolTarget::CommitMeshDescription(const FCommitter& Committer, const FCommitMeshParameters& CommitMeshParams)
 {
 	if (ensure(IsValid()) == false) return;
 
diff --git a/Engine/Plugins/Runtime/MeshModelingToolset/Source/ModelingComponentsEditorOnly/Private/ToolTargets/SkeletalMeshToolTarget.cpp b/Engine/Plugins/Runtime/MeshModelingToolset/Source/ModelingComponentsEditorOnly/Private/ToolTargets/SkeletalMeshToolTarget.cpp
--- a/Engine/Plugins/Runtime/MeshModelingToolset/Source/ModelingComponentsEditorOnly/Private/ToolTargets/SkeletalMeshToolTarget.cpp
+++ b/Engine/Plugins/Runtime/MeshModelingToolset/Source/ModelingComponentsEditorOnly/Private/ToolTargets/SkeletalMeshToolTarget.cpp
@@ -107,12 +107,13 @@ bool USkeletalMeshReadOnlyToolTarget::CommitMaterialSetUpdate(USkeletalMesh* Ske
 	return true;
 }
 
-const FMeshDescription* USkeletalMeshReadOnlyToolTarget::GetMeshDescription()
+const FMeshDescription* USkeletalMeshReadOnlyToolTarget::GetMeshDescription(const FGetMeshParameters& GetMeshParams)
 {
 	if (!ensure(IsValid()))
 	{
 		return nullptr;
 	}
+	ensure(GetMeshParams.bHaveRequestLOD == false);	// not supported yet, just returning default LOD
 
 	if (!CachedMeshDescription.IsValid())
 	{
@@ -161,7 +162,7 @@ USkeletalMesh* USkeletalMeshReadOnlyToolTarget::GetSkeletalMesh() const
 // USkeletalMeshToolTarget
 //
 
-void USkeletalMeshToolTarget::CommitMeshDescription(const FCommitter& Committer)
+void USkeletalMeshToolTarget::CommitMeshDescription(const FCommitter& Committer, const FCommitMeshParameters& CommitMeshParams)
 {
 	if (ensure(IsValid()) == false) return;
 
diff --git a/Engine/Plugins/Runtime/MeshModelingToolset/Source/ModelingComponentsEditorOnly/Private/ToolTargets/StaticMeshComponentToolTarget.cpp b/Engine/Plugins/Runtime/MeshModelingToolset/Source/ModelingComponentsEditorOnly/Private/ToolTargets/StaticMeshComponentToolTarget.cpp
--- a/Engine/Plugins/Runtime/MeshModelingToolset/Source/ModelingComponentsEditorOnly/Private/ToolTargets/StaticMeshComponentToolTarget.cpp
+++ b/Engine/Plugins/Runtime/MeshModelingToolset/Source/ModelingComponentsEditorOnly/Private/ToolTargets/StaticMeshComponentToolTarget.cpp
@@ -14,9 +14,9 @@
 
 using namespace UE::Geometry;
 
-void UStaticMeshComponentToolTarget::SetEditingLOD(EStaticMeshEditingLOD RequestedEditingLOD)
+void UStaticMeshComponentToolTarget::SetEditingLOD(EMeshLODIdentifier RequestedEditingLOD)
 {
-	EStaticMeshEditingLOD ValidEditingLOD = EStaticMeshEditingLOD::LOD0;
+	EMeshLODIdentifier ValidEditingLOD = EMeshLODIdentifier::LOD0;
 
 	UStaticMeshComponent* StaticMeshComponent = Cast<UStaticMeshComponent>(Component);
 	if (ensure(StaticMeshComponent != nullptr))
@@ -125,7 +125,7 @@ bool UStaticMeshComponentToolTarget::CommitMaterialSetUpdate(const FComponentMat
 }
 
 
-const FMeshDescription* UStaticMeshComponentToolTarget::GetMeshDescription()
+const FMeshDescription* UStaticMeshComponentToolTarget::GetMeshDescription(const FGetMeshParameters& GetMeshParams)
 {
 	static FMeshDescription EmptyMeshDescription;
 	static bool bFirst = true;
@@ -138,56 +138,86 @@ const FMeshDescription* UStaticMeshComponentToolTarget::GetMeshDescription()
 
 	if (ensure(IsValid()))
 	{
-		UStaticMesh* StaticMesh = Cast<UStaticMeshComponent>(Component)->GetStaticMesh();
-		FMeshDescription* MeshDescription = (EditingLOD == EStaticMeshEditingLOD::HiResSource) ?
-			StaticMesh->GetHiResMeshDescription() : StaticMesh->GetMeshDescription((int32)EditingLOD);
+		ensure(GetMeshParams.bWantMeshTangents == false);		// cannot support in this path because we have to modify the FoundMeshDescription pointer!
+
+		const FMeshDescription* FoundMeshDescription = nullptr;
 
-		if (!MeshDescription)
+		if (UStaticMesh* StaticMesh = Cast<UStaticMeshComponent>(Component)->GetStaticMesh())
 		{
-			return &EmptyMeshDescription;
+			EMeshLODIdentifier UseLOD = EditingLOD;
+			if (GetMeshParams.bHaveRequestLOD)
+			{
+				UseLOD = UStaticMeshToolTarget::GetValidEditingLOD(StaticMesh, GetMeshParams.RequestLOD);
+				ensure(UseLOD == GetMeshParams.RequestLOD);		// probably a bug somewhere if this is not true
+			}
+
+			FoundMeshDescription = (UseLOD == EMeshLODIdentifier::HiResSource) ?
+				StaticMesh->GetHiResMeshDescription() : StaticMesh->GetMeshDescription((int32)UseLOD);
 		}
 
-		return MeshDescription;
+		return (FoundMeshDescription != nullptr) ? FoundMeshDescription : &EmptyMeshDescription;
 	}
 	return nullptr;
 }
 
 
-void UStaticMeshComponentToolTarget::CalculateAutoGeneratedAttributes(FMeshDescription& MeshDescription)
+FMeshDescription UStaticMeshComponentToolTarget::GetMeshDescriptionCopy(const FGetMeshParameters& GetMeshParams)
 {
 	if (ensure(IsValid()))
 	{
-		UStaticMesh* StaticMesh = Cast<UStaticMeshComponent>(Component)->GetStaticMesh();
-		if (EditingLOD == EStaticMeshEditingLOD::HiResSource)
-		{
-			const FStaticMeshSourceModel& SourceModel = StaticMesh->GetHiResSourceModel();
-			UE::MeshDescription::InitializeAutoGeneratedAttributes(MeshDescription, &SourceModel.BuildSettings);
-		}
-		else
+		if (UStaticMesh* StaticMesh = Cast<UStaticMeshComponent>(Component)->GetStaticMesh())
 		{
-			UE::MeshDescription::InitializeAutoGeneratedAttributes(MeshDescription, StaticMesh, (int32)EditingLOD);
+			EMeshLODIdentifier UseLOD = EditingLOD;
+			if (GetMeshParams.bHaveRequestLOD)
+			{
+				UseLOD = UStaticMeshToolTarget::GetValidEditingLOD(StaticMesh, GetMeshParams.RequestLOD);
+				ensure(UseLOD == GetMeshParams.RequestLOD);		// probably a bug somewhere if this is not true
+			}
+			if (UseLOD == EMeshLODIdentifier::HiResSource )
+			{
+				if (StaticMesh->IsHiResMeshDescriptionValid())
+				{
+					FMeshDescription MeshDescriptionCopy = *StaticMesh->GetHiResMeshDescription();
+					const FStaticMeshSourceModel& SourceModel = StaticMesh->GetHiResSourceModel();
+					UE::MeshDescription::InitializeAutoGeneratedAttributes(MeshDescriptionCopy, &SourceModel.BuildSettings);
+					return MeshDescriptionCopy;
+				}
+			}
+			else
+			{
+				if (StaticMesh->IsMeshDescriptionValid(static_cast<int32>(UseLOD)))
+				{
+					if (FMeshDescription* SourceMesh = StaticMesh->GetMeshDescription(static_cast<int32>(UseLOD)))
+					{
+						FMeshDescription MeshDescriptionCopy = *SourceMesh;
+						UE::MeshDescription::InitializeAutoGeneratedAttributes(MeshDescriptionCopy, StaticMesh, static_cast<int32>(UseLOD));
+						return MeshDescriptionCopy;
+					}
+				}
+			}
 		}
 	}
+	
+
+	FMeshDescription EmptyMeshDescription;
+	FStaticMeshAttributes Attributes(EmptyMeshDescription);
+	Attributes.Register();
+	return EmptyMeshDescription;
 }
 
 
-void UStaticMeshComponentToolTarget::CommitMeshDescription(const FCommitter& Committer)
+void UStaticMeshComponentToolTarget::CommitMeshDescription(const FCommitter& Committer, const FCommitMeshParameters& CommitParams)
 {
 	if (ensure(IsValid()) == false) return;
 
 	UStaticMesh* StaticMesh = Cast<UStaticMeshComponent>(Component)->GetStaticMesh();
-	FMeshDescription* UpdateMeshDescription = (EditingLOD == EStaticMeshEditingLOD::HiResSource) ?
-		StaticMesh->GetHiResMeshDescription() : StaticMesh->GetMeshDescription((int32)EditingLOD);
 
-	if (!ensure(UpdateMeshDescription != nullptr))
-	{
-		return; // invalid target; cannot commit
-	}
+	EMeshLODIdentifier WriteToLOD = (CommitParams.bHaveTargetLOD && CommitParams.TargetLOD != EMeshLODIdentifier::Default) ? CommitParams.TargetLOD : EditingLOD;
 
 	// unregister the component while we update its static mesh
 	FComponentReregisterContext ComponentReregisterContext(Component);
 
-	UStaticMeshToolTarget::CommitMeshDescription(StaticMesh, UpdateMeshDescription, Committer, EditingLOD);
+	UStaticMeshToolTarget::CommitMeshDescription(StaticMesh, Committer, WriteToLOD);
 
 	// this rebuilds physics, but it doesn't undo!
 	Component->RecreatePhysicsState();
@@ -253,7 +283,7 @@ UToolTarget* UStaticMeshComponentToolTargetFactory::BuildTarget(UObject* SourceO
 }
 
 
-void UStaticMeshComponentToolTargetFactory::SetActiveEditingLOD(EStaticMeshEditingLOD NewEditingLOD)
+void UStaticMeshComponentToolTargetFactory::SetActiveEditingLOD(EMeshLODIdentifier NewEditingLOD)
 {
 	EditingLOD = NewEditingLOD;
 }
\ No newline at end of file
diff --git a/Engine/Plugins/Runtime/MeshModelingToolset/Source/ModelingComponentsEditorOnly/Private/ToolTargets/StaticMeshToolTarget.cpp b/Engine/Plugins/Runtime/MeshModelingToolset/Source/ModelingComponentsEditorOnly/Private/ToolTargets/StaticMeshToolTarget.cpp
--- a/Engine/Plugins/Runtime/MeshModelingToolset/Source/ModelingComponentsEditorOnly/Private/ToolTargets/StaticMeshToolTarget.cpp
+++ b/Engine/Plugins/Runtime/MeshModelingToolset/Source/ModelingComponentsEditorOnly/Private/ToolTargets/StaticMeshToolTarget.cpp
@@ -28,28 +28,28 @@ namespace StaticMeshToolTargetLocals
 	}
 }
 
-void UStaticMeshToolTarget::SetEditingLOD(EStaticMeshEditingLOD RequestedEditingLOD)
+void UStaticMeshToolTarget::SetEditingLOD(EMeshLODIdentifier RequestedEditingLOD)
 {
 	EditingLOD = GetValidEditingLOD(StaticMesh, RequestedEditingLOD);
 }
 
-EStaticMeshEditingLOD UStaticMeshToolTarget::GetValidEditingLOD(const UStaticMesh* StaticMeshIn, 
-	EStaticMeshEditingLOD RequestedEditingLOD)
+EMeshLODIdentifier UStaticMeshToolTarget::GetValidEditingLOD(const UStaticMesh* StaticMeshIn, 
+	EMeshLODIdentifier RequestedEditingLOD)
 {
 	using namespace StaticMeshToolTargetLocals;
 
-	EStaticMeshEditingLOD ValidEditingLOD = EStaticMeshEditingLOD::LOD0;
+	EMeshLODIdentifier ValidEditingLOD = EMeshLODIdentifier::LOD0;
 
 	if (ensure(StaticMeshIn != nullptr))
 	{
-		if (RequestedEditingLOD == EStaticMeshEditingLOD::MaxQuality)
+		if (RequestedEditingLOD == EMeshLODIdentifier::MaxQuality)
 		{
-			ValidEditingLOD = StaticMeshIn->IsHiResMeshDescriptionValid() ? EStaticMeshEditingLOD::HiResSource : EStaticMeshEditingLOD::LOD0;
+			ValidEditingLOD = StaticMeshIn->IsHiResMeshDescriptionValid() ? EMeshLODIdentifier::HiResSource : EMeshLODIdentifier::LOD0;
 		}
-		else if (RequestedEditingLOD == EStaticMeshEditingLOD::HiResSource)
+		else if (RequestedEditingLOD == EMeshLODIdentifier::HiResSource)
 		{
-			ValidEditingLOD = StaticMeshIn->IsHiResMeshDescriptionValid() ? EStaticMeshEditingLOD::HiResSource : EStaticMeshEditingLOD::LOD0;
-			if (ValidEditingLOD != EStaticMeshEditingLOD::HiResSource)
+			ValidEditingLOD = StaticMeshIn->IsHiResMeshDescriptionValid() ? EMeshLODIdentifier::HiResSource : EMeshLODIdentifier::LOD0;
+			if (ValidEditingLOD != EMeshLODIdentifier::HiResSource)
 			{
 				DisplayCriticalWarningMessage(FString(TEXT("HiRes Source selected but not available - Falling Back to LOD0")));
 			}
@@ -61,7 +61,7 @@ EStaticMeshEditingLOD UStaticMeshToolTarget::GetValidEditingLOD(const UStaticMes
 			if ((int32)ValidEditingLOD > MaxExistingLOD)
 			{
 				DisplayCriticalWarningMessage(FString::Printf(TEXT("LOD%d Requested but not available - Falling Back to LOD%d"), (int32)ValidEditingLOD, MaxExistingLOD));
-				ValidEditingLOD = (EStaticMeshEditingLOD)MaxExistingLOD;
+				ValidEditingLOD = (EMeshLODIdentifier)MaxExistingLOD;
 			}
 		}
 	}
@@ -75,14 +75,14 @@ bool UStaticMeshToolTarget::IsValid() const
 	return IsValid(StaticMesh, EditingLOD);
 }
 
-bool UStaticMeshToolTarget::IsValid(const UStaticMesh* StaticMeshIn, EStaticMeshEditingLOD EditingLODIn)
+bool UStaticMeshToolTarget::IsValid(const UStaticMesh* StaticMeshIn, EMeshLODIdentifier EditingLODIn)
 {
 	if (!StaticMeshIn || !IsValidChecked(StaticMeshIn) || StaticMeshIn->IsUnreachable() || !StaticMeshIn->IsValidLowLevel())
 	{
 		return false;
 	}
 
-	if (EditingLODIn == EStaticMeshEditingLOD::HiResSource)
+	if (EditingLODIn == EMeshLODIdentifier::HiResSource)
 	{
 		if (StaticMeshIn->IsHiResMeshDescriptionValid() == false)
 		{
@@ -184,7 +184,7 @@ bool UStaticMeshToolTarget::CommitMaterialSetUpdate(UStaticMesh* StaticMeshIn,
 	return true;
 }
 
-const FMeshDescription* UStaticMeshToolTarget::GetMeshDescription()
+const FMeshDescription* UStaticMeshToolTarget::GetMeshDescription(const FGetMeshParameters& GetMeshParams)
 {
 	static FMeshDescription EmptyMeshDescription;
 	static bool bFirst = true;
@@ -197,34 +197,39 @@ const FMeshDescription* UStaticMeshToolTarget::GetMeshDescription()
 
 	if (ensure(IsValid()))
 	{
-		FMeshDescription* MeshDescription = (EditingLOD == EStaticMeshEditingLOD::HiResSource) ?
-			StaticMesh->GetHiResMeshDescription() : StaticMesh->GetMeshDescription((int32)EditingLOD);
-
-		if (!MeshDescription)
+		EMeshLODIdentifier UseLOD = EditingLOD;
+		if (GetMeshParams.bHaveRequestLOD)
 		{
-			return &EmptyMeshDescription;
+			UseLOD = UStaticMeshToolTarget::GetValidEditingLOD(StaticMesh, GetMeshParams.RequestLOD);
+			ensure(UseLOD == GetMeshParams.RequestLOD);		// probably a bug somewhere if this is not true
 		}
 
-		return MeshDescription;
+		FMeshDescription* FoundMeshDescription = (UseLOD == EMeshLODIdentifier::HiResSource) ?
+			StaticMesh->GetHiResMeshDescription() : StaticMesh->GetMeshDescription((int32)UseLOD);
+
+		return (FoundMeshDescription != nullptr) ? FoundMeshDescription : &EmptyMeshDescription;
 	}
 	return nullptr;
 }
 
-void UStaticMeshToolTarget::CommitMeshDescription(const FCommitter& Committer)
+void UStaticMeshToolTarget::CommitMeshDescription(const FCommitter& Committer, const FCommitMeshParameters& CommitParams)
 {
 	if (ensure(IsValid()) == false) return;
 
-	FMeshDescription* UpdateMeshDescription = (EditingLOD == EStaticMeshEditingLOD::HiResSource) ?
-		StaticMesh->GetHiResMeshDescription() : StaticMesh->GetMeshDescription((int32)EditingLOD);
+	EMeshLODIdentifier WriteToLOD = (CommitParams.bHaveTargetLOD && CommitParams.TargetLOD != EMeshLODIdentifier::Default) ? CommitParams.TargetLOD : EditingLOD;
 
-	CommitMeshDescription(StaticMesh, UpdateMeshDescription, Committer, EditingLOD);
+	CommitMeshDescription(StaticMesh, Committer, WriteToLOD);
 }
 
-void UStaticMeshToolTarget::CommitMeshDescription(UStaticMesh* StaticMeshIn, FMeshDescription* MeshDescription,
-	const FCommitter& Committer, EStaticMeshEditingLOD EditingLODIn)
+void UStaticMeshToolTarget::CommitMeshDescription(UStaticMesh* StaticMeshIn, const FCommitter& Committer, EMeshLODIdentifier EditingLODIn)
 {
 	using namespace StaticMeshToolTargetLocals;
 
+	if ( ! ensure(EditingLODIn != EMeshLODIdentifier::Default && EditingLODIn != EMeshLODIdentifier::MaxQuality) )
+	{
+		UE_LOG(LogGeometry, Warning, TEXT("UStaticMeshToolTarget::CommitMeshDescription: invalid Target LOD, must specify explicit LOD"));
+		return;
+	}
 	if (StaticMeshIn->GetPathName().StartsWith(TEXT("/Engine/")))
 	{
 		DisplayCriticalWarningMessage(FString::Printf(TEXT("CANNOT MODIFY BUILT-IN ENGINE ASSET %s"), *StaticMeshIn->GetPathName()));
@@ -236,15 +241,43 @@ void UStaticMeshToolTarget::CommitMeshDescription(UStaticMesh* StaticMeshIn, FMe
 
 	// make sure transactional flag is on for this asset
 	StaticMeshIn->SetFlags(RF_Transactional);
+	// mark as modified
+	StaticMeshIn->Modify();
 
-	verify(StaticMeshIn->Modify());
+	FMeshDescription* UpdateMeshDescription = nullptr;
+	if (EditingLODIn == EMeshLODIdentifier::HiResSource)
+	{
+		UpdateMeshDescription = StaticMeshIn->GetHiResMeshDescription();
+		if (UpdateMeshDescription == nullptr)
+		{
+			UpdateMeshDescription = StaticMeshIn->CreateHiResMeshDescription();
+		}
+	}
+	else
+	{
+		int32 UseLODIndex = static_cast<int32>(EditingLODIn);
+		if (StaticMeshIn->GetNumSourceModels() < UseLODIndex+1)
+		{
+			StaticMeshIn->SetNumSourceModels(UseLODIndex+1);
+		}
+
+		UpdateMeshDescription = StaticMeshIn->GetMeshDescription(UseLODIndex);
+		if (UpdateMeshDescription == nullptr)
+		{
+			UpdateMeshDescription = StaticMeshIn->CreateMeshDescription(UseLODIndex);
+		}
+	}
 
 	// disable auto-generated normals StaticMesh build setting
 	UE::MeshDescription::FStaticMeshBuildSettingChange SettingsChange;
 	SettingsChange.AutoGeneratedNormals = UE::MeshDescription::EBuildSettingBoolChange::Disable;
-	UE::MeshDescription::ConfigureBuildSettings(StaticMeshIn, 0, SettingsChange);
+	if (static_cast<int32>(EditingLODIn) <= (int32)EMeshLODIdentifier::LOD7)
+	{
+		UE::MeshDescription::ConfigureBuildSettings(StaticMeshIn, static_cast<int32>(EditingLODIn), SettingsChange);
+	}
+	// do we need to configure build settings for highres LOD?
 
-	if (EditingLODIn == EStaticMeshEditingLOD::HiResSource)
+	if (EditingLODIn == EMeshLODIdentifier::HiResSource)
 	{
 		verify(StaticMeshIn->ModifyHiResMeshDescription());
 	}
@@ -254,17 +287,22 @@ void UStaticMeshToolTarget::CommitMeshDescription(UStaticMesh* StaticMeshIn, FMe
 	}
 
 	FCommitterParams CommitterParams;
-	CommitterParams.MeshDescriptionOut = MeshDescription;
+	CommitterParams.MeshDescriptionOut = UpdateMeshDescription;
 
 	Committer(CommitterParams);
 
-	if (EditingLODIn == EStaticMeshEditingLOD::HiResSource)
+	if (EditingLODIn == EMeshLODIdentifier::HiResSource)
 	{
 		StaticMeshIn->CommitHiResMeshDescription();
 	}
 	else
 	{
 		StaticMeshIn->CommitMeshDescription((int32)EditingLODIn);
+
+		// configure build settings to prevent the standard static mesh reduction from running and replacing the render LOD.
+		FStaticMeshSourceModel& ThisSourceModel = StaticMeshIn->GetSourceModel((int32)EditingLODIn);
+		ThisSourceModel.ReductionSettings.PercentTriangles = 1.f;
+		ThisSourceModel.ReductionSettings.PercentVertices = 1.f;
 	}
 
 	StaticMeshIn->PostEditChange();
@@ -308,7 +346,7 @@ UToolTarget* UStaticMeshToolTargetFactory::BuildTarget(UObject* SourceObject, co
 }
 
 
-void UStaticMeshToolTargetFactory::SetActiveEditingLOD(EStaticMeshEditingLOD NewEditingLOD)
+void UStaticMeshToolTargetFactory::SetActiveEditingLOD(EMeshLODIdentifier NewEditingLOD)
 {
 	EditingLOD = NewEditingLOD;
 }
diff --git a/Engine/Plugins/Runtime/MeshModelingToolset/Source/ModelingComponentsEditorOnly/Private/ToolTargets/VolumeComponentToolTarget.cpp b/Engine/Plugins/Runtime/MeshModelingToolset/Source/ModelingComponentsEditorOnly/Private/ToolTargets/VolumeComponentToolTarget.cpp
--- a/Engine/Plugins/Runtime/MeshModelingToolset/Source/ModelingComponentsEditorOnly/Private/ToolTargets/VolumeComponentToolTarget.cpp
+++ b/Engine/Plugins/Runtime/MeshModelingToolset/Source/ModelingComponentsEditorOnly/Private/ToolTargets/VolumeComponentToolTarget.cpp
@@ -106,7 +106,7 @@ void UVolumeComponentToolTarget::CommitDynamicMesh(const UE::Geometry::FDynamicM
 	Volume->PostEditChange();
 }
 
-const FMeshDescription* UVolumeComponentToolTarget::GetMeshDescription()
+const FMeshDescription* UVolumeComponentToolTarget::GetMeshDescription(const FGetMeshParameters& GetMeshParams)
 {
 	if (!ConvertedMeshDescription.IsValid())
 	{
@@ -126,8 +126,9 @@ const FMeshDescription* UVolumeComponentToolTarget::GetMeshDescription()
 	return ConvertedMeshDescription.Get();
 }
 
-void UVolumeComponentToolTarget::CommitMeshDescription(const FCommitter& Committer)
+void UVolumeComponentToolTarget::CommitMeshDescription(const FCommitter& Committer, const FCommitMeshParameters& CommitMeshParams)
 {
+	// no LODs so LODIdentifier is ignored
 	check(IsValid());
 
 	// Let the user fill our mesh description with the Committer
diff --git a/Engine/Plugins/Runtime/MeshModelingToolset/Source/ModelingComponentsEditorOnly/Public/ToolTargets/DynamicMeshComponentToolTarget.h b/Engine/Plugins/Runtime/MeshModelingToolset/Source/ModelingComponentsEditorOnly/Public/ToolTargets/DynamicMeshComponentToolTarget.h
--- a/Engine/Plugins/Runtime/MeshModelingToolset/Source/ModelingComponentsEditorOnly/Public/ToolTargets/DynamicMeshComponentToolTarget.h
+++ b/Engine/Plugins/Runtime/MeshModelingToolset/Source/ModelingComponentsEditorOnly/Public/ToolTargets/DynamicMeshComponentToolTarget.h
@@ -39,10 +39,10 @@ public:
 
 public:
 	// IMeshDescriptionProvider implementation
-	const FMeshDescription* GetMeshDescription() override;
+	const FMeshDescription* GetMeshDescription(const FGetMeshParameters& GetMeshParams = FGetMeshParameters()) override;
 
 	// IMeshDescritpionCommitter implementation
-	virtual void CommitMeshDescription(const FCommitter& Committer) override;
+	virtual void CommitMeshDescription(const FCommitter& Committer, const FCommitMeshParameters& CommitMeshParams = FCommitMeshParameters()) override;
 	using IMeshDescriptionCommitter::CommitMeshDescription; // unhide the other overload
 
 	// IDynamicMeshProvider implementation
diff --git a/Engine/Plugins/Runtime/MeshModelingToolset/Source/ModelingComponentsEditorOnly/Public/ToolTargets/SkeletalMeshComponentToolTarget.h b/Engine/Plugins/Runtime/MeshModelingToolset/Source/ModelingComponentsEditorOnly/Public/ToolTargets/SkeletalMeshComponentToolTarget.h
--- a/Engine/Plugins/Runtime/MeshModelingToolset/Source/ModelingComponentsEditorOnly/Public/ToolTargets/SkeletalMeshComponentToolTarget.h
+++ b/Engine/Plugins/Runtime/MeshModelingToolset/Source/ModelingComponentsEditorOnly/Public/ToolTargets/SkeletalMeshComponentToolTarget.h
@@ -34,19 +34,19 @@ class MODELINGCOMPONENTSEDITORONLY_API USkeletalMeshComponentReadOnlyToolTarget
 
 public:
 	// IMeshDescriptionProvider implementation
-	const FMeshDescription* GetMeshDescription() override;
+	virtual const FMeshDescription* GetMeshDescription(const FGetMeshParameters& GetMeshParams = FGetMeshParameters()) override;
 
 	// IMaterialProvider implementation
-	int32 GetNumMaterials() const override;
-	UMaterialInterface* GetMaterial(int32 MaterialIndex) const override;
-	void GetMaterialSet(FComponentMaterialSet& MaterialSetOut, bool bPreferAssetMaterials) const override;
-	bool CommitMaterialSetUpdate(const FComponentMaterialSet& MaterialSet, bool bApplyToAsset) override;
+	virtual int32 GetNumMaterials() const override;
+	virtual UMaterialInterface* GetMaterial(int32 MaterialIndex) const override;
+	virtual void GetMaterialSet(FComponentMaterialSet& MaterialSetOut, bool bPreferAssetMaterials) const override;
+	virtual bool CommitMaterialSetUpdate(const FComponentMaterialSet& MaterialSet, bool bApplyToAsset) override;
 
 	// IDynamicMeshProvider
 	virtual UE::Geometry::FDynamicMesh3 GetDynamicMesh() override;
 
 	// ISkeletalMeshBackedTarget implementation
-	USkeletalMesh* GetSkeletalMesh() const override;
+	virtual USkeletalMesh* GetSkeletalMesh() const override;
 
 protected:
 	// So that the tool target factory can poke into Component.
@@ -72,7 +72,7 @@ class MODELINGCOMPONENTSEDITORONLY_API USkeletalMeshComponentToolTarget :
 
 public:
 	// IMeshDescriptionCommitter implementation
-	void CommitMeshDescription(const FCommitter& Committer) override;
+	virtual void CommitMeshDescription(const FCommitter& Committer, const FCommitMeshParameters& CommitParams = FCommitMeshParameters()) override;
 	using IMeshDescriptionCommitter::CommitMeshDescription; // unhide the other overload
 
 	// IDynamicMeshCommitter
@@ -93,9 +93,9 @@ class MODELINGCOMPONENTSEDITORONLY_API USkeletalMeshComponentReadOnlyToolTargetF
 
 public:
 
-	bool CanBuildTarget(UObject* SourceObject, const FToolTargetTypeRequirements& TargetTypeInfo) const override;
+	virtual bool CanBuildTarget(UObject* SourceObject, const FToolTargetTypeRequirements& TargetTypeInfo) const override;
 
-	UToolTarget* BuildTarget(UObject* SourceObject, const FToolTargetTypeRequirements& TargetTypeInfo) override;
+	virtual UToolTarget* BuildTarget(UObject* SourceObject, const FToolTargetTypeRequirements& TargetTypeInfo) override;
 };
 
 
@@ -107,7 +107,7 @@ class MODELINGCOMPONENTSEDITORONLY_API USkeletalMeshComponentToolTargetFactory :
 
 public:
 
-	bool CanBuildTarget(UObject* SourceObject, const FToolTargetTypeRequirements& TargetTypeInfo) const override;
+	virtual bool CanBuildTarget(UObject* SourceObject, const FToolTargetTypeRequirements& TargetTypeInfo) const override;
 
-	UToolTarget* BuildTarget(UObject* SourceObject, const FToolTargetTypeRequirements& TargetTypeInfo) override;
+	virtual UToolTarget* BuildTarget(UObject* SourceObject, const FToolTargetTypeRequirements& TargetTypeInfo) override;
 };
diff --git a/Engine/Plugins/Runtime/MeshModelingToolset/Source/ModelingComponentsEditorOnly/Public/ToolTargets/SkeletalMeshToolTarget.h b/Engine/Plugins/Runtime/MeshModelingToolset/Source/ModelingComponentsEditorOnly/Public/ToolTargets/SkeletalMeshToolTarget.h
--- a/Engine/Plugins/Runtime/MeshModelingToolset/Source/ModelingComponentsEditorOnly/Public/ToolTargets/SkeletalMeshToolTarget.h
+++ b/Engine/Plugins/Runtime/MeshModelingToolset/Source/ModelingComponentsEditorOnly/Public/ToolTargets/SkeletalMeshToolTarget.h
@@ -35,19 +35,19 @@ public:
 	virtual bool IsValid() const override;
 
 	// IMeshDescriptionProvider implementation
-	const FMeshDescription* GetMeshDescription() override;
+	virtual const FMeshDescription* GetMeshDescription(const FGetMeshParameters& GetMeshParams = FGetMeshParameters()) override;
 
 	// IMaterialProvider implementation
-	int32 GetNumMaterials() const override;
-	UMaterialInterface* GetMaterial(int32 MaterialIndex) const override;
-	void GetMaterialSet(FComponentMaterialSet& MaterialSetOut, bool bPreferAssetMaterials) const override;
-	bool CommitMaterialSetUpdate(const FComponentMaterialSet& MaterialSet, bool bApplyToAsset) override;	
+	virtual int32 GetNumMaterials() const override;
+	virtual UMaterialInterface* GetMaterial(int32 MaterialIndex) const override;
+	virtual void GetMaterialSet(FComponentMaterialSet& MaterialSetOut, bool bPreferAssetMaterials) const override;
+	virtual bool CommitMaterialSetUpdate(const FComponentMaterialSet& MaterialSet, bool bApplyToAsset) override;	
 
 	// IDynamicMeshProvider
 	virtual UE::Geometry::FDynamicMesh3 GetDynamicMesh() override;
 
 	// ISkeletalMeshBackedTarget implementation
-	USkeletalMesh* GetSkeletalMesh() const override;
+	virtual USkeletalMesh* GetSkeletalMesh() const override;
 
 protected:
 	USkeletalMesh* SkeletalMesh = nullptr;
@@ -82,7 +82,7 @@ class MODELINGCOMPONENTSEDITORONLY_API USkeletalMeshToolTarget :
 
 public:
 	// IMeshDescriptionCommitter implementation
-	void CommitMeshDescription(const FCommitter& Committer) override;
+	virtual void CommitMeshDescription(const FCommitter& Committer, const FCommitMeshParameters& CommitMeshParams = FCommitMeshParameters()) override;
 	using IMeshDescriptionCommitter::CommitMeshDescription; // unhide the other overload
 
 	// IDynamicMeshCommitter
@@ -108,9 +108,9 @@ class MODELINGCOMPONENTSEDITORONLY_API USkeletalMeshReadOnlyToolTargetFactory :
 
 public:
 
-	bool CanBuildTarget(UObject* SourceObject, const FToolTargetTypeRequirements& TargetTypeInfo) const override;
+	virtual bool CanBuildTarget(UObject* SourceObject, const FToolTargetTypeRequirements& TargetTypeInfo) const override;
 
-	UToolTarget* BuildTarget(UObject* SourceObject, const FToolTargetTypeRequirements& TargetTypeInfo) override;
+	virtual UToolTarget* BuildTarget(UObject* SourceObject, const FToolTargetTypeRequirements& TargetTypeInfo) override;
 };
 
 
@@ -122,7 +122,7 @@ class MODELINGCOMPONENTSEDITORONLY_API USkeletalMeshToolTargetFactory : public U
 
 public:
 
-	bool CanBuildTarget(UObject* SourceObject, const FToolTargetTypeRequirements& TargetTypeInfo) const override;
+	virtual bool CanBuildTarget(UObject* SourceObject, const FToolTargetTypeRequirements& TargetTypeInfo) const override;
 
-	UToolTarget* BuildTarget(UObject* SourceObject, const FToolTargetTypeRequirements& TargetTypeInfo) override;
+	virtual UToolTarget* BuildTarget(UObject* SourceObject, const FToolTargetTypeRequirements& TargetTypeInfo) override;
 };
diff --git a/Engine/Plugins/Runtime/MeshModelingToolset/Source/ModelingComponentsEditorOnly/Public/ToolTargets/StaticMeshComponentToolTarget.h b/Engine/Plugins/Runtime/MeshModelingToolset/Source/ModelingComponentsEditorOnly/Public/ToolTargets/StaticMeshComponentToolTarget.h
--- a/Engine/Plugins/Runtime/MeshModelingToolset/Source/ModelingComponentsEditorOnly/Public/ToolTargets/StaticMeshComponentToolTarget.h
+++ b/Engine/Plugins/Runtime/MeshModelingToolset/Source/ModelingComponentsEditorOnly/Public/ToolTargets/StaticMeshComponentToolTarget.h
@@ -12,7 +12,7 @@
 #include "TargetInterfaces/StaticMeshBackedTarget.h"
 #include "TargetInterfaces/PhysicsDataSource.h"
 #include "ToolTargets/PrimitiveComponentToolTarget.h"
-#include "ComponentSourceInterfaces.h"  // for EStaticMeshEditingLOD
+#include "ComponentSourceInterfaces.h"  // for EMeshLODIdentifier
 
 #include "StaticMeshComponentToolTarget.generated.h"
 
@@ -34,10 +34,10 @@ public:
 	 * Configure active LOD to edit. Can only call this after Component is configured in base UPrimitiveComponentToolTarget.
 	 * If requested LOD does not exist, fall back to one that does.
 	 */
-	virtual void SetEditingLOD(EStaticMeshEditingLOD RequestedEditingLOD);
+	virtual void SetEditingLOD(EMeshLODIdentifier RequestedEditingLOD);
 
 	/** @return current editing LOD */
-	virtual EStaticMeshEditingLOD GetEditingLOD() const { return EditingLOD; }
+	virtual EMeshLODIdentifier GetEditingLOD() const { return EditingLOD; }
 
 public:
 	virtual bool IsValid() const override;
@@ -45,11 +45,11 @@ public:
 
 public:
 	// IMeshDescriptionProvider implementation
-	const FMeshDescription* GetMeshDescription() override;
-	virtual void CalculateAutoGeneratedAttributes(FMeshDescription& MeshDescription) override;
+	virtual const FMeshDescription* GetMeshDescription(const FGetMeshParameters& GetMeshParams = FGetMeshParameters()) override;
+	virtual FMeshDescription GetMeshDescriptionCopy(const FGetMeshParameters& GetMeshParams) override;
 
 	// IMeshDescritpionCommitter implementation
-	virtual void CommitMeshDescription(const FCommitter& Committer) override;
+	virtual void CommitMeshDescription(const FCommitter& Committer, const FCommitMeshParameters& CommitParams = FCommitMeshParameters()) override;
 	using IMeshDescriptionCommitter::CommitMeshDescription; // unhide the other overload
 
 	// IDynamicMeshProvider
@@ -60,13 +60,13 @@ public:
 	using IDynamicMeshCommitter::CommitDynamicMesh; // unhide the other overload
 
 	// IMaterialProvider implementation
-	int32 GetNumMaterials() const override;
-	UMaterialInterface* GetMaterial(int32 MaterialIndex) const override;
-	void GetMaterialSet(FComponentMaterialSet& MaterialSetOut, bool bPreferAssetMaterials) const override;
+	virtual int32 GetNumMaterials() const override;
+	virtual UMaterialInterface* GetMaterial(int32 MaterialIndex) const override;
+	virtual void GetMaterialSet(FComponentMaterialSet& MaterialSetOut, bool bPreferAssetMaterials) const override;
 	virtual bool CommitMaterialSetUpdate(const FComponentMaterialSet& MaterialSet, bool bApplyToAsset) override;	
 
 	// IStaticMeshBackedTarget
-	UStaticMesh* GetStaticMesh() const override;
+	virtual UStaticMesh* GetStaticMesh() const override;
 
 	// IPhysicsDataSource
 	virtual UBodySetup* GetBodySetup() const override;
@@ -77,7 +77,7 @@ public:
 public:
 
 protected:
-	EStaticMeshEditingLOD EditingLOD = EStaticMeshEditingLOD::LOD0;
+	EMeshLODIdentifier EditingLOD = EMeshLODIdentifier::LOD0;
 
 	friend class UStaticMeshComponentToolTargetFactory;
 };
@@ -90,18 +90,15 @@ class MODELINGCOMPONENTSEDITORONLY_API UStaticMeshComponentToolTargetFactory : p
 	GENERATED_BODY()
 
 public:
-
 	virtual bool CanBuildTarget(UObject* SourceObject, const FToolTargetTypeRequirements& TargetTypeInfo) const override;
 
 	virtual UToolTarget* BuildTarget(UObject* SourceObject, const FToolTargetTypeRequirements& TargetTypeInfo) override;
 
-
-
 public:
-	virtual EStaticMeshEditingLOD GetActiveEditingLOD() const { return EditingLOD; }
-	virtual void SetActiveEditingLOD(EStaticMeshEditingLOD NewEditingLOD);
+	virtual EMeshLODIdentifier GetActiveEditingLOD() const { return EditingLOD; }
+	virtual void SetActiveEditingLOD(EMeshLODIdentifier NewEditingLOD);
 
 protected:
 	// LOD to edit, default is to edit LOD0
-	EStaticMeshEditingLOD EditingLOD = EStaticMeshEditingLOD::LOD0;
+	EMeshLODIdentifier EditingLOD = EMeshLODIdentifier::LOD0;
 };
\ No newline at end of file
diff --git a/Engine/Plugins/Runtime/MeshModelingToolset/Source/ModelingComponentsEditorOnly/Public/ToolTargets/StaticMeshToolTarget.h b/Engine/Plugins/Runtime/MeshModelingToolset/Source/ModelingComponentsEditorOnly/Public/ToolTargets/StaticMeshToolTarget.h
--- a/Engine/Plugins/Runtime/MeshModelingToolset/Source/ModelingComponentsEditorOnly/Public/ToolTargets/StaticMeshToolTarget.h
+++ b/Engine/Plugins/Runtime/MeshModelingToolset/Source/ModelingComponentsEditorOnly/Public/ToolTargets/StaticMeshToolTarget.h
@@ -11,7 +11,7 @@
 #include "TargetInterfaces/MeshDescriptionProvider.h"
 #include "TargetInterfaces/StaticMeshBackedTarget.h"
 #include "ToolTargets/PrimitiveComponentToolTarget.h"
-#include "ComponentSourceInterfaces.h"  // for EStaticMeshEditingLOD
+#include "ComponentSourceInterfaces.h"  // for EMeshLODIdentifier
 
 #include "StaticMeshToolTarget.generated.h"
 
@@ -38,19 +38,19 @@ public:
 	 * Configure active LOD to edit. Can only call this after Component is configured in base UPrimitiveComponentToolTarget.
 	 * If requested LOD does not exist, fall back to one that does.
 	 */
-	virtual void SetEditingLOD(EStaticMeshEditingLOD RequestedEditingLOD);
+	virtual void SetEditingLOD(EMeshLODIdentifier RequestedEditingLOD);
 
 	/** @return current editing LOD */
-	virtual EStaticMeshEditingLOD GetEditingLOD() const { return EditingLOD; }
+	virtual EMeshLODIdentifier GetEditingLOD() const { return EditingLOD; }
 
 	// UToolTarget
 	virtual bool IsValid() const override;
 
 	// IMeshDescriptionProvider implementation
-	const FMeshDescription* GetMeshDescription() override;
+	virtual const FMeshDescription* GetMeshDescription(const FGetMeshParameters& GetMeshParams = FGetMeshParameters()) override;
 
 	// IMeshDescritpionCommitter implementation
-	virtual void CommitMeshDescription(const FCommitter& Committer) override;
+	virtual void CommitMeshDescription(const FCommitter& Committer, const FCommitMeshParameters& CommitParams = FCommitMeshParameters()) override;
 	using IMeshDescriptionCommitter::CommitMeshDescription; // unhide the other overload
 
 	// IDynamicMeshProvider
@@ -61,31 +61,31 @@ public:
 	using IDynamicMeshCommitter::CommitDynamicMesh; // unhide the other overload
 
 	// IMaterialProvider implementation
-	int32 GetNumMaterials() const override;
-	UMaterialInterface* GetMaterial(int32 MaterialIndex) const override;
-	void GetMaterialSet(FComponentMaterialSet& MaterialSetOut, bool bPreferAssetMaterials) const override;
+	virtual int32 GetNumMaterials() const override;
+	virtual UMaterialInterface* GetMaterial(int32 MaterialIndex) const override;
+	virtual void GetMaterialSet(FComponentMaterialSet& MaterialSetOut, bool bPreferAssetMaterials) const override;
 	virtual bool CommitMaterialSetUpdate(const FComponentMaterialSet& MaterialSet, bool bApplyToAsset) override;	
 
 	// IStaticMeshBackedTarget
-	UStaticMesh* GetStaticMesh() const override;
+	virtual UStaticMesh* GetStaticMesh() const override;
 
 	// Rest provided by parent class
 
 protected:
 	UStaticMesh* StaticMesh = nullptr;
 
-	EStaticMeshEditingLOD EditingLOD = EStaticMeshEditingLOD::LOD0;
+	EMeshLODIdentifier EditingLOD = EMeshLODIdentifier::LOD0;
 
 	friend class UStaticMeshToolTargetFactory;
 
 	friend class UStaticMeshComponentToolTarget;
 
-	static bool IsValid(const UStaticMesh* StaticMesh, EStaticMeshEditingLOD EditingLOD);
-	static EStaticMeshEditingLOD GetValidEditingLOD(const UStaticMesh* StaticMesh, 
-		EStaticMeshEditingLOD RequestedEditingLOD);
-	static void CommitMeshDescription(UStaticMesh* SkeletalMesh, FMeshDescription* MeshDescription,
-		const FCommitter& Committer, EStaticMeshEditingLOD EditingLODIn);
-	static void GetMaterialSet(const UStaticMesh* SkeletalMesh, 
+	static bool IsValid(const UStaticMesh* StaticMesh, EMeshLODIdentifier EditingLOD);
+	static EMeshLODIdentifier GetValidEditingLOD(const UStaticMesh* StaticMesh, 
+		EMeshLODIdentifier RequestedEditingLOD);
+	static void CommitMeshDescription(UStaticMesh* StaticMesh, 
+		const FCommitter& Committer, EMeshLODIdentifier EditingLODIn);
+	static void GetMaterialSet(const UStaticMesh* StaticMesh, 
 		FComponentMaterialSet& MaterialSetOut, bool bPreferAssetMaterials);
 	static bool CommitMaterialSetUpdate(UStaticMesh* SkeletalMesh, 
 		const FComponentMaterialSet& MaterialSet, bool bApplyToAsset);
@@ -107,10 +107,10 @@ public:
 
 
 public:
-	virtual EStaticMeshEditingLOD GetActiveEditingLOD() const { return EditingLOD; }
-	virtual void SetActiveEditingLOD(EStaticMeshEditingLOD NewEditingLOD);
+	virtual EMeshLODIdentifier GetActiveEditingLOD() const { return EditingLOD; }
+	virtual void SetActiveEditingLOD(EMeshLODIdentifier NewEditingLOD);
 
 protected:
 	// LOD to edit, default is to edit LOD0
-	EStaticMeshEditingLOD EditingLOD = EStaticMeshEditingLOD::LOD0;
+	EMeshLODIdentifier EditingLOD = EMeshLODIdentifier::LOD0;
 };
\ No newline at end of file
diff --git a/Engine/Plugins/Runtime/MeshModelingToolset/Source/ModelingComponentsEditorOnly/Public/ToolTargets/VolumeComponentToolTarget.h b/Engine/Plugins/Runtime/MeshModelingToolset/Source/ModelingComponentsEditorOnly/Public/ToolTargets/VolumeComponentToolTarget.h
--- a/Engine/Plugins/Runtime/MeshModelingToolset/Source/ModelingComponentsEditorOnly/Public/ToolTargets/VolumeComponentToolTarget.h
+++ b/Engine/Plugins/Runtime/MeshModelingToolset/Source/ModelingComponentsEditorOnly/Public/ToolTargets/VolumeComponentToolTarget.h
@@ -58,10 +58,10 @@ public:
 	virtual bool CommitMaterialSetUpdate(const FComponentMaterialSet& MaterialSet, bool bApplyToAsset) override { return false; }
 
 	// IMeshDescriptionProvider implementation
-	const FMeshDescription* GetMeshDescription() override;
+	virtual const FMeshDescription* GetMeshDescription(const FGetMeshParameters& GetMeshParams = FGetMeshParameters()) override;
 
 	// IMeshDescritpionCommitter implementation
-	virtual void CommitMeshDescription(const FCommitter& Committer) override;
+	virtual void CommitMeshDescription(const FCommitter& Committer, const FCommitMeshParameters& CommitMeshParams = FCommitMeshParameters()) override;
 	using IMeshDescriptionCommitter::CommitMeshDescription; // unhide the other overload
 
 	// IPhysicsDataSource implementation
diff --git a/Engine/Source/Editor/UnrealEd/Private/Tools/EditorComponentSourceFactory.cpp b/Engine/Source/Editor/UnrealEd/Private/Tools/EditorComponentSourceFactory.cpp
--- a/Engine/Source/Editor/UnrealEd/Private/Tools/EditorComponentSourceFactory.cpp
+++ b/Engine/Source/Editor/UnrealEd/Private/Tools/EditorComponentSourceFactory.cpp
@@ -21,10 +21,10 @@ static void DisplayCriticalWarningMessage(const FString& Message)
 }
 
 
-FStaticMeshComponentTarget::FStaticMeshComponentTarget(UPrimitiveComponent* Component, EStaticMeshEditingLOD EditingLODIn)
+FStaticMeshComponentTarget::FStaticMeshComponentTarget(UPrimitiveComponent* Component, EMeshLODIdentifier EditingLODIn)
 	: FPrimitiveComponentTarget(Cast<UStaticMeshComponent>(Component))
 {
-	EditingLOD = EStaticMeshEditingLOD::LOD0;
+	EditingLOD = EMeshLODIdentifier::LOD0;
 
 	UStaticMeshComponent* StaticMeshComponent = Cast<UStaticMeshComponent>(Component);
 	if (ensure(StaticMeshComponent != nullptr))
@@ -32,14 +32,14 @@ FStaticMeshComponentTarget::FStaticMeshComponentTarget(UPrimitiveComponent* Comp
 		UStaticMesh* StaticMeshAsset = StaticMeshComponent->GetStaticMesh();
 		if (ensure(StaticMeshAsset != nullptr))
 		{
-			if (EditingLODIn == EStaticMeshEditingLOD::MaxQuality)
+			if (EditingLODIn == EMeshLODIdentifier::MaxQuality)
 			{
-				EditingLOD = StaticMeshAsset->IsHiResMeshDescriptionValid() ? EStaticMeshEditingLOD::HiResSource : EStaticMeshEditingLOD::LOD0;
+				EditingLOD = StaticMeshAsset->IsHiResMeshDescriptionValid() ? EMeshLODIdentifier::HiResSource : EMeshLODIdentifier::LOD0;
 			}
-			else if (EditingLODIn == EStaticMeshEditingLOD::HiResSource)
+			else if (EditingLODIn == EMeshLODIdentifier::HiResSource)
 			{
-				EditingLOD = StaticMeshAsset->IsHiResMeshDescriptionValid() ? EStaticMeshEditingLOD::HiResSource : EStaticMeshEditingLOD::LOD0;
-				if (EditingLOD != EStaticMeshEditingLOD::HiResSource)
+				EditingLOD = StaticMeshAsset->IsHiResMeshDescriptionValid() ? EMeshLODIdentifier::HiResSource : EMeshLODIdentifier::LOD0;
+				if (EditingLOD != EMeshLODIdentifier::HiResSource)
 				{
 					DisplayCriticalWarningMessage(FString(TEXT("HiRes Source selected but not available - Falling Back to LOD0")));
 				}
@@ -51,7 +51,7 @@ FStaticMeshComponentTarget::FStaticMeshComponentTarget(UPrimitiveComponent* Comp
 				if (WantLOD > MaxExistingLOD)
 				{
 					DisplayCriticalWarningMessage(FString::Printf(TEXT("LOD%d Requested but not available - Falling Back to LOD%d"), WantLOD, MaxExistingLOD));
-					EditingLOD = (EStaticMeshEditingLOD)MaxExistingLOD;
+					EditingLOD = (EMeshLODIdentifier)MaxExistingLOD;
 				}
 			}
 		}
@@ -78,7 +78,7 @@ bool FStaticMeshComponentTarget::IsValid() const
 		return false;
 	}
 
-	if (EditingLOD == EStaticMeshEditingLOD::HiResSource)
+	if (EditingLOD == EMeshLODIdentifier::HiResSource)
 	{
 		if (StaticMesh->IsHiResMeshDescriptionValid() == false)
 		{
@@ -98,7 +98,7 @@ FMeshDescription* FStaticMeshComponentTarget::GetMesh()
 	if (ensure(IsValid()))
 	{
 		UStaticMesh* StaticMesh = Cast<UStaticMeshComponent>(Component)->GetStaticMesh();
-		return (EditingLOD == EStaticMeshEditingLOD::HiResSource) ?
+		return (EditingLOD == EMeshLODIdentifier::HiResSource) ?
 			StaticMesh->GetHiResMeshDescription() : StaticMesh->GetMeshDescription((int32)EditingLOD);
 	}
 	return nullptr;
@@ -227,7 +227,7 @@ void FStaticMeshComponentTarget::CommitMesh( const FCommitter& Committer )
 	StaticMesh->SetFlags(RF_Transactional);
 
 	verify(StaticMesh->Modify());
-	if (EditingLOD == EStaticMeshEditingLOD::HiResSource)
+	if (EditingLOD == EMeshLODIdentifier::HiResSource)
 	{
 		verify(StaticMesh->ModifyHiResMeshDescription());
 	}
@@ -241,7 +241,7 @@ void FStaticMeshComponentTarget::CommitMesh( const FCommitter& Committer )
 
 	Committer(CommitParams);
 
-	if (EditingLOD == EStaticMeshEditingLOD::HiResSource)
+	if (EditingLOD == EMeshLODIdentifier::HiResSource)
 	{
 		StaticMesh->CommitHiResMeshDescription();
 	}
diff --git a/Engine/Source/Editor/UnrealEd/Public/Tools/EditorComponentSourceFactory.h b/Engine/Source/Editor/UnrealEd/Public/Tools/EditorComponentSourceFactory.h
--- a/Engine/Source/Editor/UnrealEd/Public/Tools/EditorComponentSourceFactory.h
+++ b/Engine/Source/Editor/UnrealEd/Public/Tools/EditorComponentSourceFactory.h
@@ -4,6 +4,7 @@
 
 #include "CoreMinimal.h"
 #include "ComponentSourceInterfaces.h"
+#include "TargetInterfaces/MeshTargetInterfaceTypes.h"
 #include "Components/PrimitiveComponent.h"
 #include "Components/StaticMeshComponent.h"
 
@@ -13,7 +14,7 @@ class UNREALED_API FStaticMeshComponentTargetFactory : public FComponentTargetFa
 {
 public:
 	// new FStaticMeshComponentTargets returned by Build() will be requested for this LOD
-	EStaticMeshEditingLOD CurrentEditingLOD = EStaticMeshEditingLOD::MaxQuality;
+	EMeshLODIdentifier CurrentEditingLOD = EMeshLODIdentifier::MaxQuality;
 
 	bool CanBuild( UActorComponent* Candidate ) override;
 	TUniquePtr<FPrimitiveComponentTarget> Build( UPrimitiveComponent* PrimitiveComponent ) override;
@@ -25,7 +26,7 @@ class UNREALED_API FStaticMeshComponentTarget : public FPrimitiveComponentTarget
 {
 public:
 
-	FStaticMeshComponentTarget(UPrimitiveComponent* Component, EStaticMeshEditingLOD EditingLOD = EStaticMeshEditingLOD::LOD0);
+	FStaticMeshComponentTarget(UPrimitiveComponent* Component, EMeshLODIdentifier EditingLOD = EMeshLODIdentifier::LOD0);
 
 	virtual bool IsValid() const override;
 
@@ -42,7 +43,7 @@ public:
 
 protected:
 	// LOD to edit, default is to edit LOD0
-	EStaticMeshEditingLOD EditingLOD = EStaticMeshEditingLOD::LOD0;
+	EMeshLODIdentifier EditingLOD = EMeshLODIdentifier::LOD0;
 };
 
 
diff --git a/Engine/Source/Runtime/InteractiveToolsFramework/Public/ComponentSourceInterfaces.h b/Engine/Source/Runtime/InteractiveToolsFramework/Public/ComponentSourceInterfaces.h
--- a/Engine/Source/Runtime/InteractiveToolsFramework/Public/ComponentSourceInterfaces.h
+++ b/Engine/Source/Runtime/InteractiveToolsFramework/Public/ComponentSourceInterfaces.h
@@ -17,27 +17,6 @@ struct FMeshDescription;
 class UMaterialInterface;
 
 
-/**
- * EStaticMeshEditingLOD specifies which LOD of a StaticMesh the FStaticMeshComponentTarget
- * refers to. FStaticMeshComponentTarget will remap requested value to the closest available LOD.
- */
-enum class EStaticMeshEditingLOD
-{
-	LOD0 = 0,
-	LOD1 = 1,
-	LOD2 = 2,
-	LOD3 = 3,
-	LOD4 = 4,
-	LOD5 = 5,
-	LOD6 = 6,
-	LOD7 = 7,
-
-	HiResSource = 20,	// HiRes source mesh is optional, and will fall back to LOD0 if not stored
-	MaxQuality = 21		// use HiRes source mesh if available, or LOD0 otherwise
-};
-
-
-
 /**
  * @deprecated Use tool targets instead, such as UPrimitiveComponentToolTarget
  *
diff --git a/Engine/Source/Runtime/InteractiveToolsFramework/Public/TargetInterfaces/MeshDescriptionCommitter.h b/Engine/Source/Runtime/InteractiveToolsFramework/Public/TargetInterfaces/MeshDescriptionCommitter.h
--- a/Engine/Source/Runtime/InteractiveToolsFramework/Public/TargetInterfaces/MeshDescriptionCommitter.h
+++ b/Engine/Source/Runtime/InteractiveToolsFramework/Public/TargetInterfaces/MeshDescriptionCommitter.h
@@ -6,6 +6,7 @@
 
 #include "MeshDescription.h"
 #include "UObject/Interface.h"
+#include "TargetInterfaces/MeshTargetInterfaceTypes.h"
 
 #include "MeshDescriptionCommitter.generated.h"
 
@@ -24,7 +25,7 @@ public:
 	/**
 	 * Commits the given FMeshDescription.
 	 */
-	virtual bool CommitMeshDescription(const FMeshDescription& Mesh)
+	virtual bool CommitMeshDescription(const FMeshDescription& Mesh, const FCommitMeshParameters& CommitParams = FCommitMeshParameters() )
 	{
 		// It seems reasonable to have this function, but we'll go ahead and give a default implementation
 		// if users want to just implement the other one.
@@ -36,7 +37,7 @@ public:
 				*CommitParams.MeshDescriptionOut = Mesh;
 				bSuccess = true;
 			}
-		});
+		}, CommitParams);
 		return bSuccess;
 	}
 
@@ -44,7 +45,7 @@ public:
 	/**
 	 * Commits the given FMeshDescription.
 	 */
-	virtual bool CommitMeshDescription(FMeshDescription&& Mesh)
+	virtual bool CommitMeshDescription(FMeshDescription&& Mesh, const FCommitMeshParameters& CommitParams = FCommitMeshParameters())
 	{
 		// It seems reasonable to have this function, but we'll go ahead and give a default implementation
 		// if users want to just implement the other one.
@@ -56,7 +57,7 @@ public:
 				*CommitParams.MeshDescriptionOut = MoveTemp(Mesh);
 				bSuccess = true;
 			}
-		});
+		}, CommitParams);
 		return bSuccess;
 	}
 
@@ -81,6 +82,6 @@ protected:
 	* @param Committer A function that takes in const IMeshDescriptionCommitter::FCommitParams&
 	*  and populates the FMeshDescription pointed to by the MeshDescription pointer inside.
 	*/
-	virtual void CommitMeshDescription(const FCommitter& Committer) = 0;
+	virtual void CommitMeshDescription(const FCommitter& Committer, const FCommitMeshParameters& CommitParams = FCommitMeshParameters()) = 0;
 
 };
\ No newline at end of file
diff --git a/Engine/Source/Runtime/InteractiveToolsFramework/Public/TargetInterfaces/MeshDescriptionProvider.h b/Engine/Source/Runtime/InteractiveToolsFramework/Public/TargetInterfaces/MeshDescriptionProvider.h
--- a/Engine/Source/Runtime/InteractiveToolsFramework/Public/TargetInterfaces/MeshDescriptionProvider.h
+++ b/Engine/Source/Runtime/InteractiveToolsFramework/Public/TargetInterfaces/MeshDescriptionProvider.h
@@ -4,11 +4,11 @@
 
 #include "CoreMinimal.h"
 #include "UObject/Interface.h"
+#include "TargetInterfaces/MeshTargetInterfaceTypes.h"
+#include "MeshDescription.h"
 
 #include "MeshDescriptionProvider.generated.h"
 
-struct FMeshDescription;
-
 UINTERFACE()
 class INTERACTIVETOOLSFRAMEWORK_API UMeshDescriptionProvider : public UInterface
 {
@@ -21,19 +21,19 @@ class INTERACTIVETOOLSFRAMEWORK_API IMeshDescriptionProvider
 
 public:
 	/**
-	 * Access the MeshDescriptionavailable through this Provider. Note that this MeshDescription may or may not 
+	 * Access the MeshDescription available through this Provider. Note that this MeshDescription may or may not 
 	 * be owned by the provider and should not be modified directly. Use IMeshDescriptionCommitter for writes.
 	 * @return pointer to MeshDescription 
 	 */
-	virtual const FMeshDescription* GetMeshDescription() = 0;
+	virtual const FMeshDescription* GetMeshDescription(const FGetMeshParameters& GetMeshParams = FGetMeshParameters()) = 0;
+
 
 	/**
-	 * Compute any auto-generated attributes on the input MeshDescription. This is required to support
-	 * auto-computed Normals and Tangents on Static Mesh Assets, as those normals/tangents are computed 
-	 * on-the-fly when generating rendering data, and will be zero in the MeshDescription returned by
-	 * GetMeshDescription(). A no-op default implementation is provided.
+	 * Get a copy of the MeshDescription available through this Provider. 
 	 */
-	virtual void CalculateAutoGeneratedAttributes(FMeshDescription& MeshDescription)
+	virtual FMeshDescription GetMeshDescriptionCopy(const FGetMeshParameters& GetMeshParams)
 	{
+		return *GetMeshDescription(GetMeshParams);
 	}
+
 };
\ No newline at end of file
diff --git a/Engine/Source/Runtime/InteractiveToolsFramework/Public/TargetInterfaces/MeshTargetInterfaceTypes.h b/Engine/Source/Runtime/InteractiveToolsFramework/Public/TargetInterfaces/MeshTargetInterfaceTypes.h
new file mode 100644
--- /dev/null
+++ b/Engine/Source/Runtime/InteractiveToolsFramework/Public/TargetInterfaces/MeshTargetInterfaceTypes.h
@@ -0,0 +1,56 @@
+// Copyright Epic Games, Inc. All Rights Reserved.
+
+#pragma once
+
+
+/**
+* EMeshLODIdentifier is used to identify the LOD of a mesh for reading/writing.
+*/
+enum class EMeshLODIdentifier
+{
+	LOD0 = 0,
+	LOD1 = 1,
+	LOD2 = 2,
+	LOD3 = 3,
+	LOD4 = 4,
+	LOD5 = 5,
+	LOD6 = 6,
+	LOD7 = 7,
+
+	HiResSource = 100,	// HiRes source mesh is optional, and will fall back to LOD0 if not available
+
+	Default = 1000,		// use whatever is the "standard" LOD - generally LOD0
+	MaxQuality = 1001	// use HiRes source mesh if available, or LOD0 otherwise
+};
+typedef EMeshLODIdentifier EStaticMeshEditingLOD;
+
+/**
+ * FGetMeshParameters is used by ToolTarget Interfaces/Implementations that support returning a mesh, to allow clients
+ * to specify various options like a specific LOD of a mesh, which attributes are required, etc.
+ */
+struct INTERACTIVETOOLSFRAMEWORK_API FGetMeshParameters
+{
+	/** If true, RequestLOD specifies which mesh LOD should be returned, if available. If not available, behavior is implementation-dependent. */
+	bool bHaveRequestLOD = false;
+	/** Specify which LOD to fetch from a mesh-containing ToolTarget. Ignored unless bHaveRequestLOD == true. */
+	EMeshLODIdentifier RequestLOD = EMeshLODIdentifier::Default;
+
+	/** 
+	 * If true, returned Mesh should have Tangents available. In some cases (eg StaticMesh) this may require that the 
+	 * ToolTarget implementation computes the Tangents, which may only be possible on a copy of the mesh
+	 */
+	bool bWantMeshTangents = false;
+};
+
+
+/**
+ * FCommitMeshParameters is used by ToolTarget Interfaces/Implementations that support setting a mesh (eg on a StaticMesh Asset),
+ * to allow the client to specify various options
+ */
+struct INTERACTIVETOOLSFRAMEWORK_API FCommitMeshParameters
+{
+	/** If true, TargetLOD specifies which mesh LOD should be set/updated */
+	bool bHaveTargetLOD = false;
+	/** Specify which LOD to set in a mesh-containing ToolTarget. Ignored unless bHaveTargetLOD == true. */
+	EMeshLODIdentifier TargetLOD = EMeshLODIdentifier::Default;
+};
\ No newline at end of file
