MeshModelingToolset: Fix focusing selected actor not working properly in PolyEdit orTriEdit when selection was empty

#rnx

diff --git a/Engine/Plugins/Runtime/MeshModelingToolset/Source/MeshModelingTools/Private/EditMeshPolygonsTool.cpp b/Engine/Plugins/Runtime/MeshModelingToolset/Source/MeshModelingTools/Private/EditMeshPolygonsTool.cpp
--- a/Engine/Plugins/Runtime/MeshModelingToolset/Source/MeshModelingTools/Private/EditMeshPolygonsTool.cpp
+++ b/Engine/Plugins/Runtime/MeshModelingToolset/Source/MeshModelingTools/Private/EditMeshPolygonsTool.cpp
@@ -643,12 +643,12 @@ void UEditMeshPolygonsTool::UpdateGizmoFrame(const FFrame3d* UseFrame)
 
 FBox UEditMeshPolygonsTool::GetWorldSpaceFocusBox()
 {
-	if (SelectionMechanic && SelectionMechanic->HasSelection())
+	if (ensure(SelectionMechanic))
 	{
 		FAxisAlignedBox3d Bounds = SelectionMechanic->GetSelectionBounds(true);
 		return (FBox)Bounds;
 	}
-	return FBox();
+	return FBox(EForceInit::ForceInit);
 }
 
 bool UEditMeshPolygonsTool::GetWorldSpaceFocusPoint(const FRay& WorldRay, FVector& PointOut)
diff --git a/Engine/Plugins/Runtime/MeshModelingToolset/Source/ModelingComponents/Private/Selection/PolygonSelectionMechanic.cpp b/Engine/Plugins/Runtime/MeshModelingToolset/Source/ModelingComponents/Private/Selection/PolygonSelectionMechanic.cpp
--- a/Engine/Plugins/Runtime/MeshModelingToolset/Source/ModelingComponents/Private/Selection/PolygonSelectionMechanic.cpp
+++ b/Engine/Plugins/Runtime/MeshModelingToolset/Source/ModelingComponents/Private/Selection/PolygonSelectionMechanic.cpp
@@ -731,7 +731,14 @@ FAxisAlignedBox3d UPolygonSelectionMechanic::GetSelectionBounds(bool bWorld) con
 	}
 	else 
 	{
-		return Mesh->GetBounds();
+		if (bWorld)
+		{
+			return FAxisAlignedBox3d(Mesh->GetBounds(), TargetTransform);
+		}
+		else
+		{
+			return Mesh->GetBounds();
+		}
 	}
 }
 
