add missing UPROPERTY() to lattice control point mechanic class

diff --git a/Engine/Plugins/Runtime/MeshModelingToolset/Source/ModelingComponents/Public/Mechanics/LatticeControlPointsMechanic.h b/Engine/Plugins/Runtime/MeshModelingToolset/Source/ModelingComponents/Public/Mechanics/LatticeControlPointsMechanic.h
--- a/Engine/Plugins/Runtime/MeshModelingToolset/Source/ModelingComponents/Public/Mechanics/LatticeControlPointsMechanic.h
+++ b/Engine/Plugins/Runtime/MeshModelingToolset/Source/ModelingComponents/Public/Mechanics/LatticeControlPointsMechanic.h
@@ -166,9 +166,10 @@ protected:
 	void ClearHover();
 
 	// Support for selection
+	UPROPERTY()
+	TObjectPtr<URectangleMarqueeMechanic> MarqueeMechanic;
 	TSet<int32> SelectedPointIDs;
 	TSet<int32> PreDragSelection;
-	URectangleMarqueeMechanic* MarqueeMechanic;
 	TArray<int32> CurrentDragSelection;
 	void OnDragRectangleStarted();
 	void OnDragRectangleChanged(const FCameraRectangle& Rectangle);
