ModelingComponents: add AssetUtils/StaticMeshMaterialUtil.* with various utility functions for working with StaticMesh materials
#rnx

diff --git a/Engine/Plugins/Runtime/MeshModelingToolset/Source/ModelingComponents/Private/AssetUtils/StaticMeshMaterialUtil.cpp b/Engine/Plugins/Runtime/MeshModelingToolset/Source/ModelingComponents/Private/AssetUtils/StaticMeshMaterialUtil.cpp
new file mode 100644
--- /dev/null
+++ b/Engine/Plugins/Runtime/MeshModelingToolset/Source/ModelingComponents/Private/AssetUtils/StaticMeshMaterialUtil.cpp
@@ -0,0 +1,173 @@
+// Copyright Epic Games, Inc. All Rights Reserved.
+
+#include "AssetUtils/StaticMeshMaterialUtil.h"
+
+#include "Engine/StaticMesh.h"
+#include "MeshDescription.h"
+#include "StaticMeshAttributes.h"
+
+using namespace UE::AssetUtils;
+
+
+bool UE::AssetUtils::GetStaticMeshLODAssetMaterials(
+	UStaticMesh* StaticMeshAsset,
+	int32 LODIndex,
+	FStaticMeshLODMaterialSetInfo& MaterialInfoOut)
+{
+	if (!StaticMeshAsset)
+	{
+		return false;
+	}
+
+#if WITH_EDITOR
+
+	if (StaticMeshAsset->IsSourceModelValid(LODIndex) == false)
+	{
+		return false;
+	}
+
+	TArray<FStaticMaterial> StaticMaterials = StaticMeshAsset->GetStaticMaterials();
+	MaterialInfoOut.MaterialSlots.Reset();
+	for (FStaticMaterial Mat : StaticMaterials)
+	{
+		MaterialInfoOut.MaterialSlots.Add( FStaticMeshMaterialSlot{ Mat.MaterialInterface, Mat.MaterialSlotName } );
+	}
+
+	const FMeshSectionInfoMap& SectionInfoMap = StaticMeshAsset->GetSectionInfoMap();
+	MaterialInfoOut.LODIndex = LODIndex;
+	MaterialInfoOut.NumSections = SectionInfoMap.GetSectionNumber(LODIndex);
+
+	MaterialInfoOut.SectionSlotIndexes.SetNum(MaterialInfoOut.NumSections);
+	MaterialInfoOut.SectionMaterials.SetNum(MaterialInfoOut.NumSections);
+	for (int32 k = 0; k < MaterialInfoOut.NumSections; ++k)
+	{
+		FMeshSectionInfo SectionInfo = SectionInfoMap.Get(LODIndex, k);
+		if (SectionInfo.MaterialIndex >= 0 && SectionInfo.MaterialIndex < MaterialInfoOut.MaterialSlots.Num())
+		{
+			MaterialInfoOut.SectionSlotIndexes[k] = SectionInfo.MaterialIndex;
+			MaterialInfoOut.SectionMaterials[k] = MaterialInfoOut.MaterialSlots[SectionInfo.MaterialIndex].Material;
+		}
+		else
+		{
+			ensure(false);		// this is *not* supposed to be able to happen! SectionMap is broken...
+			MaterialInfoOut.SectionSlotIndexes[k] = -1;
+			MaterialInfoOut.SectionMaterials[k] = nullptr;
+		}
+	}
+
+	return true;
+
+#else
+	// TODO: how would we handle this for runtime static mesh?
+	return false;
+#endif
+
+}
+
+
+
+bool UE::AssetUtils::GetStaticMeshLODMaterialListBySection(
+	UStaticMesh* StaticMeshAsset,
+	int32 LODIndex,
+	TArray<UMaterialInterface*>& MaterialListOut,
+	TArray<int32>& MaterialIndexOut)
+{
+#if WITH_EDITOR
+	// need valid MeshDescription in Editor path
+	if (StaticMeshAsset->IsMeshDescriptionValid(LODIndex) == false)
+	{
+		return false;
+	}
+#endif
+
+	FStaticMeshLODMaterialSetInfo MaterialSetInfo;
+	if (GetStaticMeshLODAssetMaterials(StaticMeshAsset, LODIndex, MaterialSetInfo) == false)
+	{
+		return false;
+	}
+
+#if WITH_EDITOR
+
+	const FMeshDescription* SourceMesh = StaticMeshAsset->GetMeshDescription(LODIndex);
+
+	// # Sections == # PolygonGroups
+	int32 NumPolygonGroups = SourceMesh->PolygonGroups().Num();
+
+	MaterialListOut.Reset();
+	MaterialIndexOut.Reset();
+	for (int32 k = 0; k < NumPolygonGroups; ++k)
+	{
+		int32 UseSlotIndex = -1;
+		if (k < MaterialSetInfo.SectionSlotIndexes.Num())
+		{
+			UseSlotIndex = MaterialSetInfo.SectionSlotIndexes[k];
+		}
+		else if (MaterialSetInfo.SectionSlotIndexes.Num() > 0)
+		{
+			UseSlotIndex = 0;
+		}
+
+		if (UseSlotIndex >= 0)
+		{
+			MaterialIndexOut.Add(UseSlotIndex);
+			MaterialListOut.Add(MaterialSetInfo.MaterialSlots[UseSlotIndex].Material);
+		}
+		else
+		{
+			MaterialIndexOut.Add(-1);
+			MaterialListOut.Add(nullptr);
+		}
+	}
+
+	return true;
+
+#else
+	// TODO: how would we handle this for runtime static mesh?
+	return false;
+#endif
+
+
+}
+
+
+
+FName UE::AssetUtils::GenerateNewMaterialSlotName(
+	const TArray<FStaticMaterial>& ExistingMaterials,
+	UMaterialInterface* SlotMaterial,
+	int32 NewSlotIndex)
+{
+	FString MaterialName = (SlotMaterial) ? SlotMaterial->GetName() : TEXT("Material");
+	FName BaseName(MaterialName);
+
+	bool bFound = false;
+	for (const FStaticMaterial& Mat : ExistingMaterials)
+	{
+		if (Mat.MaterialSlotName == BaseName || Mat.ImportedMaterialSlotName == BaseName)
+		{
+			bFound = true;
+			break;
+		}
+	}
+	if (bFound == false && SlotMaterial != nullptr)
+	{
+		return BaseName;
+	}
+
+	bFound = true;
+	while (bFound)
+	{
+		bFound = false;
+
+		BaseName = FName(FString::Printf(TEXT("%s_%d"), *MaterialName, NewSlotIndex++));
+		for (const FStaticMaterial& Mat : ExistingMaterials)
+		{
+			if (Mat.MaterialSlotName == BaseName || Mat.ImportedMaterialSlotName == BaseName)
+			{
+				bFound = true;
+				break;
+			}
+		}
+	}
+
+	return BaseName;
+}
\ No newline at end of file
diff --git a/Engine/Plugins/Runtime/MeshModelingToolset/Source/ModelingComponents/Public/AssetUtils/StaticMeshMaterialUtil.h b/Engine/Plugins/Runtime/MeshModelingToolset/Source/ModelingComponents/Public/AssetUtils/StaticMeshMaterialUtil.h
new file mode 100644
--- /dev/null
+++ b/Engine/Plugins/Runtime/MeshModelingToolset/Source/ModelingComponents/Public/AssetUtils/StaticMeshMaterialUtil.h
@@ -0,0 +1,77 @@
+// Copyright Epic Games, Inc. All Rights Reserved.
+
+#pragma once
+
+#include "CoreMinimal.h"
+
+class UMaterialInterface;
+class UStaticMesh;
+
+
+namespace UE
+{
+namespace AssetUtils
+{
+	/**
+	 * Information about a mesh material slot (eg a FStaticMaterial)
+	 */
+	struct MODELINGCOMPONENTS_API FStaticMeshMaterialSlot
+	{
+		UMaterialInterface* Material;
+		FName SlotName;
+	};
+
+	/**
+	 * Information about the materials assigned to a StaticMesh LOD.
+	 * The StaticMesh stores a Material List and per-mesh-Section assignments into that Material List.
+	 * This relationship is not necessarily 1-1, multiple Sections can use the same material, and there
+	 * can be arbitrary remappings, and the storage of this information is somewhat scattered across StaticMesh.
+	 * This data structure encapsulates the critical information needed to (eg) assign materials
+	 * to a new mesh/component such that they match the rendered materials on an existing StaticMesh.
+	 * 
+	 * The GetStaticMeshAssetMaterials() function can be used to build this data structure for a given StaticMesh LOD.
+	 */
+	struct MODELINGCOMPONENTS_API FStaticMeshLODMaterialSetInfo
+	{
+		/** MaterialSlots array is a copy of the UStaticMesh::StaticMaterials array/data */
+		TArray<FStaticMeshMaterialSlot> MaterialSlots;
+
+		/** LOD Index the data below refers to */
+		int32 LODIndex = 0;
+		/** Number of Sections on the mesh LOD */
+		int32 NumSections = 0;
+		/** The index into MaterialSlots array of the Material assigned to each Section, in linear order */
+		TArray<int32> SectionSlotIndexes;
+		/** The Material assigned to each Section, in linear order */
+		TArray<UMaterialInterface*> SectionMaterials;
+	};
+
+	/**
+	 * Extract information about the material set for a given LODIndex of a StaticMeshAsset
+	 */
+	MODELINGCOMPONENTS_API bool GetStaticMeshLODAssetMaterials(
+		UStaticMesh* StaticMeshAsset,
+		int32 LODIndex,
+		FStaticMeshLODMaterialSetInfo& MaterialInfoOut);
+
+	/**
+	 * Construct the linear per-section material list for a given LODIndex of a StaticMeshAsset
+	 * @param MaterialListOut the list of linear per-section indices into the Asset Material List
+	 * @param MaterialIndexOut the corresponding list of linear per-section indices into the Asset Material List
+	 */
+	MODELINGCOMPONENTS_API bool GetStaticMeshLODMaterialListBySection(
+		UStaticMesh* StaticMeshAsset,
+		int32 LODIndex,
+		TArray<UMaterialInterface*>& MaterialListOut,
+		TArray<int32>& MaterialIndexOut);
+
+	/**
+	 * Generate a new unique material slot name for the given SlotMaterial and NewSlotIndex,
+	 * ensuring that the name is not already used in the set of ExistingMaterials
+	 */
+	MODELINGCOMPONENTS_API FName GenerateNewMaterialSlotName(
+		const TArray<FStaticMaterial>& ExistingMaterials,
+		UMaterialInterface* SlotMaterial,
+		int32 NewSlotIndex);
+}
+}
\ No newline at end of file
