MeshModelingTools: Add ability to avoid hitting back faces in selection in PolyEdit, to make it easier to work with inside-out meshes. Also save settings for PolyEdit and TriEdit separately.

#rnx

diff --git a/Engine/Plugins/Runtime/MeshModelingToolset/Source/MeshModelingTools/Private/EditMeshPolygonsTool.cpp b/Engine/Plugins/Runtime/MeshModelingToolset/Source/MeshModelingTools/Private/EditMeshPolygonsTool.cpp
--- a/Engine/Plugins/Runtime/MeshModelingToolset/Source/MeshModelingTools/Private/EditMeshPolygonsTool.cpp
+++ b/Engine/Plugins/Runtime/MeshModelingToolset/Source/MeshModelingTools/Private/EditMeshPolygonsTool.cpp
@@ -62,6 +62,11 @@ namespace EditMeshPolygonsToolLocals
 	FText TriEditDefaultMessage = LOCTEXT("OnStartEditMeshPolygonsTool", "Select PolyGroups to edit mesh. Use middle mouse on gizmo to reposition it. "
 		"Hold Ctrl while translating or (in local mode) rotating to align to scene. Shift and Ctrl change marquee select "
 		"behavior. Q toggles Gizmo Orientation Lock.");
+
+	FString GetPropertyCacheIdentifier(bool bTriangleMode)
+	{
+		return bTriangleMode ? TEXT("TriEditTool") : TEXT("PolyEditTool");
+	}
 }
 
 /*
@@ -231,7 +236,7 @@ void UEditMeshPolygonsTool::Setup()
 
 	// Initialize the common properties but don't add them yet, because we want them to be under the activity-specific ones.
 	CommonProps = NewObject<UPolyEditCommonProperties>(this);
-	CommonProps->RestoreProperties(this);
+	CommonProps->RestoreProperties(this, GetPropertyCacheIdentifier(bTriangleMode));
 	CommonProps->WatchProperty(CommonProps->LocalFrameMode,
 		[this](ELocalFrameMode) { UpdateGizmoFrame(); });
 	CommonProps->WatchProperty(CommonProps->bLockRotation,
@@ -293,7 +298,7 @@ void UEditMeshPolygonsTool::Setup()
 	SelectionMechanic = NewObject<UPolygonSelectionMechanic>(this);
 	SelectionMechanic->bAddSelectionFilterPropertiesToParentTool = false; // We'll do this ourselves later
 	SelectionMechanic->Setup(this);
-	SelectionMechanic->Properties->RestoreProperties(this);
+	SelectionMechanic->Properties->RestoreProperties(this, GetPropertyCacheIdentifier(bTriangleMode));
 	SelectionMechanic->OnSelectionChanged.AddUObject(this, &UEditMeshPolygonsTool::OnSelectionModifiedEvent);
 	if (bTriangleMode)
 	{
@@ -442,12 +447,15 @@ bool UEditMeshPolygonsTool::IsToolInputSelectionUsable(const UPersistentMeshSele
 
 void UEditMeshPolygonsTool::Shutdown(EToolShutdownType ShutdownType)
 {
+	using namespace EditMeshPolygonsToolLocals;
+
 	if (CurrentActivity)
 	{
 		CurrentActivity->End(ShutdownType);
 		CurrentActivity = nullptr;
 	}
-	CommonProps->SaveProperties(this);
+	CommonProps->SaveProperties(this, GetPropertyCacheIdentifier(bTriangleMode));
+	SelectionMechanic->Properties->SaveProperties(this, GetPropertyCacheIdentifier(bTriangleMode));
 
 	GetToolManager()->GetContextObjectStore()->RemoveContextObjectsOfType(UPolyEditActivityContext::StaticClass());
 	ActivityContext = nullptr;
diff --git a/Engine/Plugins/Runtime/MeshModelingToolset/Source/MeshModelingTools/Private/ToolActivities/PolyEditInsertEdgeActivity.cpp b/Engine/Plugins/Runtime/MeshModelingToolset/Source/MeshModelingTools/Private/ToolActivities/PolyEditInsertEdgeActivity.cpp
--- a/Engine/Plugins/Runtime/MeshModelingToolset/Source/MeshModelingTools/Private/ToolActivities/PolyEditInsertEdgeActivity.cpp
+++ b/Engine/Plugins/Runtime/MeshModelingToolset/Source/MeshModelingTools/Private/ToolActivities/PolyEditInsertEdgeActivity.cpp
@@ -627,6 +627,7 @@ bool UPolyEditInsertEdgeActivity::TopologyHitTest(const FRay& WorldRay,
 
 	FGroupTopologySelection Selection;
 	FVector3d Position, Normal;
+	TopologySelectorSettings.bHitBackFaces = ActivityContext->SelectionMechanic->Properties->bHitBackFaces;
 	if (TopologySelector->FindSelectedElement(TopologySelectorSettings,
 		LocalRay, Selection, Position, Normal))
 	{
@@ -654,6 +655,7 @@ bool UPolyEditInsertEdgeActivity::GetHoveredItem(const FRay& WorldRay,
 	FGroupTopologySelection Selection;
 	FVector3d Position, Normal;
 	int32 EdgeSegmentID;
+	TopologySelectorSettings.bHitBackFaces = ActivityContext->SelectionMechanic->Properties->bHitBackFaces;
 	if (!TopologySelector->FindSelectedElement(
 		TopologySelectorSettings, LocalRay, Selection, Position, Normal, &EdgeSegmentID))
 	{
diff --git a/Engine/Plugins/Runtime/MeshModelingToolset/Source/MeshModelingTools/Private/ToolActivities/PolyEditInsertEdgeLoopActivity.cpp b/Engine/Plugins/Runtime/MeshModelingToolset/Source/MeshModelingTools/Private/ToolActivities/PolyEditInsertEdgeLoopActivity.cpp
--- a/Engine/Plugins/Runtime/MeshModelingToolset/Source/MeshModelingTools/Private/ToolActivities/PolyEditInsertEdgeLoopActivity.cpp
+++ b/Engine/Plugins/Runtime/MeshModelingToolset/Source/MeshModelingTools/Private/ToolActivities/PolyEditInsertEdgeLoopActivity.cpp
@@ -355,6 +355,7 @@ FInputRayHit UPolyEditInsertEdgeLoopActivity::HitTest(const FRay& WorldRay)
 		(FVector3d)TargetTransform.InverseTransformVector(WorldRay.Direction), false);
 	FGroupTopologySelection Selection;
 	FVector3d Position, Normal;
+	TopologySelectorSettings.bHitBackFaces = ActivityContext->SelectionMechanic->Properties->bHitBackFaces;
 	if (TopologySelector->FindSelectedElement(
 		TopologySelectorSettings, LocalRay, Selection, Position, Normal))
 	{
@@ -381,6 +382,7 @@ bool UPolyEditInsertEdgeLoopActivity::UpdateHoveredItem(const FRay& WorldRay)
 	FGroupTopologySelection Selection;
 	FVector3d Position, Normal;
 	int32 EdgeSegmentID;
+	TopologySelectorSettings.bHitBackFaces = ActivityContext->SelectionMechanic->Properties->bHitBackFaces;
 	if (!TopologySelector->FindSelectedElement(
 		TopologySelectorSettings, LocalRay, Selection, Position, Normal, &EdgeSegmentID))
 	{
diff --git a/Engine/Plugins/Runtime/MeshModelingToolset/Source/ModelingComponents/Private/Selection/GroupTopologySelector.cpp b/Engine/Plugins/Runtime/MeshModelingToolset/Source/ModelingComponents/Private/Selection/GroupTopologySelector.cpp
--- a/Engine/Plugins/Runtime/MeshModelingToolset/Source/ModelingComponents/Private/Selection/GroupTopologySelector.cpp
+++ b/Engine/Plugins/Runtime/MeshModelingToolset/Source/ModelingComponents/Private/Selection/GroupTopologySelector.cpp
@@ -10,8 +10,8 @@
 using namespace UE::Geometry;
 
 // Local utility function forward declarations
-bool IsOccluded(const FGeometrySet3::FNearest& ClosestElement, const FVector3d& ViewOrigin, const FDynamicMeshAABBTree3* Spatial);
-bool IsOccluded(const FVector3d& Point, const FVector3d& ViewOrigin, const FDynamicMeshAABBTree3* Spatial);
+bool IsOccluded(const FGeometrySet3::FNearest& ClosestElement, const FVector3d& ViewOrigin, const FDynamicMeshAABBTree3* Spatial, bool bBackFacesOcclude);
+bool IsOccluded(const FVector3d& Point, const FVector3d& ViewOrigin, const FDynamicMeshAABBTree3* Spatial, bool bBackFacesOcclude);
 void AddNewEdgeLoopEdgesFromCorner(const FGroupTopology& Topology, int32 EdgeID, int32 CornerID, TSet<int32>& EdgeSet);
 bool GetNextEdgeLoopEdge(const FGroupTopology& Topology, int32 IncomingEdgeID, int32 CornerID, int32& NextEdgeIDOut);
 void AddNewEdgeRingEdges(const FGroupTopology& Topology, int32 StartEdgeID, int32 ForwardGroupID, TSet<int32>& EdgeSet);
@@ -101,13 +101,21 @@ bool FGroupTopologySelector::FindSelectedElement(const FSelectionSettings& Setti
 	FDynamicMeshAABBTree3* Spatial = GetSpatial();
 	const FGeometrySet3& TopoSpatial = GetGeometrySet();
 
+	IMeshSpatial::FQueryOptions SpatialQueryOptions;
+	if (Spatial && !Settings.bHitBackFaces)
+	{
+		SpatialQueryOptions.TriangleFilterF = [Spatial, &Ray](int32 Tid) {
+			return Spatial->GetMesh()->GetTriNormal(Tid).Dot(Ray.Direction) < 0;
+		};
+	}
+
 	// We start by intersecting with the mesh triangles because even when selecting corners or edges, we set
 	// the normal based on the true triangle that we hit. If we end up with a simple face selection, we will
 	// end up using this result.
 	double RayParameter = -1;
 	int HitTriangleID = IndexConstants::InvalidID;
 	FVector3d TriangleHitPos;
-	bool bActuallyHitSurface = (Spatial != nullptr) ? Spatial->FindNearestHitTriangle(Ray, RayParameter, HitTriangleID) : false;
+	bool bActuallyHitSurface = (Spatial != nullptr) ? Spatial->FindNearestHitTriangle(Ray, RayParameter, HitTriangleID, SpatialQueryOptions) : false;
 	if (bActuallyHitSurface)
 	{
 		TriangleHitPos = Ray.PointAt(RayParameter);
@@ -247,7 +255,7 @@ bool FGroupTopologySelector::DoCornerBasedSelection(const FSelectionSettings& Se
 	}
 
 	// Also bail if the closest element is not visible.
-	if (!Settings.bIgnoreOcclusion && IsOccluded(*ClosestElement, Ray.Origin, Spatial))
+	if (!Settings.bIgnoreOcclusion && IsOccluded(*ClosestElement, Ray.Origin, Spatial, Settings.bHitBackFaces))
 	{
 		return false;
 	}
@@ -413,7 +421,7 @@ bool FGroupTopologySelector::DoEdgeBasedSelection(const FSelectionSettings& Sett
 	}
 
 	// Also bail if the closest element is not visible.
-	if (!Settings.bIgnoreOcclusion && IsOccluded(*ClosestElement, Ray.Origin, Spatial))
+	if (!Settings.bIgnoreOcclusion && IsOccluded(*ClosestElement, Ray.Origin, Spatial, Settings.bHitBackFaces))
 	{
 		return false;
 	}
@@ -565,17 +573,31 @@ bool FGroupTopologySelector::DoEdgeBasedSelection(const FSelectionSettings& Sett
 	return false;
 }
 
-bool IsOccluded(const FGeometrySet3::FNearest& ClosestElement, const FVector3d& ViewOrigin, const FDynamicMeshAABBTree3* Spatial)
+bool IsOccluded(const FGeometrySet3::FNearest& ClosestElement, const FVector3d& ViewOrigin, 
+	const FDynamicMeshAABBTree3* Spatial, bool bBackFacesOcclude)
 {
-	return IsOccluded(ClosestElement.NearestGeoPoint, ViewOrigin, Spatial);
+	return IsOccluded(ClosestElement.NearestGeoPoint, ViewOrigin, Spatial, bBackFacesOcclude);
 }
 
-bool IsOccluded(const FVector3d& Point, const FVector3d& ViewOrigin, const FDynamicMeshAABBTree3* Spatial)
+bool IsOccluded(const FVector3d& Point, const FVector3d& ViewOrigin, 
+	const FDynamicMeshAABBTree3* Spatial, bool bBackFacesOcclude)
 {
-	// Shoot ray backwards to see if we hit something. 
+	// We will shoot a ray backwards to see if we hit something. 
 	FRay3d ToEyeRay(Point, UE::Geometry::Normalized(ViewOrigin - Point), true);
 	ToEyeRay.Origin += (double)(100 * FMathf::ZeroTolerance) * ToEyeRay.Direction;
-	if (Spatial->FindNearestHitTriangle(ToEyeRay) >= 0)
+
+	IMeshSpatial::FQueryOptions QueryOptions;
+	if (!bBackFacesOcclude)
+	{
+		// If back faces don't occlude (ie we weren't hitting back faces on the way forward), 
+		// then the direction toward the camera should be same as normal direction.
+		QueryOptions.TriangleFilterF = [Spatial, &ToEyeRay](int32 Tid)
+		{
+			return Spatial->GetMesh()->GetTriNormal(Tid).Dot(ToEyeRay.Direction) > 0;
+		};
+	}
+
+	if (Spatial->FindNearestHitTriangle(ToEyeRay, QueryOptions) >= 0)
 	{
 		return true;
 	}
@@ -605,7 +627,8 @@ bool FGroupTopologySelector::FindSelectedElement(const FSelectionSettings& Setti
 			{
 				FVector TransformedPoint = TargetTransform.TransformPosition(PointPosition);
 				return CameraRectangle.IsProjectedPointInRectangle(TransformedPoint)
-					&& (Settings.bIgnoreOcclusion || !IsOccluded(PointPosition, LocalCameraOrigin, Spatial));
+					&& (Settings.bIgnoreOcclusion || !IsOccluded(PointPosition, LocalCameraOrigin, 
+						Spatial, Settings.bHitBackFaces));
 			},
 			ResultOut.SelectedCornerIDs);
 	}
@@ -619,8 +642,9 @@ bool FGroupTopologySelector::FindSelectedElement(const FSelectionSettings& Setti
 				// Testing occlusion properly seems like it would be a pain, so we'll just consider something occluded
 				// if one of the endpoints is occluded. This will handle the common case of not wanting to select a
 				// hidden edge that is connected to a visible corner.
-				if (!Settings.bIgnoreOcclusion && (IsOccluded(Curve.Start(), LocalCameraOrigin, Spatial)
-						|| IsOccluded(Curve.End(), LocalCameraOrigin, Spatial)))
+				if (!Settings.bIgnoreOcclusion 
+					&& (IsOccluded(Curve.Start(), LocalCameraOrigin, Spatial, Settings.bHitBackFaces)
+						|| IsOccluded(Curve.End(), LocalCameraOrigin, Spatial, Settings.bHitBackFaces)))
 				{
 					return false;
 				}
diff --git a/Engine/Plugins/Runtime/MeshModelingToolset/Source/ModelingComponents/Private/Selection/PolygonSelectionMechanic.cpp b/Engine/Plugins/Runtime/MeshModelingToolset/Source/ModelingComponents/Private/Selection/PolygonSelectionMechanic.cpp
--- a/Engine/Plugins/Runtime/MeshModelingToolset/Source/ModelingComponents/Private/Selection/PolygonSelectionMechanic.cpp
+++ b/Engine/Plugins/Runtime/MeshModelingToolset/Source/ModelingComponents/Private/Selection/PolygonSelectionMechanic.cpp
@@ -344,6 +344,7 @@ FGroupTopologySelector::FSelectionSettings UPolygonSelectionMechanic::GetTopoSel
 	Settings.bEnableFaceHits = Properties->bSelectFaces;
 	Settings.bEnableEdgeHits = Properties->bSelectEdges;
 	Settings.bEnableCornerHits = Properties->bSelectVertices;
+	Settings.bHitBackFaces = Properties->bHitBackFaces;
 
 	if (!PersistentSelection.IsEmpty() && (ShouldAddToSelectionFunc() || ShouldRemoveFromSelectionFunc()))
 	{
diff --git a/Engine/Plugins/Runtime/MeshModelingToolset/Source/ModelingComponents/Public/Selection/GroupTopologySelector.h b/Engine/Plugins/Runtime/MeshModelingToolset/Source/ModelingComponents/Public/Selection/GroupTopologySelector.h
--- a/Engine/Plugins/Runtime/MeshModelingToolset/Source/ModelingComponents/Public/Selection/GroupTopologySelector.h
+++ b/Engine/Plugins/Runtime/MeshModelingToolset/Source/ModelingComponents/Public/Selection/GroupTopologySelector.h
@@ -50,6 +50,11 @@ public:
 		bool bEnableEdgeHits = true;
 		bool bEnableCornerHits = true;
 		
+		// When false, hits of a triangles that are facing away from the camera are not considered. This can
+		// be useful when working with inside-out meshes, where we usually want to select the farther, visible
+		// wall.
+		bool bHitBackFaces = true;
+
 		// The following are mainly useful for ortho viewport selection:
 
 		// Prefer an edge projected to a point rather than the point, and a face projected to an edge
diff --git a/Engine/Plugins/Runtime/MeshModelingToolset/Source/ModelingComponents/Public/Selection/PolygonSelectionMechanic.h b/Engine/Plugins/Runtime/MeshModelingToolset/Source/ModelingComponents/Public/Selection/PolygonSelectionMechanic.h
--- a/Engine/Plugins/Runtime/MeshModelingToolset/Source/ModelingComponents/Public/Selection/PolygonSelectionMechanic.h
+++ b/Engine/Plugins/Runtime/MeshModelingToolset/Source/ModelingComponents/Public/Selection/PolygonSelectionMechanic.h
@@ -47,6 +47,10 @@ public:
 	UPROPERTY(EditAnywhere, Category = SelectionFilter, meta = (EditCondition = "bSelectEdges"))
 	bool bSelectEdgeRings = false;
 
+	/** When false, faces that face away from the camera are ignored in selection and occlusion. Useful for working with inside-out meshes. */
+	UPROPERTY(EditAnywhere, Category = SelectionFilter, AdvancedDisplay)
+	bool bHitBackFaces = true;
+
 	UPROPERTY(EditAnywhere, Category = SelectionFilter, AdvancedDisplay)
 	bool bEnableMarquee = true;
 
