Subdiv tool: interpolate Material IDs from input mesh to subdivided mesh

#rnx

diff --git a/Engine/Plugins/Runtime/MeshModelingToolset/Source/ModelingComponentsEditorOnly/Private/Operations/SubdividePoly.cpp b/Engine/Plugins/Runtime/MeshModelingToolset/Source/ModelingComponentsEditorOnly/Private/Operations/SubdividePoly.cpp
--- a/Engine/Plugins/Runtime/MeshModelingToolset/Source/ModelingComponentsEditorOnly/Private/Operations/SubdividePoly.cpp
+++ b/Engine/Plugins/Runtime/MeshModelingToolset/Source/ModelingComponentsEditorOnly/Private/Operations/SubdividePoly.cpp
@@ -730,6 +730,76 @@ bool FSubdividePoly::ComputeSubdividedMesh(FDynamicMesh3& OutMesh)
 		}
 	}
 
+	//
+	// Interpolate material IDs
+	//
+
+	const bool bHasMaterialIDs = OriginalMesh.HasAttributes() && OriginalMesh.Attributes()->HasMaterialID();
+
+	TArray<int> RefinedMaterialIDs;
+	if (bHasMaterialIDs)
+	{
+		const FDynamicMeshMaterialAttribute* MaterialIDs = OriginalMesh.Attributes()->GetMaterialID();
+
+		// Find the most common material ID for a given group
+		auto MaterialIDForGroup = [this, MaterialIDs](const FGroupTopology::FGroup& Group) -> int32
+		{
+			TMap<int32, int> MaterialIDVotes;
+			for (int32 TriangleID : Group.Triangles)
+			{
+				int32 MatID = MaterialIDs->GetValue(TriangleID);
+
+				if (MaterialIDVotes.Contains(MatID))
+				{
+					++MaterialIDVotes[MatID];
+				}
+				else
+				{
+					MaterialIDVotes.Add(MatID, 1);
+				}
+			}
+
+			int MaxVotes = -1;
+			int32 MaxVoteMaterialID = 0;
+			for (const TPair<int32, int>& IDVotePair : MaterialIDVotes)
+			{
+				if (IDVotePair.Value > MaxVotes)
+				{
+					MaxVotes = IDVotePair.Value;
+					MaxVoteMaterialID = IDVotePair.Key;
+				}
+			}
+
+			return MaxVoteMaterialID;
+		};
+
+		TArray<int32> SourceMaterialIDs;
+		if (SubdivisionScheme == ESubdivisionScheme::Loop)
+		{
+			for (int TriangleID : OriginalMesh.TriangleIndicesItr())
+			{
+				SourceMaterialIDs.Add(MaterialIDs->GetValue(TriangleID));
+			}
+		}
+		else
+		{
+			for (const FGroupTopology::FGroup& Group : GroupTopology.Groups)
+			{
+				SourceMaterialIDs.Add(MaterialIDForGroup(Group));
+			}
+		}
+
+		for (int CurrentLevel = 1; CurrentLevel <= Level; ++CurrentLevel)
+		{
+			// TODO: Don't keep resizing -- preallocate one big buffer and move through it
+			RefinedMaterialIDs.SetNumUninitialized(Refiner->TopologyRefiner->GetLevel(CurrentLevel).GetNumFaces());
+			int* s = SourceMaterialIDs.GetData();
+			int* d = RefinedMaterialIDs.GetData();
+			Interpolator.InterpolateFaceUniform(CurrentLevel, s, d);
+			SourceMaterialIDs = RefinedMaterialIDs;
+		}
+	}
+
 
 	//
 	// Interpolate UVs
@@ -784,7 +854,8 @@ bool FSubdividePoly::ComputeSubdividedMesh(FDynamicMesh3& OutMesh)
 	{
 		OutMesh.EnableVertexNormals(FVector3f{ 0,0,0 });
 	}
-	if ((NormalComputationMethod != ESubdivisionOutputNormals::None) || (UVComputationMethod != ESubdivisionOutputUVs::None))
+
+	if ((NormalComputationMethod != ESubdivisionOutputNormals::None) || (UVComputationMethod != ESubdivisionOutputUVs::None) || bHasMaterialIDs)
 	{
 		OutMesh.EnableAttributes();
 	}
@@ -820,6 +891,11 @@ bool FSubdividePoly::ComputeSubdividedMesh(FDynamicMesh3& OutMesh)
 		}
 	};
 
+	if (bHasMaterialIDs)
+	{
+		OutMesh.Attributes()->EnableMaterialID();
+	}
+
 	for (int FaceID = 0; FaceID < FinalLevel.GetNumFaces(); ++FaceID)
 	{
 		int GroupID = bNewPolyGroups ? OutMesh.AllocateTriangleGroup() : RefinedGroupIDs[FaceID];
@@ -831,8 +907,8 @@ bool FSubdividePoly::ComputeSubdividedMesh(FDynamicMesh3& OutMesh)
 		{
 			FIndex3i TriA{ Face[0], Face[1], Face[3] };
 			FIndex3i TriB{ Face[2], Face[3], Face[1] };
-			OutMesh.AppendTriangle(TriA, GroupID);
-			OutMesh.AppendTriangle(TriB, GroupID);
+			int TriAIndex = OutMesh.AppendTriangle(TriA, GroupID);
+			int TriBIndex = OutMesh.AppendTriangle(TriB, GroupID);
 
 			if (UVComputationMethod == ESubdivisionOutputUVs::Interpolated)
 			{
@@ -844,11 +920,17 @@ bool FSubdividePoly::ComputeSubdividedMesh(FDynamicMesh3& OutMesh)
 				FIndex3i UVTriB{ FaceUVIndices[2], FaceUVIndices[3], FaceUVIndices[1] };
 				AddUVTriangleIfValid(UVTriB);
 			}
+
+			if (bHasMaterialIDs)
+			{
+				OutMesh.Attributes()->GetMaterialID()->SetValue(TriAIndex, RefinedMaterialIDs[FaceID]);
+				OutMesh.Attributes()->GetMaterialID()->SetValue(TriBIndex, RefinedMaterialIDs[FaceID]);
+			}
 		}
 		else
 		{
 			check(Face.size() == 3);
-			OutMesh.AppendTriangle(FIndex3i{ Face[0], Face[1], Face[2] }, GroupID);
+			int TriIndex = OutMesh.AppendTriangle(FIndex3i{ Face[0], Face[1], Face[2] }, GroupID);
 
 			if (UVComputationMethod == ESubdivisionOutputUVs::Interpolated)
 			{
@@ -856,6 +938,11 @@ bool FSubdividePoly::ComputeSubdividedMesh(FDynamicMesh3& OutMesh)
 				FIndex3i UVTri{ FaceUVIndices[0], FaceUVIndices[1], FaceUVIndices[2] };
 				AddUVTriangleIfValid(UVTri);
 			}
+
+			if (bHasMaterialIDs)
+			{
+				OutMesh.Attributes()->GetMaterialID()->SetValue(TriIndex, RefinedMaterialIDs[FaceID]);
+			}
 		}
 	}
 
