MeshModelingTools: Change loop variable to reference to remove compiler warning on mac builds.

#rnx

diff --git a/Engine/Plugins/Runtime/MeshModelingToolset/Source/ModelingComponents/Private/Selection/GroupTopologySelector.cpp b/Engine/Plugins/Runtime/MeshModelingToolset/Source/ModelingComponents/Private/Selection/GroupTopologySelector.cpp
--- a/Engine/Plugins/Runtime/MeshModelingToolset/Source/ModelingComponents/Private/Selection/GroupTopologySelector.cpp
+++ b/Engine/Plugins/Runtime/MeshModelingToolset/Source/ModelingComponents/Private/Selection/GroupTopologySelector.cpp
@@ -553,7 +553,7 @@ bool FGroupTopologySelector::DoEdgeBasedSelection(const FSelectionSettings& Sett
 	{
 		if (Settings.bSelectDownRay)
 		{
-			for (const FIndex2i ElementTuple : DownRayElements)
+			for (const FIndex2i& ElementTuple : DownRayElements)
 			{
 				ResultOut.SelectedEdgeIDs.Add(ElementTuple.A);
 			}
