MeshModeling: enable SubdividePoly to handle input mesh with no UV attribute

#rnx

diff --git a/Engine/Plugins/Runtime/MeshModelingToolset/Source/ModelingComponentsEditorOnly/Private/Operations/SubdividePoly.cpp b/Engine/Plugins/Runtime/MeshModelingToolset/Source/ModelingComponentsEditorOnly/Private/Operations/SubdividePoly.cpp
--- a/Engine/Plugins/Runtime/MeshModelingToolset/Source/ModelingComponentsEditorOnly/Private/Operations/SubdividePoly.cpp
+++ b/Engine/Plugins/Runtime/MeshModelingToolset/Source/ModelingComponentsEditorOnly/Private/Operations/SubdividePoly.cpp
@@ -579,7 +579,11 @@ bool FSubdividePoly::ComputeTopologySubdivision()
 		Descriptor.numVertsPerFace = NumVertsPerFace.GetData();
 		Descriptor.vertIndicesPerFace = BoundaryVertsPerFace.GetData();
 
-		if (UVComputationMethod == ESubdivisionOutputUVs::Interpolated)
+		const bool bShouldInterpolateUVs = (UVComputationMethod == ESubdivisionOutputUVs::Interpolated) &&
+			OriginalMesh.HasAttributes() && 
+			(OriginalMesh.Attributes()->PrimaryUV() != nullptr);
+
+		if (bShouldInterpolateUVs)
 		{
 			UVChannel.numValues = BoundaryVertsPerFace.Num();
 
@@ -806,13 +810,12 @@ bool FSubdividePoly::ComputeSubdividedMesh(FDynamicMesh3& OutMesh)
 	// 
 	TArray<SubdividePolyLocal::SubdUVVertex> RefinedUVs;
 
-	if (UVComputationMethod == ESubdivisionOutputUVs::Interpolated)
-	{
-		if (!OriginalMesh.HasAttributes() || !OriginalMesh.Attributes()->PrimaryUV())
-		{
-			return false;
-		}
+	const bool bShouldInterpolateUVs = (UVComputationMethod == ESubdivisionOutputUVs::Interpolated) &&
+		OriginalMesh.HasAttributes() &&
+		(OriginalMesh.Attributes()->PrimaryUV() != nullptr);
 
+	if (bShouldInterpolateUVs)
+	{
 		TArray<SubdividePolyLocal::SubdUVVertex> SourceUVs;
 		bool bGetUVsOK;
 
@@ -869,7 +872,7 @@ bool FSubdividePoly::ComputeSubdividedMesh(FDynamicMesh3& OutMesh)
 
 	const OpenSubdiv::Far::TopologyLevel& FinalLevel = Refiner->TopologyRefiner->GetLevel(Level);
 
-	check(UVComputationMethod != ESubdivisionOutputUVs::Interpolated || FinalLevel.GetNumFVarValues() == RefinedUVs.Num());
+	check(!bShouldInterpolateUVs || FinalLevel.GetNumFVarValues() == RefinedUVs.Num());
 	check(bNewPolyGroups || FinalLevel.GetNumFaces() == RefinedGroupIDs.Num());
 
 	// Add the faces (manually triangulate the output here)
@@ -910,7 +913,7 @@ bool FSubdividePoly::ComputeSubdividedMesh(FDynamicMesh3& OutMesh)
 			int TriAIndex = OutMesh.AppendTriangle(TriA, GroupID);
 			int TriBIndex = OutMesh.AppendTriangle(TriB, GroupID);
 
-			if (UVComputationMethod == ESubdivisionOutputUVs::Interpolated)
+			if (bShouldInterpolateUVs)
 			{
 				OpenSubdiv::Far::ConstIndexArray FaceUVIndices = FinalLevel.GetFaceFVarValues(FaceID);
 				
@@ -932,7 +935,7 @@ bool FSubdividePoly::ComputeSubdividedMesh(FDynamicMesh3& OutMesh)
 			check(Face.size() == 3);
 			int TriIndex = OutMesh.AppendTriangle(FIndex3i{ Face[0], Face[1], Face[2] }, GroupID);
 
-			if (UVComputationMethod == ESubdivisionOutputUVs::Interpolated)
+			if (bShouldInterpolateUVs)
 			{
 				OpenSubdiv::Far::ConstIndexArray FaceUVIndices = FinalLevel.GetFaceFVarValues(FaceID);
 				FIndex3i UVTri{ FaceUVIndices[0], FaceUVIndices[1], FaceUVIndices[2] };
@@ -964,7 +967,7 @@ bool FSubdividePoly::ComputeSubdividedMesh(FDynamicMesh3& OutMesh)
 		FMeshNormals::InitializeOverlayToPerVertexNormals(OutMesh.Attributes()->PrimaryNormals(), bUseExistingMeshVertexNormals);
 	}
 
-	if (UVComputationMethod == ESubdivisionOutputUVs::Interpolated)
+	if (bShouldInterpolateUVs)
 	{
 		SubdividePolyLocal::InitializeOverlayToFaceVertexUVs(OutMesh.Attributes()->PrimaryUV(), UVTriangles, UVElements);
 	}
