MeshModelingToolset: Make ModelingOperatorsEditorOnly module available for Linux/Mac

diff --git a/Engine/Plugins/Runtime/MeshModelingToolset/MeshModelingToolset.uplugin b/Engine/Plugins/Runtime/MeshModelingToolset/MeshModelingToolset.uplugin
--- a/Engine/Plugins/Runtime/MeshModelingToolset/MeshModelingToolset.uplugin
+++ b/Engine/Plugins/Runtime/MeshModelingToolset/MeshModelingToolset.uplugin
@@ -38,8 +38,7 @@
 		{
 			"Name": "ModelingOperatorsEditorOnly",
 			"Type": "Editor",
-			"LoadingPhase": "Default",
-			"PlatformAllowList": [ "Win64" ]
+			"LoadingPhase": "Default"
 		}
 	],
 	"Plugins": [
