Modeling Mode: Improve handling of USkeletalMeshComponentToolTarget to properly check the status of the underlying SkeletalMesh rather than simply the status of the component, like the UStaticMeshComponentToolTarget does.

Improve the language around the asset only targets to better clarify their purpose.

#rnx

diff --git a/Engine/Plugins/Runtime/MeshModelingToolset/Source/ModelingComponentsEditorOnly/Private/ToolTargets/SkeletalMeshComponentToolTarget.cpp b/Engine/Plugins/Runtime/MeshModelingToolset/Source/ModelingComponentsEditorOnly/Private/ToolTargets/SkeletalMeshComponentToolTarget.cpp
--- a/Engine/Plugins/Runtime/MeshModelingToolset/Source/ModelingComponentsEditorOnly/Private/ToolTargets/SkeletalMeshComponentToolTarget.cpp
+++ b/Engine/Plugins/Runtime/MeshModelingToolset/Source/ModelingComponentsEditorOnly/Private/ToolTargets/SkeletalMeshComponentToolTarget.cpp
@@ -21,6 +21,21 @@ namespace USkeletalMeshComponentToolTargetLocals
 // USkeletalMeshComponentReadOnlyToolTarget
 //
 
+bool USkeletalMeshComponentReadOnlyToolTarget::IsValid() const
+{
+	if (!UPrimitiveComponentToolTarget::IsValid())
+	{
+		return false;
+	}
+	USkinnedMeshComponent* SkinnedMeshComponent = Cast<USkinnedMeshComponent>(Component);
+	if (SkinnedMeshComponent == nullptr)
+	{
+		return false;
+	}
+	const USkeletalMesh* SkeletalMesh = SkinnedMeshComponent->SkeletalMesh;
+
+	return USkeletalMeshReadOnlyToolTarget::IsValid(SkeletalMesh);
+}
 
 int32 USkeletalMeshComponentReadOnlyToolTarget::GetNumMaterials() const
 {
diff --git a/Engine/Plugins/Runtime/MeshModelingToolset/Source/ModelingComponentsEditorOnly/Private/ToolTargets/SkeletalMeshToolTarget.cpp b/Engine/Plugins/Runtime/MeshModelingToolset/Source/ModelingComponentsEditorOnly/Private/ToolTargets/SkeletalMeshToolTarget.cpp
--- a/Engine/Plugins/Runtime/MeshModelingToolset/Source/ModelingComponentsEditorOnly/Private/ToolTargets/SkeletalMeshToolTarget.cpp
+++ b/Engine/Plugins/Runtime/MeshModelingToolset/Source/ModelingComponentsEditorOnly/Private/ToolTargets/SkeletalMeshToolTarget.cpp
@@ -25,7 +25,17 @@ namespace USkeletalMeshToolTargetLocals
 
 bool USkeletalMeshReadOnlyToolTarget::IsValid() const
 {
-	return SkeletalMesh && IsValidChecked(SkeletalMesh) && !SkeletalMesh->IsUnreachable() && SkeletalMesh->IsValidLowLevel();
+	return IsValid(SkeletalMesh);
+}
+
+bool USkeletalMeshReadOnlyToolTarget::IsValid(const USkeletalMesh* SkeletalMeshIn)
+{
+	if (!SkeletalMeshIn || !IsValidChecked(SkeletalMeshIn) || SkeletalMeshIn->IsUnreachable() || !SkeletalMeshIn->IsValidLowLevel())
+	{
+		return false;
+	}
+
+	return true;
 }
 
 int32 USkeletalMeshReadOnlyToolTarget::GetNumMaterials() const
diff --git a/Engine/Plugins/Runtime/MeshModelingToolset/Source/ModelingComponentsEditorOnly/Public/ToolTargets/SkeletalMeshComponentToolTarget.h b/Engine/Plugins/Runtime/MeshModelingToolset/Source/ModelingComponentsEditorOnly/Public/ToolTargets/SkeletalMeshComponentToolTarget.h
--- a/Engine/Plugins/Runtime/MeshModelingToolset/Source/ModelingComponentsEditorOnly/Public/ToolTargets/SkeletalMeshComponentToolTarget.h
+++ b/Engine/Plugins/Runtime/MeshModelingToolset/Source/ModelingComponentsEditorOnly/Public/ToolTargets/SkeletalMeshComponentToolTarget.h
@@ -33,6 +33,10 @@ class MODELINGCOMPONENTSEDITORONLY_API USkeletalMeshComponentReadOnlyToolTarget
 	GENERATED_BODY()
 
 public:
+
+	// UToolTarget implementation
+	virtual bool IsValid() const override;
+
 	// IMeshDescriptionProvider implementation
 	virtual const FMeshDescription* GetMeshDescription(const FGetMeshParameters& GetMeshParams = FGetMeshParameters()) override;
 
diff --git a/Engine/Plugins/Runtime/MeshModelingToolset/Source/ModelingComponentsEditorOnly/Public/ToolTargets/SkeletalMeshToolTarget.h b/Engine/Plugins/Runtime/MeshModelingToolset/Source/ModelingComponentsEditorOnly/Public/ToolTargets/SkeletalMeshToolTarget.h
--- a/Engine/Plugins/Runtime/MeshModelingToolset/Source/ModelingComponentsEditorOnly/Public/ToolTargets/SkeletalMeshToolTarget.h
+++ b/Engine/Plugins/Runtime/MeshModelingToolset/Source/ModelingComponentsEditorOnly/Public/ToolTargets/SkeletalMeshToolTarget.h
@@ -19,6 +19,9 @@ class USkeletalMesh;
 
 /**
  * A tool target backed by a read-only skeletal mesh.
+ * 
+ * This is a special tool target that refers to the underlying asset (in this case a skeletal mesh), rather than indirectly through a component.
+ * This type of target is used in cases, such as opening an asset through the content browser, when there is no component available.
  */
 UCLASS(Transient)
 class MODELINGCOMPONENTSEDITORONLY_API USkeletalMeshReadOnlyToolTarget :
@@ -57,6 +60,8 @@ protected:
 	friend class USkeletalMeshComponentReadOnlyToolTarget;
 	friend class USkeletalMeshComponentToolTarget;
 
+	static bool IsValid(const USkeletalMesh* USkeletalMesh);
+
 	static void GetMeshDescription(const USkeletalMesh* SkeletalMesh, FMeshDescription& MeshDescriptionOut);
 	static void GetMaterialSet(const USkeletalMesh* SkeletalMesh, FComponentMaterialSet& MaterialSetOut,
 		bool bPreferAssetMaterials);
diff --git a/Engine/Plugins/Runtime/MeshModelingToolset/Source/ModelingComponentsEditorOnly/Public/ToolTargets/StaticMeshToolTarget.h b/Engine/Plugins/Runtime/MeshModelingToolset/Source/ModelingComponentsEditorOnly/Public/ToolTargets/StaticMeshToolTarget.h
--- a/Engine/Plugins/Runtime/MeshModelingToolset/Source/ModelingComponentsEditorOnly/Public/ToolTargets/StaticMeshToolTarget.h
+++ b/Engine/Plugins/Runtime/MeshModelingToolset/Source/ModelingComponentsEditorOnly/Public/ToolTargets/StaticMeshToolTarget.h
@@ -20,6 +20,9 @@ class UStaticMesh;
 /**
  * A tool target backed by a static mesh asset that can provide and take a mesh
  * description.
+ * 
+ * This is a special tool target that refers to the underlying asset (in this case a static mesh), rather than indirectly through a component.
+ * This type of target is used in cases, such as opening an asset through the content browser, when there is no component available.
  */
 UCLASS(Transient)
 class MODELINGCOMPONENTSEDITORONLY_API UStaticMeshToolTarget : 
