MeshModelingToolset: Cleanup AutoUV tool

diff --git a/Engine/Plugins/Editor/ModelingToolsEditorMode/Source/ModelingToolsEditorMode/Private/ModelingToolsManagerActions.cpp b/Engine/Plugins/Editor/ModelingToolsEditorMode/Source/ModelingToolsEditorMode/Private/ModelingToolsManagerActions.cpp
--- a/Engine/Plugins/Editor/ModelingToolsEditorMode/Source/ModelingToolsEditorMode/Private/ModelingToolsManagerActions.cpp
+++ b/Engine/Plugins/Editor/ModelingToolsEditorMode/Source/ModelingToolsEditorMode/Private/ModelingToolsManagerActions.cpp
@@ -83,11 +83,11 @@ void FModelingToolsManagerCommands::RegisterCommands()
 	// Attributes
 
 	// UVs
-	REGISTER_MODELING_TOOL_COMMAND(BeginGlobalUVGenerateTool, "AutoUV", "Auto-Unwrap and pack UVs for Mesh");
-	REGISTER_MODELING_TOOL_COMMAND(BeginGroupUVGenerateTool, "Unwrap", "Recalculate UV unwrapping for Mesh regions");
-	REGISTER_MODELING_TOOL_COMMAND(BeginUVProjectionTool, "Project", "Set UVs from Projection");
-	REGISTER_MODELING_TOOL_COMMAND(BeginUVSeamEditTool, "SeamEd", "Add UV Seams to Mesh");
-	REGISTER_MODELING_TOOL_COMMAND(BeginTransformUVIslandsTool, "XForm", "Transform UV islands in UV space");
+	REGISTER_MODELING_TOOL_COMMAND(BeginGlobalUVGenerateTool, "AutoUV", "Automatically unwrap and pack UVs for mesh");
+	REGISTER_MODELING_TOOL_COMMAND(BeginGroupUVGenerateTool, "Unwrap", "Perform UV unwrapping for mesh");
+	REGISTER_MODELING_TOOL_COMMAND(BeginUVProjectionTool, "Project", "Set UVs from projection");
+	REGISTER_MODELING_TOOL_COMMAND(BeginUVSeamEditTool, "SeamEd", "Add UV seams to mesh");
+	REGISTER_MODELING_TOOL_COMMAND(BeginTransformUVIslandsTool, "XFormUV", "Transform UV islands in UV space");
 	REGISTER_MODELING_TOOL_COMMAND(BeginUVLayoutTool, "Layout", "Transform and pack existing UVs");
 
 	// Baking
diff --git a/Engine/Plugins/Experimental/MeshLODToolset/Source/MeshLODToolset/Public/Graphs/GenerateStaticMeshLODProcess.h b/Engine/Plugins/Experimental/MeshLODToolset/Source/MeshLODToolset/Public/Graphs/GenerateStaticMeshLODProcess.h
--- a/Engine/Plugins/Experimental/MeshLODToolset/Source/MeshLODToolset/Public/Graphs/GenerateStaticMeshLODProcess.h
+++ b/Engine/Plugins/Experimental/MeshLODToolset/Source/MeshLODToolset/Public/Graphs/GenerateStaticMeshLODProcess.h
@@ -180,11 +180,11 @@ struct FGenerateStaticMeshLODProcess_UVSettings_PatchBuilder
 	UPROPERTY(EditAnywhere, Category = "PatchBuilder", meta = (UIMin = "0.1", UIMax = "2.0", ClampMin = "0.01", ClampMax = "100.0"))
 	float CurvatureAlignment = 1.0f;
 
-	/** Number of smoothing steps to apply in the ExpMap UV Generation method (Smoothing slightly increases distortion but produces more stable results) */
+	/** Number of smoothing steps to apply; this slightly increases distortion but produces more stable results. */
 	UPROPERTY(EditAnywhere, Category = "PatchBuilder", AdvancedDisplay, meta = (UIMin = "0", UIMax = "25", ClampMin = "0", ClampMax = "1000"))
 	int SmoothingSteps = 5;
 
-	/** Smoothing parameter, larger values result in faster smoothing in each step */
+	/** Smoothing parameter; larger values result in faster smoothing in each step. */
 	UPROPERTY(EditAnywhere, Category = "PatchBuilder", AdvancedDisplay, meta = (UIMin = "0", UIMax = "1.0", ClampMin = "0", ClampMax = "1.0"))
 	float SmoothingAlpha = 0.25f;
 };
diff --git a/Engine/Plugins/Experimental/MeshModelingToolsetExp/Source/MeshModelingToolsEditorOnly/Private/ParameterizeMeshTool.cpp b/Engine/Plugins/Experimental/MeshModelingToolsetExp/Source/MeshModelingToolsEditorOnly/Private/ParameterizeMeshTool.cpp
--- a/Engine/Plugins/Experimental/MeshModelingToolsetExp/Source/MeshModelingToolsEditorOnly/Private/ParameterizeMeshTool.cpp
+++ b/Engine/Plugins/Experimental/MeshModelingToolsetExp/Source/MeshModelingToolsEditorOnly/Private/ParameterizeMeshTool.cpp
@@ -208,7 +208,6 @@ bool UParameterizeMeshTool::CanAccept() const
 
 TUniquePtr<FDynamicMeshOperator> UParameterizeMeshTool::MakeNewOperator()
 {
-	FAxisAlignedBox3d MeshBounds = Preview->PreviewMesh->GetMesh()->GetBounds();
 	TUniquePtr<FParameterizeMeshOp> ParameterizeMeshOp = MakeUnique<FParameterizeMeshOp>();
 
 	ParameterizeMeshOp->InputMesh = InputMesh;
@@ -216,8 +215,8 @@ TUniquePtr<FDynamicMeshOperator> UParameterizeMeshTool::MakeNewOperator()
 	ParameterizeMeshOp->Method = (EParamOpBackend)(int)Settings->Method;
 
 	// uvatlas options
-	ParameterizeMeshOp->Stretch   = UVAtlasProperties->ChartStretch;
-	ParameterizeMeshOp->NumCharts = UVAtlasProperties->NumCharts;
+	ParameterizeMeshOp->Stretch   = UVAtlasProperties->IslandStretch;
+	ParameterizeMeshOp->NumCharts = UVAtlasProperties->NumIslands;
 	
 	// xatlas options
 	ParameterizeMeshOp->XAtlasMaxIterations = XAtlasProperties->MaxIterations;
@@ -230,11 +229,11 @@ TUniquePtr<FDynamicMeshOperator> UParameterizeMeshTool::MakeNewOperator()
 		ParameterizeMeshOp->InputGroupLayer = PolygroupLayerProperties->GetSelectedLayer(*InputMesh);
 	}
 	ParameterizeMeshOp->PatchCurvatureAlignmentWeight = PatchBuilderProperties->CurvatureAlignment;
-	ParameterizeMeshOp->PatchMergingMetricThresh = PatchBuilderProperties->MergingThreshold;
-	ParameterizeMeshOp->PatchMergingAngleThresh = PatchBuilderProperties->MaxAngleDeviation;
+	ParameterizeMeshOp->PatchMergingMetricThresh = PatchBuilderProperties->MergingDistortionThreshold;
+	ParameterizeMeshOp->PatchMergingAngleThresh = PatchBuilderProperties->MergingAngleThreshold;
 	ParameterizeMeshOp->ExpMapNormalSmoothingSteps = PatchBuilderProperties->SmoothingSteps;
 	ParameterizeMeshOp->ExpMapNormalSmoothingAlpha = PatchBuilderProperties->SmoothingAlpha;
-	ParameterizeMeshOp->bEnablePacking = PatchBuilderProperties->bAutoPack;
+	ParameterizeMeshOp->bEnablePacking = PatchBuilderProperties->bRepack;
 	ParameterizeMeshOp->Width = ParameterizeMeshOp->Height = PatchBuilderProperties->TextureResolution;
 
 
diff --git a/Engine/Plugins/Experimental/MeshModelingToolsetExp/Source/MeshModelingToolsEditorOnly/Public/ParameterizeMeshTool.h b/Engine/Plugins/Experimental/MeshModelingToolsetExp/Source/MeshModelingToolsEditorOnly/Public/ParameterizeMeshTool.h
--- a/Engine/Plugins/Experimental/MeshModelingToolsetExp/Source/MeshModelingToolsEditorOnly/Public/ParameterizeMeshTool.h
+++ b/Engine/Plugins/Experimental/MeshModelingToolsetExp/Source/MeshModelingToolsEditorOnly/Public/ParameterizeMeshTool.h
@@ -13,7 +13,7 @@
 
 #include "ParameterizeMeshTool.generated.h"
 
-// predeclarations
+// Forward declarations
 class UDynamicMeshComponent;
 class UMaterialInterface;
 class UMaterialInstanceDynamic;
@@ -38,13 +38,11 @@ class MESHMODELINGTOOLSEDITORONLY_API UParameterizeMeshToolPatchBuilderGroupLaye
 	GENERATED_BODY()
 public:
 
-	UPROPERTY(EditAnywhere, Category = PolygroupLayers)
+	UPROPERTY(EditAnywhere, Category = "PolyGroup Layer", meta = (DisplayName = "Constrain to PolyGroups"))
 	bool bConstrainToPolygroups = false;
-
 };
 
 
-
 /**
  * UParameterizeMeshTool automatically decomposes the input mesh into charts, solves for UVs,
  * and then packs the resulting charts
@@ -98,11 +96,9 @@ protected:
 	UPROPERTY()
 	TObjectPtr<UUVLayoutPreview> UVLayoutView = nullptr;
 
-protected:
 	UPROPERTY()
 	TObjectPtr<UMeshOpPreviewWithBackgroundCompute> Preview = nullptr;
 
-protected:
 	TSharedPtr<UE::Geometry::FDynamicMesh3, ESPMode::ThreadSafe> InputMesh;
 
 	void OnMethodTypeChanged();
diff --git a/Engine/Plugins/Experimental/UVEditor/Source/UVEditorTools/Private/UVEditorRecomputeUVsTool.cpp b/Engine/Plugins/Experimental/UVEditor/Source/UVEditorTools/Private/UVEditorRecomputeUVsTool.cpp
--- a/Engine/Plugins/Experimental/UVEditor/Source/UVEditorTools/Private/UVEditorRecomputeUVsTool.cpp
+++ b/Engine/Plugins/Experimental/UVEditor/Source/UVEditorTools/Private/UVEditorRecomputeUVsTool.cpp
@@ -259,8 +259,8 @@ void UUVEditorRecomputeUVsTool::RecordAnalytics()
 	}
 	if (Settings->UnwrapType == ERecomputeUVsPropertiesUnwrapType::IslandMerging)
 	{
-		Attributes.Add(FAnalyticsEventAttribute(TEXT("Settings.MergingDistortionLimit"), Settings->MergingDistortionLimit));
-		Attributes.Add(FAnalyticsEventAttribute(TEXT("Settings.MergingAngleLimit"), Settings->MergingAngleLimit));
+		Attributes.Add(FAnalyticsEventAttribute(TEXT("Settings.MergingDistortionThreshold"), Settings->MergingDistortionThreshold));
+		Attributes.Add(FAnalyticsEventAttribute(TEXT("Settings.MergingAngleThreshold"), Settings->MergingAngleThreshold));
 	}
 	
 	Attributes.Add(AnalyticsEventAttributeEnum(TEXT("Settings.LayoutType"), Settings->LayoutType));
diff --git a/Engine/Plugins/Experimental/UVEditor/Source/UVEditorToolsEditorOnly/Private/UVEditorParameterizeMeshTool.cpp b/Engine/Plugins/Experimental/UVEditor/Source/UVEditorToolsEditorOnly/Private/UVEditorParameterizeMeshTool.cpp
--- a/Engine/Plugins/Experimental/UVEditor/Source/UVEditorToolsEditorOnly/Private/UVEditorParameterizeMeshTool.cpp
+++ b/Engine/Plugins/Experimental/UVEditor/Source/UVEditorToolsEditorOnly/Private/UVEditorParameterizeMeshTool.cpp
@@ -3,10 +3,7 @@
 #include "UVEditorParameterizeMeshTool.h"
 #include "InteractiveToolManager.h"
 #include "ToolBuilderUtil.h"
-#include "ToolSetupUtil.h"
-#include "ModelingToolTargetUtil.h"
 #include "DynamicMesh/DynamicMesh3.h"
-#include "FaceGroupUtil.h"
 #include "ParameterizationOps/ParameterizeMeshOp.h"
 #include "Properties/ParameterizeMeshProperties.h"
 #include "ToolTargets/UVEditorToolMeshInput.h"
@@ -254,19 +251,19 @@ void UUVEditorParameterizeMeshTool::RecordAnalytics()
 		case EParameterizeMeshUVMethod::PatchBuilder:
 			Attributes.Add(FAnalyticsEventAttribute(FString::Printf(TEXT("Settings.%s.InitialPatches"), *MethodName), PatchBuilderProperties->InitialPatches));
 			Attributes.Add(FAnalyticsEventAttribute(FString::Printf(TEXT("Settings.%s.CurvatureAlignment"), *MethodName), PatchBuilderProperties->CurvatureAlignment));
-			Attributes.Add(FAnalyticsEventAttribute(FString::Printf(TEXT("Settings.%s.MergingThreshold"), *MethodName), PatchBuilderProperties->MergingThreshold));
-			Attributes.Add(FAnalyticsEventAttribute(FString::Printf(TEXT("Settings.%s.MaxAngleDeviation"), *MethodName), PatchBuilderProperties->MaxAngleDeviation));
+			Attributes.Add(FAnalyticsEventAttribute(FString::Printf(TEXT("Settings.%s.MergingDistortionThreshold"), *MethodName), PatchBuilderProperties->MergingDistortionThreshold));
+			Attributes.Add(FAnalyticsEventAttribute(FString::Printf(TEXT("Settings.%s.MergingAngleThreshold"), *MethodName), PatchBuilderProperties->MergingAngleThreshold));
 			Attributes.Add(FAnalyticsEventAttribute(FString::Printf(TEXT("Settings.%s.SmoothingSteps"), *MethodName), PatchBuilderProperties->SmoothingSteps));
 			Attributes.Add(FAnalyticsEventAttribute(FString::Printf(TEXT("Settings.%s.SmoothingAlpha"), *MethodName), PatchBuilderProperties->SmoothingAlpha));
-			Attributes.Add(FAnalyticsEventAttribute(FString::Printf(TEXT("Settings.%s.AutoPack"), *MethodName), PatchBuilderProperties->bAutoPack));
-			if (PatchBuilderProperties->bAutoPack)
+			Attributes.Add(FAnalyticsEventAttribute(FString::Printf(TEXT("Settings.%s.Repack"), *MethodName), PatchBuilderProperties->bRepack));
+			if (PatchBuilderProperties->bRepack)
 			{
 				Attributes.Add(FAnalyticsEventAttribute(FString::Printf(TEXT("Settings.%s.TextureResolution"), *MethodName), PatchBuilderProperties->TextureResolution));
 			}
 			break;
 		case EParameterizeMeshUVMethod::UVAtlas:
-			Attributes.Add(FAnalyticsEventAttribute(FString::Printf(TEXT("Settings.%s.ChartStretch"), *MethodName), UVAtlasProperties->ChartStretch));
-			Attributes.Add(FAnalyticsEventAttribute(FString::Printf(TEXT("Settings.%s.NumCharts"), *MethodName), UVAtlasProperties->NumCharts));
+			Attributes.Add(FAnalyticsEventAttribute(FString::Printf(TEXT("Settings.%s.IslandStretch"), *MethodName), UVAtlasProperties->IslandStretch));
+			Attributes.Add(FAnalyticsEventAttribute(FString::Printf(TEXT("Settings.%s.NumIslands"), *MethodName), UVAtlasProperties->NumIslands));
 			break;
 		case EParameterizeMeshUVMethod::XAtlas:
 			Attributes.Add(FAnalyticsEventAttribute(FString::Printf(TEXT("Settings.%s.MaxIterations"), *MethodName), XAtlasProperties->MaxIterations));
diff --git a/Engine/Plugins/Runtime/MeshModelingToolset/Source/MeshModelingTools/Private/UVProjectionTool.cpp b/Engine/Plugins/Runtime/MeshModelingToolset/Source/MeshModelingTools/Private/UVProjectionTool.cpp
--- a/Engine/Plugins/Runtime/MeshModelingToolset/Source/MeshModelingTools/Private/UVProjectionTool.cpp
+++ b/Engine/Plugins/Runtime/MeshModelingToolset/Source/MeshModelingTools/Private/UVProjectionTool.cpp
@@ -418,11 +418,11 @@ void UUVProjectionTool::OnMeshUpdated(UMeshOpPreviewWithBackgroundCompute* Previ
 }
 
 
-void UUVProjectionTool::UpdatePlaneFromClick(const FVector3d& Position, const FVector3d& Normal, bool bTransationOnly)
+void UUVProjectionTool::UpdatePlaneFromClick(const FVector3d& Position, const FVector3d& Normal, bool bTransitionOnly)
 {
 	FFrame3d CurrentFrame(TransformProxy->GetTransform());
 	CurrentFrame.Origin = Position;
-	if (bTransationOnly == false)
+	if (bTransitionOnly == false)
 	{
 		CurrentFrame.AlignAxis(2, Normal);
 	}
diff --git a/Engine/Plugins/Runtime/MeshModelingToolset/Source/MeshModelingTools/Public/UVProjectionTool.h b/Engine/Plugins/Runtime/MeshModelingToolset/Source/MeshModelingTools/Public/UVProjectionTool.h
--- a/Engine/Plugins/Runtime/MeshModelingToolset/Source/MeshModelingTools/Public/UVProjectionTool.h
+++ b/Engine/Plugins/Runtime/MeshModelingToolset/Source/MeshModelingTools/Public/UVProjectionTool.h
@@ -144,7 +144,7 @@ public:
 		EditCondition = "ProjectionType == EUVProjectionMethod::ExpMap", EditConditionHides))
 	int ExpMapSmoothingSteps = 0;
 
-	/** Smoothing parameter, larger values result in faster smoothing in each step. */
+	/** Smoothing parameter; larger values result in faster smoothing in each step. */
 	UPROPERTY(EditAnywhere, Category = "ExpMap Projection", meta = (DisplayName = "Smoothing Alpha", UIMin = "0", UIMax = "1",
 		EditCondition = "ProjectionType == EUVProjectionMethod::ExpMap", EditConditionHides))
 	float ExpMapSmoothingAlpha = 0.25f;
@@ -295,7 +295,7 @@ protected:
 	UPROPERTY()
 	TObjectPtr<USingleClickInputBehavior> ClickToSetPlaneBehavior;
 
-	void UpdatePlaneFromClick(const FVector3d& Position, const FVector3d& Normal, bool bTransationOnly);
+	void UpdatePlaneFromClick(const FVector3d& Position, const FVector3d& Normal, bool bTransitionOnly);
 
 	//
 	// Support for Action Buttons
diff --git a/Engine/Plugins/Runtime/MeshModelingToolset/Source/ModelingOperators/Private/ParameterizationOps/RecomputeUVsOp.cpp b/Engine/Plugins/Runtime/MeshModelingToolset/Source/ModelingOperators/Private/ParameterizationOps/RecomputeUVsOp.cpp
--- a/Engine/Plugins/Runtime/MeshModelingToolset/Source/ModelingOperators/Private/ParameterizationOps/RecomputeUVsOp.cpp
+++ b/Engine/Plugins/Runtime/MeshModelingToolset/Source/ModelingOperators/Private/ParameterizationOps/RecomputeUVsOp.cpp
@@ -411,8 +411,8 @@ TUniquePtr<FDynamicMeshOperator> URecomputeUVsOpFactory::MakeNewOperator()
 	RecomputeUVsOp->NormalSmoothingRounds = Settings->SmoothingSteps;
 	RecomputeUVsOp->NormalSmoothingAlpha = Settings->SmoothingAlpha;
 
-	RecomputeUVsOp->MergingThreshold = Settings->MergingDistortionLimit;
-	RecomputeUVsOp->MaxNormalDeviationDeg = Settings->MergingAngleLimit;
+	RecomputeUVsOp->MergingThreshold = Settings->MergingDistortionThreshold;
+	RecomputeUVsOp->MaxNormalDeviationDeg = Settings->MergingAngleThreshold;
 
 	RecomputeUVsOp->bPackUVs = false;
 	RecomputeUVsOp->bNormalizeAreas = false;
diff --git a/Engine/Plugins/Runtime/MeshModelingToolset/Source/ModelingOperators/Public/Properties/RecomputeUVsProperties.h b/Engine/Plugins/Runtime/MeshModelingToolset/Source/ModelingOperators/Public/Properties/RecomputeUVsProperties.h
--- a/Engine/Plugins/Runtime/MeshModelingToolset/Source/ModelingOperators/Public/Properties/RecomputeUVsProperties.h
+++ b/Engine/Plugins/Runtime/MeshModelingToolset/Source/ModelingOperators/Public/Properties/RecomputeUVsProperties.h
@@ -82,20 +82,20 @@ public:
 		EditCondition = "UnwrapType == ERecomputeUVsPropertiesUnwrapType::ExpMap || UnwrapType == ERecomputeUVsPropertiesUnwrapType::IslandMerging"))
 	int SmoothingSteps = 5;
 
-	/** Smoothing parameter, larger values result in faster smoothing in each step. This is only enabled if the Unwrap Type is set to ExpMap or Island Merging. */
+	/** Smoothing parameter; larger values result in faster smoothing in each step. This is only enabled if the Unwrap Type is set to ExpMap or Island Merging. */
 	UPROPERTY(EditAnywhere, Category = "UV Unwrap", meta = (UIMin = "0", UIMax = "1.0", ClampMin = "0", ClampMax = "1.0",
 		EditCondition = "UnwrapType == ERecomputeUVsPropertiesUnwrapType::ExpMap || UnwrapType == ERecomputeUVsPropertiesUnwrapType::IslandMerging"))
 	float SmoothingAlpha = 0.25f;
 
 	/** Threshold for stretching and distortion below which island merging is allowed; larger values increase the allowable UV distortion. This is only enabled if the Unwrap Type is set to Island Merging. */
-	UPROPERTY(EditAnywhere, Category = "UV Unwrap", meta = (UIMin = "1.0", UIMax = "5.0", ClampMin = "1.0",
+	UPROPERTY(EditAnywhere, Category = "UV Unwrap", meta = (DisplayName = "Distortion Threshold", UIMin = "1.0", UIMax = "5.0", ClampMin = "1.0",
 		EditCondition = "UnwrapType == ERecomputeUVsPropertiesUnwrapType::IslandMerging"))
-	float MergingDistortionLimit = 1.5f;
+	float MergingDistortionThreshold = 1.5f;
 
 	/** Threshold for the average face normal deviation below  which island merging is allowed. This is only enabled if the Unwrap Type is set to Island Merging. */
-	UPROPERTY(EditAnywhere, Category = "UV Unwrap", meta = (UIMin = "0.0", UIMax = "90.0", ClampMin = "0.0", ClampMax = "180.0",
+	UPROPERTY(EditAnywhere, Category = "UV Unwrap", meta = (DisplayName = "Angle Threshold", UIMin = "0.0", UIMax = "90.0", ClampMin = "0.0", ClampMax = "180.0",
 		EditCondition = "UnwrapType == ERecomputeUVsPropertiesUnwrapType::IslandMerging"))
-	float MergingAngleLimit = 45.0f;
+	float MergingAngleThreshold = 45.0f;
 
 	/** Uniformly scale and translate UV islands collectively to pack them into the unit square, i.e. fit between 0 and 1 with no overlap. */
 	UPROPERTY(EditAnywhere, Category = "UV Layout")
diff --git a/Engine/Plugins/Runtime/MeshModelingToolset/Source/ModelingOperatorsEditorOnly/Private/ParameterizationOps/ParameterizeMeshOp.cpp b/Engine/Plugins/Runtime/MeshModelingToolset/Source/ModelingOperatorsEditorOnly/Private/ParameterizationOps/ParameterizeMeshOp.cpp
--- a/Engine/Plugins/Runtime/MeshModelingToolset/Source/ModelingOperatorsEditorOnly/Private/ParameterizationOps/ParameterizeMeshOp.cpp
+++ b/Engine/Plugins/Runtime/MeshModelingToolset/Source/ModelingOperatorsEditorOnly/Private/ParameterizationOps/ParameterizeMeshOp.cpp
@@ -461,18 +461,18 @@ TUniquePtr<FDynamicMeshOperator> UParameterizeMeshOperatorFactory::MakeNewOperat
 		Op->Method = EParamOpBackend::PatchBuilder;
 		Op->InitialPatchCount = PatchBuilderProperties->InitialPatches;
 		Op->PatchCurvatureAlignmentWeight = PatchBuilderProperties->CurvatureAlignment;
-		Op->PatchMergingMetricThresh = PatchBuilderProperties->MergingThreshold;
-		Op->PatchMergingAngleThresh = PatchBuilderProperties->MaxAngleDeviation;
+		Op->PatchMergingMetricThresh = PatchBuilderProperties->MergingDistortionThreshold;
+		Op->PatchMergingAngleThresh = PatchBuilderProperties->MergingAngleThreshold;
 		Op->ExpMapNormalSmoothingSteps = PatchBuilderProperties->SmoothingSteps;
 		Op->ExpMapNormalSmoothingAlpha = PatchBuilderProperties->SmoothingAlpha;
-		Op->bEnablePacking = PatchBuilderProperties->bAutoPack;
+		Op->bEnablePacking = PatchBuilderProperties->bRepack;
 		Op->Width = PatchBuilderProperties->TextureResolution;
 		Op->Height = PatchBuilderProperties->TextureResolution;
 		break;
 	case EParameterizeMeshUVMethod::UVAtlas:
 		Op->Method = EParamOpBackend::UVAtlas;		
-		Op->Stretch = UVAtlasProperties->ChartStretch;
-		Op->NumCharts = UVAtlasProperties->NumCharts;
+		Op->Stretch = UVAtlasProperties->IslandStretch;
+		Op->NumCharts = UVAtlasProperties->NumIslands;
 		break;
 	case EParameterizeMeshUVMethod::XAtlas:
 		Op->Method = EParamOpBackend::XAtlas;
diff --git a/Engine/Plugins/Runtime/MeshModelingToolset/Source/ModelingOperatorsEditorOnly/Public/Properties/ParameterizeMeshProperties.h b/Engine/Plugins/Runtime/MeshModelingToolset/Source/ModelingOperatorsEditorOnly/Public/Properties/ParameterizeMeshProperties.h
--- a/Engine/Plugins/Runtime/MeshModelingToolset/Source/ModelingOperatorsEditorOnly/Public/Properties/ParameterizeMeshProperties.h
+++ b/Engine/Plugins/Runtime/MeshModelingToolset/Source/ModelingOperatorsEditorOnly/Public/Properties/ParameterizeMeshProperties.h
@@ -10,13 +10,15 @@
 UENUM()
 enum class EParameterizeMeshUVMethod
 {
-	// keep values the same as UE::Geometry::EParamOpBackend!
-	/** Compute Automatic UVs using PatchBuilder */
+	// Keep values the same as UE::Geometry::EParamOpBackend!
+	// The order is different on purpose to have them sorted alphabetically in the UI.
+
+	/** Compute automatic UVs using the Patch Builder technique */
 	PatchBuilder = 0,
-	/** Compute Automatic UVs using XAtlas */
-	XAtlas = 1,
-	/** Compute Automatic UVs using UVAtlas */
+	/** Compute automatic UVs using the UVAtlas technique */
 	UVAtlas = 2,
+	/** Compute automatic UVs using the XAtlas technique */
+	XAtlas = 1,
 };
 
 
@@ -25,8 +27,8 @@ class MODELINGOPERATORSEDITORONLY_API UParameterizeMeshToolProperties : public U
 {
 	GENERATED_BODY()
 public:
-	/** Automatic UV Generation technique to use */
-	UPROPERTY(EditAnywhere, Category = Options)
+	/** Automatic UV generation technique to use */
+	UPROPERTY(EditAnywhere, Category = AutoUV)
 	EParameterizeMeshUVMethod Method = EParameterizeMeshUVMethod::PatchBuilder;
 };
 
@@ -39,13 +41,13 @@ class MODELINGOPERATORSEDITORONLY_API UParameterizeMeshToolUVAtlasProperties : p
 {
 	GENERATED_BODY()
 public:
-	/** Maximum amount of stretch, from none to any.  If zero stretch is specified each triangle will likely be its own chart */
+	/** Maximum amount of stretch, from none to unbounded. If zero stretch is specified, each triangle will likely be its own UV island. */
 	UPROPERTY(EditAnywhere, Category = UVAtlas, meta = (UIMin = "0", UIMax = "1", ClampMin = "0", ClampMax = "1"))
-	float ChartStretch = 0.11f;
+	float IslandStretch = 0.11f;
 
-	/** Hint at number of Charts. 0 (Default) means UVAtlas will decide */
+	/** Hint at number of UV islands. The default of 0 means it is determined automatically. */
 	UPROPERTY(EditAnywhere, Category = UVAtlas, meta = (UIMin = "0", UIMax = "100", ClampMin = "0", ClampMax = "10000"))
-	int NumCharts = 0;
+	int NumIslands = 0;
 };
 
 
@@ -54,7 +56,7 @@ class MODELINGOPERATORSEDITORONLY_API UParameterizeMeshToolXAtlasProperties : pu
 {
 	GENERATED_BODY()
 public:
-	/** Number of solve iterations. Higher values generally result in better charts. */
+	/** Number of solve iterations; higher values generally result in better UV islands. */
 	UPROPERTY(EditAnywhere, Category = XAtlas, meta = (UIMin = "1", UIMax = "10", ClampMin = "1", ClampMax = "1000"))
 	int MaxIterations = 1;
 };
@@ -65,36 +67,35 @@ class MODELINGOPERATORSEDITORONLY_API UParameterizeMeshToolPatchBuilderPropertie
 {
 	GENERATED_BODY()
 public:
-	/** Number of initial patches mesh will be split into before computing island merging */
-	UPROPERTY(EditAnywhere, Category = "PatchBuilder", meta = (UIMin = "1", UIMax = "1000", ClampMin = "1", ClampMax = "99999999"))
+	/** Number of initial patches the mesh will be split into before island merging. */
+	UPROPERTY(EditAnywhere, Category = PatchBuilder, meta = (UIMin = "1", UIMax = "1000", ClampMin = "1", ClampMax = "99999999"))
 	int InitialPatches = 100;
 
-	/** This parameter controls alignment of the initial patches to creases in the mesh */
-	UPROPERTY(EditAnywhere, Category = "PatchBuilder", meta = (UIMin = "0.1", UIMax = "2.0", ClampMin = "0.01", ClampMax = "100.0"))
+	/** Alignment of the initial patches to creases in the mesh.*/
+	UPROPERTY(EditAnywhere, Category = PatchBuilder, meta = (UIMin = "0.1", UIMax = "2.0", ClampMin = "0.01", ClampMax = "100.0"))
 	float CurvatureAlignment = 1.0f;
 
-	/** Distortion/Stretching Threshold for island merging - larger values increase the allowable UV stretching */
-	UPROPERTY(EditAnywhere, Category = "PatchBuilder", meta = (UIMin = "1.0", UIMax = "5.0", ClampMin = "1.0"))
-	float MergingThreshold = 1.5f;
+	/** Threshold for stretching and distortion below which island merging is allowed; larger values increase the allowable UV distortion. */
+	UPROPERTY(EditAnywhere, Category = PatchBuilder, meta = (DisplayName = "Distortion Threshold", UIMin = "1.0", UIMax = "5.0", ClampMin = "1.0"))
+	float MergingDistortionThreshold = 1.5f;
 
-	/** UV islands will not be merged if their average face normals deviate by larger than this amount */
-	UPROPERTY(EditAnywhere, Category = "PatchBuilder", meta = (UIMin = "0.0", UIMax = "90.0", ClampMin = "0.0", ClampMax = "180.0"))
-	float MaxAngleDeviation = 45.0f;
+	/** Threshold for the average face normal deviation below which island merging is allowed. */
+	UPROPERTY(EditAnywhere, Category = PatchBuilder, meta = (DisplayName = "Angle Threshold", UIMin = "0.0", UIMax = "90.0", ClampMin = "0.0", ClampMax = "180.0"))
+	float MergingAngleThreshold = 45.0f;
 
-	/** Number of smoothing steps to apply in the ExpMap UV Generation method (Smoothing slightly increases distortion but produces more stable results) */
-	UPROPERTY(EditAnywhere, Category = "PatchBuilder", meta = (UIMin = "0", UIMax = "25", ClampMin = "0", ClampMax = "1000"))
+	/** Number of smoothing steps to apply; this slightly increases distortion but produces more stable results. */
+	UPROPERTY(EditAnywhere, Category = PatchBuilder, meta = (UIMin = "0", UIMax = "25", ClampMin = "0", ClampMax = "1000"))
 	int SmoothingSteps = 5;
 
-	/** Smoothing parameter, larger values result in faster smoothing in each step */
-	UPROPERTY(EditAnywhere, Category = "PatchBuilder", meta = (UIMin = "0", UIMax = "1.0", ClampMin = "0", ClampMax = "1.0"))
+	/** Smoothing parameter; larger values result in faster smoothing in each step. */
+	UPROPERTY(EditAnywhere, Category = PatchBuilder, meta = (UIMin = "0", UIMax = "1.0", ClampMin = "0", ClampMax = "1.0"))
 	float SmoothingAlpha = 0.25f;
 
+	/** Automatically pack result UVs into the unit square, i.e. fit between 0 and 1 with no overlap. */
+	UPROPERTY(EditAnywhere, Category = PatchBuilder)
+	bool bRepack = true;
 
-	/** If enabled, result UVs are automatically packed into the standard UV 0-1 square */
-	UPROPERTY(EditAnywhere, Category = "PatchBuilder")
-	bool bAutoPack = true;
-
-	/** Target texture resolution used for UV packing, which determines gutter size */
-	UPROPERTY(EditAnywhere, Category = "PatchBuilder", meta = (UIMin = "64", UIMax = "2048", ClampMin = "2", ClampMax = "4096", EditCondition = bAutoPack))
+	/** Expected resolution of the output textures; this controls spacing left between UV islands to avoid interpolation artifacts. This is only enabled when Repack is enabled. */
+	UPROPERTY(EditAnywhere, Category = PatchBuilder, meta = (UIMin = "64", UIMax = "2048", ClampMin = "2", ClampMax = "4096", EditCondition = bAutoPack))
 	int TextureResolution = 1024;
 };
