MeshModelingToolset: Hide voxel tools using ProxyLOD under Linux/Mac

diff --git a/Engine/Plugins/Editor/ModelingToolsEditorMode/Source/ModelingToolsEditorMode/ModelingToolsEditorMode.Build.cs b/Engine/Plugins/Editor/ModelingToolsEditorMode/Source/ModelingToolsEditorMode/ModelingToolsEditorMode.Build.cs
--- a/Engine/Plugins/Editor/ModelingToolsEditorMode/Source/ModelingToolsEditorMode/ModelingToolsEditorMode.Build.cs
+++ b/Engine/Plugins/Editor/ModelingToolsEditorMode/Source/ModelingToolsEditorMode/ModelingToolsEditorMode.Build.cs
@@ -80,5 +80,7 @@ public class ModelingToolsEditorMode : ModuleRules
 			}
 			);
 
+		PublicDefinitions.Add("WITH_PROXYLOD=" + (Target.Platform == UnrealTargetPlatform.Win64 ? '1' : '0'));
+
 	}
 }
diff --git a/Engine/Plugins/Editor/ModelingToolsEditorMode/Source/ModelingToolsEditorMode/Private/ModelingToolsEditorMode.cpp b/Engine/Plugins/Editor/ModelingToolsEditorMode/Source/ModelingToolsEditorMode/Private/ModelingToolsEditorMode.cpp
--- a/Engine/Plugins/Editor/ModelingToolsEditorMode/Source/ModelingToolsEditorMode/Private/ModelingToolsEditorMode.cpp
+++ b/Engine/Plugins/Editor/ModelingToolsEditorMode/Source/ModelingToolsEditorMode/Private/ModelingToolsEditorMode.cpp
@@ -531,11 +531,13 @@ void UModelingToolsEditorMode::Enter()
 	auto UVLayoutToolBuilder = NewObject<UUVLayoutToolBuilder>();
 	RegisterTool(ToolManagerCommands.BeginUVLayoutTool, TEXT("BeginUVLayoutTool"), UVLayoutToolBuilder);
 
+#if WITH_PROXYLOD
 	auto MergeMeshesToolBuilder = NewObject<UMergeMeshesToolBuilder>();
 	RegisterTool(ToolManagerCommands.BeginVoxelMergeTool, TEXT("BeginVoxelMergeTool"), MergeMeshesToolBuilder);
 
 	auto VoxelCSGMeshesToolBuilder = NewObject<UVoxelCSGMeshesToolBuilder>();
 	RegisterTool(ToolManagerCommands.BeginVoxelBooleanTool, TEXT("BeginVoxelBooleanTool"), VoxelCSGMeshesToolBuilder);
+#endif	// WITH_PROXYLOD
 
 	auto VoxelSolidifyMeshesToolBuilder = NewObject<UVoxelSolidifyMeshesToolBuilder>();
 	RegisterTool(ToolManagerCommands.BeginVoxelSolidifyTool, TEXT("BeginVoxelSolidifyTool"), VoxelSolidifyMeshesToolBuilder);
diff --git a/Engine/Plugins/Editor/ModelingToolsEditorMode/Source/ModelingToolsEditorMode/Private/ModelingToolsEditorModeToolkit.cpp b/Engine/Plugins/Editor/ModelingToolsEditorMode/Source/ModelingToolsEditorMode/Private/ModelingToolsEditorModeToolkit.cpp
--- a/Engine/Plugins/Editor/ModelingToolsEditorMode/Source/ModelingToolsEditorMode/Private/ModelingToolsEditorModeToolkit.cpp
+++ b/Engine/Plugins/Editor/ModelingToolsEditorMode/Source/ModelingToolsEditorMode/Private/ModelingToolsEditorModeToolkit.cpp
@@ -629,8 +629,10 @@ void FModelingToolsEditorModeToolkit::BuildToolPalette(FName PaletteIndex, class
 		ToolbarBuilder.AddToolBarButton(Commands.BeginVoxelSolidifyTool);
 		ToolbarBuilder.AddToolBarButton(Commands.BeginVoxelBlendTool);
 		ToolbarBuilder.AddToolBarButton(Commands.BeginVoxelMorphologyTool);
+#if WITH_PROXYLOD
 		ToolbarBuilder.AddToolBarButton(Commands.BeginVoxelBooleanTool);
 		ToolbarBuilder.AddToolBarButton(Commands.BeginVoxelMergeTool);
+#endif	// WITH_PROXYLOD
 	}
 	else if (PaletteIndex == TriModelingTabName)
 	{
diff --git a/Engine/Plugins/Editor/ModelingToolsEditorMode/Source/ModelingToolsEditorMode/Private/ModelingToolsManagerActions.cpp b/Engine/Plugins/Editor/ModelingToolsEditorMode/Source/ModelingToolsEditorMode/Private/ModelingToolsManagerActions.cpp
--- a/Engine/Plugins/Editor/ModelingToolsEditorMode/Source/ModelingToolsEditorMode/Private/ModelingToolsManagerActions.cpp
+++ b/Engine/Plugins/Editor/ModelingToolsEditorMode/Source/ModelingToolsEditorMode/Private/ModelingToolsManagerActions.cpp
@@ -124,8 +124,11 @@ void FModelingToolsManagerCommands::RegisterCommands()
 	REGISTER_MODELING_TOOL_COMMAND(BeginEditNormalsTool, "Nrmls", "Modify/Update Mesh Normals");
 	REGISTER_MODELING_TOOL_COMMAND(BeginEditTangentsTool, "Tngnts", "Update Mesh Tangents");
 	REGISTER_MODELING_TOOL_COMMAND(BeginRemoveOccludedTrianglesTool, "Jacket", "Remove Hidden Triangles from selected Meshes");
+#if WITH_PROXYLOD
+	// The ProxyLOD plugin is currently only available on Windows. Without it, the following tools do not work as expected.
 	REGISTER_MODELING_TOOL_COMMAND(BeginVoxelMergeTool, "VoxMrg", "Merge Selected Meshes (Voxel-Based)");
 	REGISTER_MODELING_TOOL_COMMAND(BeginVoxelBooleanTool, "VoxBool", "Boolean Selected Meshes (Voxel-Based)");
+#endif	// WITH_PROXYLOD
 	REGISTER_MODELING_TOOL_COMMAND(BeginVoxelSolidifyTool, "VoxWrap", "Wrap Selected Meshes (Voxel-Based)");
 	REGISTER_MODELING_TOOL_COMMAND(BeginVoxelBlendTool, "VoxBlnd", "Blend Selected Meshes (Voxel-Based)");
 	REGISTER_MODELING_TOOL_COMMAND(BeginVoxelMorphologyTool, "VoxMrph", "Offset/Inset Selected Meshes (Voxel-Based)");
diff --git a/Engine/Plugins/Editor/ModelingToolsEditorMode/Source/ModelingToolsEditorMode/Public/ModelingToolsManagerActions.h b/Engine/Plugins/Editor/ModelingToolsEditorMode/Source/ModelingToolsEditorMode/Public/ModelingToolsManagerActions.h
--- a/Engine/Plugins/Editor/ModelingToolsEditorMode/Source/ModelingToolsEditorMode/Public/ModelingToolsManagerActions.h
+++ b/Engine/Plugins/Editor/ModelingToolsEditorMode/Source/ModelingToolsEditorMode/Public/ModelingToolsManagerActions.h
@@ -84,8 +84,12 @@ public:
 	TSharedPtr<FUICommandInfo> BeginHoleFillTool;
 	TSharedPtr<FUICommandInfo> BeginLatticeDeformerTool;
 	TSharedPtr<FUICommandInfo> BeginPolygonCutTool;
+
+#if WITH_PROXYLOD
 	TSharedPtr<FUICommandInfo> BeginVoxelMergeTool;
 	TSharedPtr<FUICommandInfo> BeginVoxelBooleanTool;
+#endif	// WITH_PROXYLOD
+
 	TSharedPtr<FUICommandInfo> BeginVoxelSolidifyTool;
 	TSharedPtr<FUICommandInfo> BeginVoxelBlendTool;
 	TSharedPtr<FUICommandInfo> BeginVoxelMorphologyTool;
diff --git a/Engine/Plugins/Runtime/MeshModelingToolset/Source/ModelingOperatorsEditorOnly/ModelingOperatorsEditorOnly.Build.cs b/Engine/Plugins/Runtime/MeshModelingToolset/Source/ModelingOperatorsEditorOnly/ModelingOperatorsEditorOnly.Build.cs
--- a/Engine/Plugins/Runtime/MeshModelingToolset/Source/ModelingOperatorsEditorOnly/ModelingOperatorsEditorOnly.Build.cs
+++ b/Engine/Plugins/Runtime/MeshModelingToolset/Source/ModelingOperatorsEditorOnly/ModelingOperatorsEditorOnly.Build.cs
@@ -53,7 +53,9 @@ public class ModelingOperatorsEditorOnly : ModuleRules
 			}
 			);
 
-		if (Target.Platform == UnrealTargetPlatform.Win64)
+		bool bWithProxyLOD = Target.Platform == UnrealTargetPlatform.Win64;
+		PrivateDefinitions.Add("WITH_PROXYLOD=" + (bWithProxyLOD ? '1' : '0'));
+		if (bWithProxyLOD)
 		{
 			PrivateDependencyModuleNames.AddRange(
 				new string[]
diff --git a/Engine/Plugins/Runtime/MeshModelingToolset/Source/ModelingOperatorsEditorOnly/Private/CompositionOps/VoxelBooleanMeshesOp.cpp b/Engine/Plugins/Runtime/MeshModelingToolset/Source/ModelingOperatorsEditorOnly/Private/CompositionOps/VoxelBooleanMeshesOp.cpp
--- a/Engine/Plugins/Runtime/MeshModelingToolset/Source/ModelingOperatorsEditorOnly/Private/CompositionOps/VoxelBooleanMeshesOp.cpp
+++ b/Engine/Plugins/Runtime/MeshModelingToolset/Source/ModelingOperatorsEditorOnly/Private/CompositionOps/VoxelBooleanMeshesOp.cpp
@@ -9,23 +9,16 @@
 #include "MeshConstraintsUtil.h"
 #include "CleaningOps/EditNormalsOp.h"
 
-
-// ProxyLOD currently only available on Windows. On other platforms we will make this a no-op
-#if PLATFORM_WINDOWS
-#define HAVE_PROXYLOD 1
-#else
-#define HAVE_PROXYLOD 0
-#endif
-
-#if HAVE_PROXYLOD
+// The ProxyLOD plugin is currently only available on Windows. On other platforms we will make this a no-op.
+#if WITH_PROXYLOD
 #include "ProxyLODVolume.h"
-#endif
+#endif	// WITH_PROXYLOD
 
 using namespace UE::Geometry;
 
 void FVoxelBooleanMeshesOp::CalculateResult(FProgressCancel* Progress)
 {
-#if HAVE_PROXYLOD
+#if WITH_PROXYLOD
 	FMeshDescription ResultMeshDescription;
 	FStaticMeshAttributes Attributes(ResultMeshDescription);
 	Attributes.Register();
@@ -177,13 +170,13 @@ void FVoxelBooleanMeshesOp::CalculateResult(FProgressCancel* Progress)
 		ResultMesh = EditNormalsOp.ExtractResult(); // return the edit normals operator copy to this tool.
 	}
 
-#else	// HAVE_PROXYLOD
+#else	// WITH_PROXYLOD
 
 	FMeshDescriptionToDynamicMesh Converter;
 	Converter.Convert(InputMeshArray[0].Mesh, *ResultMesh);
 	ResultTransform = FTransform3d(InputMeshArray[0].Transform);
 
-#endif	// HAVE_PROXYLOD
+#endif	// WITH_PROXYLOD
 
 }
 
diff --git a/Engine/Plugins/Runtime/MeshModelingToolset/Source/ModelingOperatorsEditorOnly/Private/CompositionOps/VoxelMergeMeshesOp.cpp b/Engine/Plugins/Runtime/MeshModelingToolset/Source/ModelingOperatorsEditorOnly/Private/CompositionOps/VoxelMergeMeshesOp.cpp
--- a/Engine/Plugins/Runtime/MeshModelingToolset/Source/ModelingOperatorsEditorOnly/Private/CompositionOps/VoxelMergeMeshesOp.cpp
+++ b/Engine/Plugins/Runtime/MeshModelingToolset/Source/ModelingOperatorsEditorOnly/Private/CompositionOps/VoxelMergeMeshesOp.cpp
@@ -9,22 +9,16 @@
 #include "MeshConstraintsUtil.h"
 #include "CleaningOps/EditNormalsOp.h"
 
-// ProxyLOD currently only available on Windows. On other platforms we will make this a no-op
-#if PLATFORM_WINDOWS
-#define HAVE_PROXYLOD 1
-#else
-#define HAVE_PROXYLOD 0
-#endif
-
-#if HAVE_PROXYLOD
+// The ProxyLOD plugin is currently only available on Windows. On other platforms we will make this a no-op.
+#if WITH_PROXYLOD
 #include "ProxyLODVolume.h"
-#endif
+#endif	// WITH_PROXYLOD
 
 using namespace UE::Geometry;
 
 void FVoxelMergeMeshesOp::CalculateResult(FProgressCancel* Progress)
 {
-#if HAVE_PROXYLOD
+#if WITH_PROXYLOD
 
 	FMeshDescription MergedMeshesDescription;
 	FStaticMeshAttributes Attributes(MergedMeshesDescription);
@@ -141,13 +135,13 @@ void FVoxelMergeMeshesOp::CalculateResult(FProgressCancel* Progress)
 	}
 
 
-#else	// HAVE_PROXYLOD
+#else	// WITH_PROXYLOD
 
 	FMeshDescriptionToDynamicMesh Converter;
 	Converter.Convert(InputMeshArray[0].Mesh, *ResultMesh);
 	ResultTransform = FTransform3d(InputMeshArray[0].Transform);
 
-#endif	// HAVE_PROXYLOD
+#endif	// WITH_PROXYLOD
 
 }
 
diff --git a/Engine/Plugins/Runtime/MeshModelingToolset/Source/ModelingOperatorsEditorOnly/Private/ParameterizationOps/ParameterizeMeshOp.cpp b/Engine/Plugins/Runtime/MeshModelingToolset/Source/ModelingOperatorsEditorOnly/Private/ParameterizationOps/ParameterizeMeshOp.cpp
--- a/Engine/Plugins/Runtime/MeshModelingToolset/Source/ModelingOperatorsEditorOnly/Private/ParameterizationOps/ParameterizeMeshOp.cpp
+++ b/Engine/Plugins/Runtime/MeshModelingToolset/Source/ModelingOperatorsEditorOnly/Private/ParameterizationOps/ParameterizeMeshOp.cpp
@@ -13,16 +13,13 @@
 #include "DynamicSubmesh3.h"
 #include "XAtlasWrapper.h"
 
-// ProxyLOD currently only available on Windows. On other platforms we will make this a no-op
-#if PLATFORM_WINDOWS
-#define HAVE_PROXYLOD 1
-#else
-#define HAVE_PROXYLOD 0
-#endif
-
-#if HAVE_PROXYLOD
+// The ProxyLOD plugin is currently only available on Windows.
+#if WITH_PROXYLOD
 #include "ProxyLODParameterization.h"
-#endif
+#endif	// WITH_PROXYLOD
+
+// We do not have UVAtlas available if ProxyLOD is not available.
+#define NO_UVATLAS WITH_PROXYLOD ? 0 : 1
 
 using namespace UE::Geometry;
 
@@ -235,6 +232,13 @@ void FParameterizeMeshOp::CopyNewUVsToMesh(
 
 bool FParameterizeMeshOp::ComputeUVs_UVAtlas(FDynamicMesh3& Mesh,  TFunction<bool(float)>& Interrupter)
 {
+#if NO_UVATLAS
+
+	ensureMsgf(false, TEXT("UVAtlas not available; this should not be called!"));
+	return false;
+
+#else	// NO_UVATLAS
+	
 	TRACE_CPUPROFILER_EVENT_SCOPE(ParameterizeMeshOp_ComputeUVs_UVAtlas);
 
 	// the UVAtlas code is unhappy if you feed it a single degenerate triangle
@@ -277,17 +281,11 @@ bool FParameterizeMeshOp::ComputeUVs_UVAtlas(FDynamicMesh3& Mesh,  TFunction<boo
 	float MaxStretch     = Stretch;
 	int32 MaxChartNumber = NumCharts;
 
-#if HAVE_PROXYLOD
 	TUniquePtr<IProxyLODParameterization> ParameterizationTool = IProxyLODParameterization::CreateTool();
 	bool bSuccess = ParameterizationTool->GenerateUVs(Width, Height, Gutter, LinearMesh.VertexBuffer,
 	                                                  LinearMesh.IndexBuffer, LinearMesh.AdjacencyBuffer, Interrupter,
 	                                                  UVVertexBuffer, UVIndexBuffer, VertexRemapArray, MaxStretch,
 	                                                  MaxChartNumber);
-#else
-	ensureMsgf(false, TEXT("ProxyLOD not available, this should not be called"));
-	bool bSuccess = false;
-#endif
-
 
 	// Add the UVs to the FDynamicMesh
 	if (bSuccess)
@@ -300,6 +298,8 @@ bool FParameterizeMeshOp::ComputeUVs_UVAtlas(FDynamicMesh3& Mesh,  TFunction<boo
 	}
 
 	return bSuccess;
+
+#endif	// NO_UVATLAS
 }
 
 
@@ -424,11 +424,12 @@ void FParameterizeMeshOp::CalculateResult(FProgressCancel* Progress)
 	// The UV atlas callback uses a float progress-based interrupter. 
 	TFunction<bool(float)> Interrupter = [Progress](float)->bool {return !(Progress && Progress->Cancelled()); };
 
-#if HAVE_PROXYLOD
-	bool bUseXAtlas = (Method == EParamOpBackend::XAtlas);
-#else
-	bool bUseXAtlas = true;
-#endif
+#if NO_UVATLAS
+	constexpr bool bUseXAtlas = true;
+	UE_LOG(LogGeometry, Warning, TEXT("AutoUV method UVAtlas not available; falling back to XAtlas instead."));
+#else	// NO_UVATLAS
+	const bool bUseXAtlas = (Method == EParamOpBackend::XAtlas);
+#endif	// NO_UVATLAS
 
 	NewResultInfo = FGeometryResult(EGeometryResultType::InProgress);
 
diff --git a/Engine/Plugins/Runtime/MeshModelingToolset/Source/ModelingOperatorsEditorOnly/Public/ParameterizationOps/ParameterizeMeshOp.h b/Engine/Plugins/Runtime/MeshModelingToolset/Source/ModelingOperatorsEditorOnly/Public/ParameterizationOps/ParameterizeMeshOp.h
--- a/Engine/Plugins/Runtime/MeshModelingToolset/Source/ModelingOperatorsEditorOnly/Public/ParameterizationOps/ParameterizeMeshOp.h
+++ b/Engine/Plugins/Runtime/MeshModelingToolset/Source/ModelingOperatorsEditorOnly/Public/ParameterizationOps/ParameterizeMeshOp.h
@@ -23,8 +23,8 @@ namespace Geometry
 enum class EParamOpBackend
 {
 	PatchBuilder = 0,
-	XAtlas = 1,
-	UVAtlas = 2
+	UVAtlas = 1,
+	XAtlas = 2
 };
 
 class MODELINGOPERATORSEDITORONLY_API FParameterizeMeshOp : public FDynamicMeshOperator
diff --git a/Engine/Plugins/Runtime/MeshModelingToolset/Source/ModelingOperatorsEditorOnly/Public/Properties/ParameterizeMeshProperties.h b/Engine/Plugins/Runtime/MeshModelingToolset/Source/ModelingOperatorsEditorOnly/Public/Properties/ParameterizeMeshProperties.h
--- a/Engine/Plugins/Runtime/MeshModelingToolset/Source/ModelingOperatorsEditorOnly/Public/Properties/ParameterizeMeshProperties.h
+++ b/Engine/Plugins/Runtime/MeshModelingToolset/Source/ModelingOperatorsEditorOnly/Public/Properties/ParameterizeMeshProperties.h
@@ -11,14 +11,13 @@ UENUM()
 enum class EParameterizeMeshUVMethod
 {
 	// Keep values the same as UE::Geometry::EParamOpBackend!
-	// The order is different on purpose to have them sorted alphabetically in the UI.
 
 	/** Compute automatic UVs using the Patch Builder technique */
 	PatchBuilder = 0,
 	/** Compute automatic UVs using the UVAtlas technique */
-	UVAtlas = 2,
+	UVAtlas = 1,
 	/** Compute automatic UVs using the XAtlas technique */
-	XAtlas = 1,
+	XAtlas = 2,
 };
 
 
