MeshModelingToolset: Fix polygroup edge rendering after undoing a `Straighten` edge edit in the PolyEd tool

#rnx

diff --git a/Engine/Plugins/Runtime/MeshModelingToolset/Source/MeshModelingTools/Private/EditMeshPolygonsTool.cpp b/Engine/Plugins/Runtime/MeshModelingToolset/Source/MeshModelingTools/Private/EditMeshPolygonsTool.cpp
--- a/Engine/Plugins/Runtime/MeshModelingToolset/Source/MeshModelingTools/Private/EditMeshPolygonsTool.cpp
+++ b/Engine/Plugins/Runtime/MeshModelingToolset/Source/MeshModelingTools/Private/EditMeshPolygonsTool.cpp
@@ -344,10 +344,9 @@ void UEditMeshPolygonsTool::Setup()
 	};
 	ActivityContext->EmitCurrentMeshChangeAndUpdate = [this](const FText& TransactionLabel,
 		TUniquePtr<UE::Geometry::FDynamicMeshChange> MeshChangeIn,
-		const UE::Geometry::FGroupTopologySelection& OutputSelection, 
-		bool bTopologyModified) 
+		const UE::Geometry::FGroupTopologySelection& OutputSelection) 
 	{
-		EmitCurrentMeshChangeAndUpdate(TransactionLabel, MoveTemp(MeshChangeIn), OutputSelection, bTopologyModified);
+		EmitCurrentMeshChangeAndUpdate(TransactionLabel, MoveTemp(MeshChangeIn), OutputSelection);
 	};
 	GetToolManager()->GetContextObjectStore()->RemoveContextObjectsOfType(UPolyEditActivityContext::StaticClass());
 	GetToolManager()->GetContextObjectStore()->AddContextObject(ActivityContext);
@@ -1167,7 +1166,8 @@ void UEditMeshPolygonsTool::ApplyMerge()
 		NewSelection.SelectedGroupIDs.Add(NewGroupID);
 	}
 	
-	EmitCurrentMeshChangeAndUpdate(LOCTEXT("PolyMeshMergeChange", "Merge"), ChangeTracker.EndChange(), NewSelection, true);
+	EmitCurrentMeshChangeAndUpdate(LOCTEXT("PolyMeshMergeChange", "Merge"),
+		ChangeTracker.EndChange(), NewSelection);
 }
 
 
@@ -1203,7 +1203,8 @@ void UEditMeshPolygonsTool::ApplyDelete()
 	Editor.RemoveTriangles(ActiveTriangleSelection, true);
 
 	FGroupTopologySelection NewSelection;
-	EmitCurrentMeshChangeAndUpdate(LOCTEXT("PolyMeshDeleteChange", "Delete"), ChangeTracker.EndChange(), NewSelection, true);
+	EmitCurrentMeshChangeAndUpdate(LOCTEXT("PolyMeshDeleteChange", "Delete"),
+		ChangeTracker.EndChange(), NewSelection);
 }
 
 
@@ -1231,9 +1232,9 @@ void UEditMeshPolygonsTool::ApplyRecalcNormals()
 
 	// We actually don't even need any of the wrapper around this change since we're not altering
 	// positions or topology (so no other structures need updating), but we go ahead and go the
-	// same route as everything else.
-	EmitCurrentMeshChangeAndUpdate(LOCTEXT("PolyMeshRecalcNormalsChange", "Recalc Normals"), 
-		ChangeTracker.EndChange(), ActiveSelection, false);
+	// same route as everything else. See :HandlePositionOnlyMeshChanges
+	EmitCurrentMeshChangeAndUpdate(LOCTEXT("PolyMeshRecalcNormalsChange", "Recalculate Normals"), 
+		ChangeTracker.EndChange(), ActiveSelection);
 }
 
 
@@ -1263,7 +1264,7 @@ void UEditMeshPolygonsTool::ApplyFlipNormals()
 
 	// Note the topology can change in that the ordering of edge elements can reverse
 	EmitCurrentMeshChangeAndUpdate(LOCTEXT("PolyMeshFlipNormalsChange", "Flip Normals"), 
-		ChangeTracker.EndChange(), ActiveSelection, true);
+		ChangeTracker.EndChange(), ActiveSelection);
 }
 
 
@@ -1343,8 +1344,8 @@ void UEditMeshPolygonsTool::ApplyRetriangulate()
 		GetToolManager()->DisplayMessage(LOCTEXT("OnRetriangulateFailures", "Some faces could not be retriangulated"), EToolMessageLevel::UserWarning);
 	}
 
-	EmitCurrentMeshChangeAndUpdate(LOCTEXT("PolyMeshRetriangulateChange", "Retriangulate"), 
-		ChangeTracker.EndChange(), ActiveSelection, true);
+	EmitCurrentMeshChangeAndUpdate(LOCTEXT("PolyMeshRetriangulateChange", "Retriangulate"),
+		ChangeTracker.EndChange(), ActiveSelection);
 }
 
 
@@ -1365,7 +1366,8 @@ void UEditMeshPolygonsTool::SimplifyByGroups()
 
 	FGroupTopologySelection NewSelection; // Empty the selection
 
-	EmitCurrentMeshChangeAndUpdate(LOCTEXT("PolyMeshSimplifyByGroup", "Simplify by Group"), ChangeTracker.EndChange(), NewSelection, true);
+	EmitCurrentMeshChangeAndUpdate(LOCTEXT("PolyMeshSimplifyByGroup", "Simplify by Group"),
+		ChangeTracker.EndChange(), NewSelection);
 }
 
 
@@ -1394,8 +1396,8 @@ void UEditMeshPolygonsTool::ApplyDecompose()
 		}
 	}
 
-
-	EmitCurrentMeshChangeAndUpdate(LOCTEXT("PolyMeshDecomposeChange", "Decompose"), ChangeTracker.EndChange(), NewSelection, true);
+	EmitCurrentMeshChangeAndUpdate(LOCTEXT("PolyMeshDecomposeChange", "Decompose"),
+		ChangeTracker.EndChange(), NewSelection);
 }
 
 
@@ -1420,7 +1422,8 @@ void UEditMeshPolygonsTool::ApplyDisconnect()
 	FDynamicMeshEditor Editor(Mesh);
 	Editor.DisconnectTriangles(AllTriangles, false);
 
-	EmitCurrentMeshChangeAndUpdate(LOCTEXT("PolyMeshDisconnectChange", "Disconnect"), ChangeTracker.EndChange(), ActiveSelection, true);
+	EmitCurrentMeshChangeAndUpdate(LOCTEXT("PolyMeshDisconnectChange", "Disconnect"),
+		ChangeTracker.EndChange(), ActiveSelection);
 }
 
 
@@ -1451,7 +1454,8 @@ void UEditMeshPolygonsTool::ApplyDuplicate()
 	FGroupTopologySelection NewSelection;
 	NewSelection.SelectedGroupIDs.Append(bTriangleMode ? EditResult.NewTriangles : EditResult.NewGroups);
 
-	EmitCurrentMeshChangeAndUpdate(LOCTEXT("PolyMeshDisconnectChange", "Disconnect"), ChangeTracker.EndChange(), NewSelection, true);
+	EmitCurrentMeshChangeAndUpdate(LOCTEXT("PolyMeshDisconnectChange", "Disconnect"),
+		ChangeTracker.EndChange(), NewSelection);
 }
 
 
@@ -1489,7 +1493,8 @@ void UEditMeshPolygonsTool::ApplyCollapseEdge()
 
 	// emit undo
 	FGroupTopologySelection NewSelection;
-	EmitCurrentMeshChangeAndUpdate(LOCTEXT("PolyMeshEdgeCollapseChange", "Collapse"), ChangeTracker.EndChange(), NewSelection, true);
+	EmitCurrentMeshChangeAndUpdate(LOCTEXT("PolyMeshEdgeCollapseChange", "Collapse"),
+		ChangeTracker.EndChange(), NewSelection);
 }
 
 
@@ -1528,7 +1533,8 @@ void UEditMeshPolygonsTool::ApplyWeldEdges()
 	}
 
 	FGroupTopologySelection NewSelection;
-	EmitCurrentMeshChangeAndUpdate(LOCTEXT("PolyMeshWeldEdgeChange", "Weld Edges"), ChangeTracker.EndChange(), NewSelection, true);
+	EmitCurrentMeshChangeAndUpdate(LOCTEXT("PolyMeshWeldEdgeChange", "Weld Edges"),
+		ChangeTracker.EndChange(), NewSelection);
 }
 
 
@@ -1563,9 +1569,14 @@ void UEditMeshPolygonsTool::ApplyStraightenEdges()
 		}
 	}
 
+	// TODO :HandlePositionOnlyMeshChanges Due to the group topology storing edge IDs that do not stay the same across
+	// undo/redo events even when the mesh topology stays the same after a FDynamicMeshChange, we actually have to treat
+	// all FDynamicMeshChange-based transactions as affecting group topology. Here we only changed vertex positions so
+	// we could add a separate overload that takes a FMeshVertexChange, and possibly one that takes an attribute change
+	// (or unify the three via an interface)
 	FGroupTopologySelection NewSelection;
 	EmitCurrentMeshChangeAndUpdate(LOCTEXT("PolyMeshStraightenEdgeChange", "Straighten Edges"), 
-		ChangeTracker.EndChange(), NewSelection, false);
+		ChangeTracker.EndChange(), NewSelection);
 }
 
 
@@ -1622,7 +1633,8 @@ void UEditMeshPolygonsTool::ApplyFillHole()
 		}
 	}
 
-	EmitCurrentMeshChangeAndUpdate(LOCTEXT("PolyMeshFillHoleChange", "Fill Hole"), ChangeTracker.EndChange(), NewSelection, true);
+	EmitCurrentMeshChangeAndUpdate(LOCTEXT("PolyMeshFillHoleChange", "Fill Hole"),
+		ChangeTracker.EndChange(), NewSelection);
 }
 
 
@@ -1652,7 +1664,8 @@ void UEditMeshPolygonsTool::ApplyPokeSingleFace()
 		}
 	}
 
-	EmitCurrentMeshChangeAndUpdate(LOCTEXT("PolyMeshPokeChange", "Poke Faces"), ChangeTracker.EndChange(), NewSelection, true);
+	EmitCurrentMeshChangeAndUpdate(LOCTEXT("PolyMeshPokeChange", "Poke Faces"),
+		ChangeTracker.EndChange(), NewSelection);
 }
 
 
@@ -1683,7 +1696,8 @@ void UEditMeshPolygonsTool::ApplyFlipSingleEdge()
 	}
 
 	// Group topology may or may not change, but just assume that it does.
-	EmitCurrentMeshChangeAndUpdate(LOCTEXT("PolyMeshFlipChange", "Flip Edges"), ChangeTracker.EndChange(), ActiveSelection, true);
+	EmitCurrentMeshChangeAndUpdate(LOCTEXT("PolyMeshFlipChange", "Flip Edges"),
+		ChangeTracker.EndChange(), ActiveSelection);
 }
 
 void UEditMeshPolygonsTool::ApplyCollapseSingleEdge()
@@ -1726,7 +1740,7 @@ void UEditMeshPolygonsTool::ApplyCollapseSingleEdge()
 	}
 
 	EmitCurrentMeshChangeAndUpdate(LOCTEXT("PolyMeshCollapseChange", "Collapse Edges"), 
-		ChangeTracker.EndChange(), FGroupTopologySelection(), true);
+		ChangeTracker.EndChange(), FGroupTopologySelection());
 }
 
 void UEditMeshPolygonsTool::ApplySplitSingleEdge()
@@ -1766,7 +1780,8 @@ void UEditMeshPolygonsTool::ApplySplitSingleEdge()
 		}
 	}
 
-	EmitCurrentMeshChangeAndUpdate(LOCTEXT("PolyMeshSplitChange", "Split Edges"), ChangeTracker.EndChange(), NewSelection, true);
+	EmitCurrentMeshChangeAndUpdate(LOCTEXT("PolyMeshSplitChange", "Split Edges"),
+		ChangeTracker.EndChange(), NewSelection);
 }
 
 
@@ -1801,8 +1816,7 @@ bool UEditMeshPolygonsTool::BeginMeshFaceEditChange()
 
 void UEditMeshPolygonsTool::EmitCurrentMeshChangeAndUpdate(const FText& TransactionLabel,
 	TUniquePtr<FDynamicMeshChange> MeshChangeIn,
-	const FGroupTopologySelection& OutputSelection, 
-	bool bGroupTopologyModified)
+	const FGroupTopologySelection& OutputSelection)
 {
 	// open top-level transaction
 	GetToolManager()->BeginUndoTransaction(TransactionLabel);
@@ -1819,11 +1833,11 @@ void UEditMeshPolygonsTool::EmitCurrentMeshChangeAndUpdate(const FText& Transact
 	const FGroupTopologySelection* OutputSelectionToUse = &OutputSelection;
 
 	// If the selection is going to be cleared, we need to do it explicitly ourselves so that we can emit a change.
-	if (!SelectionMechanic->GetActiveSelection().IsEmpty() && (bSelectionModified || bGroupTopologyModified))
+	if (!SelectionMechanic->GetActiveSelection().IsEmpty() && bSelectionModified)
 	{
 		if (bReferencingSameSelection)
 		{
-			// Need to make a copy because OutputSelection will get cleared due to bGroupTopologyModified
+			// Need to make a copy because OutputSelection will get cleared
 			TempSelection = OutputSelection;
 			OutputSelectionToUse = &TempSelection;
 		}
@@ -1834,16 +1848,16 @@ void UEditMeshPolygonsTool::EmitCurrentMeshChangeAndUpdate(const FText& Transact
 	}
 
 	GetToolManager()->EmitObjectChange(this, 
-		MakeUnique<FEditMeshPolygonsToolMeshChange>(MoveTemp(MeshChangeIn), bGroupTopologyModified),
+		MakeUnique<FEditMeshPolygonsToolMeshChange>(MoveTemp(MeshChangeIn)),
 		TransactionLabel);
 
 	// Update related structures
-	UpdateFromCurrentMesh(bGroupTopologyModified);
-	ModifiedTopologyCounter += bGroupTopologyModified;
+	UpdateFromCurrentMesh(true);
+	ModifiedTopologyCounter += 1;
 
 	// Set output selection either if we changed selections (to something non-empty), or if
 	// our selection got cleared due to bGroupTopologyModified.
-	if (!OutputSelectionToUse->IsEmpty() && (bSelectionModified || bGroupTopologyModified))
+	if (!OutputSelectionToUse->IsEmpty() && bSelectionModified)
 	{
 		SelectionMechanic->BeginChange();
 		SelectionMechanic->SetSelection(*OutputSelectionToUse);
@@ -1981,6 +1995,10 @@ void FEditMeshPolygonsToolMeshChange::Apply(UObject* Object)
 {
 	UEditMeshPolygonsTool* Tool = Cast<UEditMeshPolygonsTool>(Object);
 	
+	// This function currently only supports FDynamicMeshChange but that should be issued only when the mesh changes
+	// topology. For now we use it even when eg vertex postions change. See :HandlePositionOnlyMeshChanges
+	bool bGroupTopologyModified = true;
+	
 	MeshChange->Apply(Tool->CurrentMesh.Get(), false);
 	Tool->UpdateFromCurrentMesh(bGroupTopologyModified);
 	Tool->ModifiedTopologyCounter += bGroupTopologyModified;
@@ -1991,6 +2009,10 @@ void FEditMeshPolygonsToolMeshChange::Revert(UObject* Object)
 {
 	UEditMeshPolygonsTool* Tool = Cast<UEditMeshPolygonsTool>(Object);
 	
+	// This function currently only supports FDynamicMeshChange but that should be issued only when the mesh changes
+	// topology. For now we use it even when eg vertex postions change. See :HandlePositionOnlyMeshChanges
+	bool bGroupTopologyModified = true;
+	
 	MeshChange->Apply(Tool->CurrentMesh.Get(), true);
 	Tool->UpdateFromCurrentMesh(bGroupTopologyModified);
 	Tool->ModifiedTopologyCounter -= bGroupTopologyModified;
diff --git a/Engine/Plugins/Runtime/MeshModelingToolset/Source/MeshModelingTools/Private/ToolActivities/PolyEditBevelEdgeActivity.cpp b/Engine/Plugins/Runtime/MeshModelingToolset/Source/MeshModelingTools/Private/ToolActivities/PolyEditBevelEdgeActivity.cpp
--- a/Engine/Plugins/Runtime/MeshModelingToolset/Source/MeshModelingTools/Private/ToolActivities/PolyEditBevelEdgeActivity.cpp
+++ b/Engine/Plugins/Runtime/MeshModelingToolset/Source/MeshModelingTools/Private/ToolActivities/PolyEditBevelEdgeActivity.cpp
@@ -287,7 +287,7 @@ void UPolyEditBevelEdgeActivity::ApplyBevel()
 
 	// Emit undo  (also updates relevant structures)
 	ActivityContext->EmitCurrentMeshChangeAndUpdate(LOCTEXT("PolyMeshBevelChange", "Bevel"),
-		ChangeTracker.EndChange(), NewSelection, true);
+		ChangeTracker.EndChange(), NewSelection);
 
 	ActiveSelection = NewSelection;
 }
diff --git a/Engine/Plugins/Runtime/MeshModelingToolset/Source/MeshModelingTools/Private/ToolActivities/PolyEditCutFacesActivity.cpp b/Engine/Plugins/Runtime/MeshModelingToolset/Source/MeshModelingTools/Private/ToolActivities/PolyEditCutFacesActivity.cpp
--- a/Engine/Plugins/Runtime/MeshModelingToolset/Source/MeshModelingTools/Private/ToolActivities/PolyEditCutFacesActivity.cpp
+++ b/Engine/Plugins/Runtime/MeshModelingToolset/Source/MeshModelingTools/Private/ToolActivities/PolyEditCutFacesActivity.cpp
@@ -186,7 +186,7 @@ void UPolyEditCutFacesActivity::ApplyCutFaces()
 
 	// Emit undo (also updates relevant structures)
 	ActivityContext->EmitCurrentMeshChangeAndUpdate(LOCTEXT("PolyMeshCutFacesChange", "Cut Faces"),
-		ChangeTracker.EndChange(), OutputSelection, true);
+		ChangeTracker.EndChange(), OutputSelection);
 
 	// End activity
 	Clear();
diff --git a/Engine/Plugins/Runtime/MeshModelingToolset/Source/MeshModelingTools/Private/ToolActivities/PolyEditExtrudeActivity.cpp b/Engine/Plugins/Runtime/MeshModelingToolset/Source/MeshModelingTools/Private/ToolActivities/PolyEditExtrudeActivity.cpp
--- a/Engine/Plugins/Runtime/MeshModelingToolset/Source/MeshModelingTools/Private/ToolActivities/PolyEditExtrudeActivity.cpp
+++ b/Engine/Plugins/Runtime/MeshModelingToolset/Source/MeshModelingTools/Private/ToolActivities/PolyEditExtrudeActivity.cpp
@@ -455,7 +455,7 @@ void UPolyEditExtrudeActivity::ApplyExtrude()
 
 	// Emit undo  (also updates relevant structures)
 	ActivityContext->EmitCurrentMeshChangeAndUpdate(LOCTEXT("PolyMeshExtrudeChange", "Extrude"),
-		ChangeTracker.EndChange(), NewSelection, true);
+		ChangeTracker.EndChange(), NewSelection);
 
 	ActiveSelection = NewSelection;
 }
diff --git a/Engine/Plugins/Runtime/MeshModelingToolset/Source/MeshModelingTools/Private/ToolActivities/PolyEditInsertEdgeActivity.cpp b/Engine/Plugins/Runtime/MeshModelingToolset/Source/MeshModelingTools/Private/ToolActivities/PolyEditInsertEdgeActivity.cpp
--- a/Engine/Plugins/Runtime/MeshModelingToolset/Source/MeshModelingTools/Private/ToolActivities/PolyEditInsertEdgeActivity.cpp
+++ b/Engine/Plugins/Runtime/MeshModelingToolset/Source/MeshModelingTools/Private/ToolActivities/PolyEditInsertEdgeActivity.cpp
@@ -253,7 +253,7 @@ void UPolyEditInsertEdgeActivity::Tick(float DeltaTime)
 			// Emit transaction
 			FGroupTopologySelection EmptySelection;
 			ActivityContext->EmitCurrentMeshChangeAndUpdate(LOCTEXT("EdgeInsertionTransactionName", "Edge Insertion"),
-				ChangeTracker.EndChange(), EmptySelection, true);
+				ChangeTracker.EndChange(), EmptySelection);
 
 			ToolState = EState::GettingStart;
 		}
diff --git a/Engine/Plugins/Runtime/MeshModelingToolset/Source/MeshModelingTools/Private/ToolActivities/PolyEditInsertEdgeLoopActivity.cpp b/Engine/Plugins/Runtime/MeshModelingToolset/Source/MeshModelingTools/Private/ToolActivities/PolyEditInsertEdgeLoopActivity.cpp
--- a/Engine/Plugins/Runtime/MeshModelingToolset/Source/MeshModelingTools/Private/ToolActivities/PolyEditInsertEdgeLoopActivity.cpp
+++ b/Engine/Plugins/Runtime/MeshModelingToolset/Source/MeshModelingTools/Private/ToolActivities/PolyEditInsertEdgeLoopActivity.cpp
@@ -295,7 +295,7 @@ void UPolyEditInsertEdgeLoopActivity::Tick(float DeltaTime)
 			// Emit transaction
 			FGroupTopologySelection EmptySelection;
 			ActivityContext->EmitCurrentMeshChangeAndUpdate(LOCTEXT("EdgeLoopInsertionTransactionName", "Edge Loop Insertion"),
-				ChangeTracker.EndChange(), EmptySelection, true);
+				ChangeTracker.EndChange(), EmptySelection);
 		}
 
 		PreviewEdges.Reset();
diff --git a/Engine/Plugins/Runtime/MeshModelingToolset/Source/MeshModelingTools/Private/ToolActivities/PolyEditInsetOutsetActivity.cpp b/Engine/Plugins/Runtime/MeshModelingToolset/Source/MeshModelingTools/Private/ToolActivities/PolyEditInsetOutsetActivity.cpp
--- a/Engine/Plugins/Runtime/MeshModelingToolset/Source/MeshModelingTools/Private/ToolActivities/PolyEditInsetOutsetActivity.cpp
+++ b/Engine/Plugins/Runtime/MeshModelingToolset/Source/MeshModelingTools/Private/ToolActivities/PolyEditInsetOutsetActivity.cpp
@@ -189,7 +189,7 @@ void UPolyEditInsetOutsetActivity::ApplyInset()
 
 	// Emit undo (also updates relevant structures)
 	ActivityContext->EmitCurrentMeshChangeAndUpdate(LOCTEXT("PolyMeshInsetOutsetChange", "Inset/Outset"),
-		Inset.ChangeTracker->EndChange(), ActiveSelection, true);
+		Inset.ChangeTracker->EndChange(), ActiveSelection);
 }
 
 void UPolyEditInsetOutsetActivity::Clear()
diff --git a/Engine/Plugins/Runtime/MeshModelingToolset/Source/MeshModelingTools/Private/ToolActivities/PolyEditPlanarProjectionUVActivity.cpp b/Engine/Plugins/Runtime/MeshModelingToolset/Source/MeshModelingTools/Private/ToolActivities/PolyEditPlanarProjectionUVActivity.cpp
--- a/Engine/Plugins/Runtime/MeshModelingToolset/Source/MeshModelingTools/Private/ToolActivities/PolyEditPlanarProjectionUVActivity.cpp
+++ b/Engine/Plugins/Runtime/MeshModelingToolset/Source/MeshModelingTools/Private/ToolActivities/PolyEditPlanarProjectionUVActivity.cpp
@@ -195,9 +195,10 @@ void UPolyEditPlanarProjectionUVActivity::ApplySetUVs()
 	FDynamicMeshEditor Editor(ActivityContext->CurrentMesh.Get());
 	Editor.SetTriangleUVsFromProjection(ActiveTriangleSelection, PlanarFrame, UVScale, FVector2f::Zero(), false, 0);
 
-	// Emit undo (also updates relevant structures)
+	// Emit undo (also updates relevant structures). We didn't change the mesh topology here but for now we use
+	// the same route as everything else. See :HandlePositionOnlyMeshChanges
 	ActivityContext->EmitCurrentMeshChangeAndUpdate(LOCTEXT("PolyMeshSetUVsChange", "Set UVs"),
-		ChangeTracker.EndChange(), ActiveSelection, false);
+		ChangeTracker.EndChange(), ActiveSelection);
 
 	// End activity
 	Clear();
diff --git a/Engine/Plugins/Runtime/MeshModelingToolset/Source/MeshModelingTools/Public/EditMeshPolygonsTool.h b/Engine/Plugins/Runtime/MeshModelingToolset/Source/MeshModelingTools/Public/EditMeshPolygonsTool.h
--- a/Engine/Plugins/Runtime/MeshModelingToolset/Source/MeshModelingTools/Public/EditMeshPolygonsTool.h
+++ b/Engine/Plugins/Runtime/MeshModelingToolset/Source/MeshModelingTools/Public/EditMeshPolygonsTool.h
@@ -25,6 +25,7 @@ class UDragAlignmentMechanic;
 class UMeshOpPreviewWithBackgroundCompute; 
 class FMeshVertexChangeBuilder;
 class UPersistentMeshSelection;
+class UEditMeshPolygonsTool;
 class UPolyEditActivityContext;
 class UPolyEditInsertEdgeActivity;
 class UPolyEditInsertEdgeLoopActivity;
@@ -583,8 +584,7 @@ protected:
 	// Emit an undoable change to CurrentMesh and update related structures (preview, spatial, etc)
 	void EmitCurrentMeshChangeAndUpdate(const FText& TransactionLabel,
 		TUniquePtr<UE::Geometry::FDynamicMeshChange> MeshChangeIn,
-		const UE::Geometry::FGroupTopologySelection& OutputSelection,
-		bool bTopologyChanged);
+		const UE::Geometry::FGroupTopologySelection& OutputSelection);
 
 	// Emit an undoable start of an activity
 	void EmitActivityStart(const FText& TransactionLabel);
@@ -671,9 +671,8 @@ protected:
 class MESHMODELINGTOOLS_API FEditMeshPolygonsToolMeshChange : public FToolCommandChange
 {
 public:
-	FEditMeshPolygonsToolMeshChange(TUniquePtr<UE::Geometry::FDynamicMeshChange> MeshChangeIn, bool bGroupTopologyModified)
+	FEditMeshPolygonsToolMeshChange(TUniquePtr<UE::Geometry::FDynamicMeshChange> MeshChangeIn)
 		: MeshChange(MoveTemp(MeshChangeIn))
-		, bGroupTopologyModified(bGroupTopologyModified)
 	{};
 
 	virtual void Apply(UObject* Object) override;
@@ -682,7 +681,6 @@ public:
 
 protected:
 	TUniquePtr<UE::Geometry::FDynamicMeshChange> MeshChange;
-	bool bGroupTopologyModified;
 };
 
 
diff --git a/Engine/Plugins/Runtime/MeshModelingToolset/Source/MeshModelingTools/Public/ToolActivities/PolyEditActivityContext.h b/Engine/Plugins/Runtime/MeshModelingToolset/Source/MeshModelingTools/Public/ToolActivities/PolyEditActivityContext.h
--- a/Engine/Plugins/Runtime/MeshModelingToolset/Source/MeshModelingTools/Public/ToolActivities/PolyEditActivityContext.h
+++ b/Engine/Plugins/Runtime/MeshModelingToolset/Source/MeshModelingTools/Public/ToolActivities/PolyEditActivityContext.h
@@ -66,8 +66,7 @@ public:
 	 */
 	TUniqueFunction<void (const FText& TransactionLabel,
 			TUniquePtr<UE::Geometry::FDynamicMeshChange> MeshChangeIn,
-			const UE::Geometry::FGroupTopologySelection& OutputSelection,
-			bool bGroupTopologyChanged)> 
+			const UE::Geometry::FGroupTopologySelection& OutputSelection)> 
 		EmitCurrentMeshChangeAndUpdate;
 
 	/**
