MeshModelingTools: Make PolyExt snapping options grey out in a way that corresponds to current behavior, make them respond to the viewport snapping toggle, and make the Shift key disable surface snapping the way it does for all the other snapping types.

#rnx

diff --git a/Engine/Plugins/Runtime/MeshModelingToolset/Source/MeshModelingTools/Private/DrawPolygonTool.cpp b/Engine/Plugins/Runtime/MeshModelingToolset/Source/MeshModelingTools/Private/DrawPolygonTool.cpp
--- a/Engine/Plugins/Runtime/MeshModelingToolset/Source/MeshModelingTools/Private/DrawPolygonTool.cpp
+++ b/Engine/Plugins/Runtime/MeshModelingToolset/Source/MeshModelingTools/Private/DrawPolygonTool.cpp
@@ -158,6 +158,8 @@ void UDrawPolygonTool::Setup()
 
 	SnapProperties = NewObject<UDrawPolygonToolSnapProperties>(this);
 	SnapProperties->RestoreProperties(this);
+	SnapProperties->bSnapToWorldGrid = GetToolManager()->GetContextQueriesAPI()
+		->GetCurrentSnappingSettings().bEnablePositionGridSnapping;
 
 	// register tool properties
 	if (OutputTypeProperties->ShouldShowPropertySet())
@@ -247,6 +249,16 @@ void UDrawPolygonTool::OnTick(float DeltaTime)
 		// faster to do this as an override rather than destroying/recreating the gizmo via UpdateShowGizmoState
 		PlaneTransformGizmo->SetVisibility(AllowDrawPlaneUpdates());
 	}
+	if (SnapProperties)
+	{
+		bool bSnappingEnabledInViewport = GetToolManager()->GetContextQueriesAPI()
+			->GetCurrentSnappingSettings().bEnablePositionGridSnapping;
+		if (SnapProperties->bSnapToWorldGrid != bSnappingEnabledInViewport)
+		{
+			SnapProperties->bSnapToWorldGrid = bSnappingEnabledInViewport;
+			NotifyOfPropertyChangeByTool(SnapProperties);
+		}
+	}
 }
 
 
@@ -552,16 +564,13 @@ bool UDrawPolygonTool::FindDrawPlaneHitPoint(const FInputDeviceRay& ClickPos, FV
 	}
 	else 
 	{
-		if (SnapProperties->bSnapToWorldGrid)
+		FVector3d WorldGridSnapPos;
+		if (ToolSceneQueriesUtil::FindWorldGridSnapPoint(this, HitPos, WorldGridSnapPos))
 		{
-			FVector3d WorldGridSnapPos;
-			if (ToolSceneQueriesUtil::FindWorldGridSnapPoint(this, HitPos, WorldGridSnapPos))
-			{
-				WorldGridSnapPos = Frame.ToPlane(WorldGridSnapPos, 2);
-				SnapEngine.AddPointTarget(WorldGridSnapPos, CurrentGridSnapID, 
-					FBasePositionSnapSolver3::FCustomMetric::Replace(999), SnapEngine.MinInternalPriority() - 5);
-				LastGridSnapPoint = WorldGridSnapPos;
-			}
+			WorldGridSnapPos = Frame.ToPlane(WorldGridSnapPos, 2);
+			SnapEngine.AddPointTarget(WorldGridSnapPos, CurrentGridSnapID, 
+				FBasePositionSnapSolver3::FCustomMetric::Replace(999), SnapEngine.MinInternalPriority() - 5);
+			LastGridSnapPoint = WorldGridSnapPos;
 		}
 
 		if (SnapProperties->bSnapToVertices || SnapProperties->bSnapToEdges)
@@ -607,8 +616,8 @@ bool UDrawPolygonTool::FindDrawPlaneHitPoint(const FInputDeviceRay& ClickPos, FV
 	}
 
 
-	// if not snap and we want to hit objects, do that
-	if (SnapProperties->bSnapToSurfaces)
+	// if not yet snapped and we want to hit objects, do that
+	if (SnapProperties->bSnapToSurfaces && !bIgnoreSnappingToggle)
 	{
 		FHitResult Result;
 		bool bWorldHit = ToolSceneQueriesUtil::FindNearestVisibleObjectHit(this, Result, ClickPos.WorldRay);
@@ -976,7 +985,7 @@ void UDrawPolygonTool::BeginInteractiveExtrude()
 	};
 	HeightMechanic->WorldPointSnapFunc = [this](const FVector3d& WorldPos, FVector3d& SnapPos)
 	{
-		if (bIgnoreSnappingToggle == false && SnapProperties->bEnableSnapping && SnapProperties->bSnapToWorldGrid)
+		if (bIgnoreSnappingToggle == false && SnapProperties->bEnableSnapping)
 		{
 			return ToolSceneQueriesUtil::FindWorldGridSnapPoint(this, WorldPos, SnapPos);
 		}
diff --git a/Engine/Plugins/Runtime/MeshModelingToolset/Source/MeshModelingTools/Public/DrawPolygonTool.h b/Engine/Plugins/Runtime/MeshModelingToolset/Source/MeshModelingTools/Public/DrawPolygonTool.h
--- a/Engine/Plugins/Runtime/MeshModelingToolset/Source/MeshModelingTools/Public/DrawPolygonTool.h
+++ b/Engine/Plugins/Runtime/MeshModelingToolset/Source/MeshModelingTools/Public/DrawPolygonTool.h
@@ -129,36 +129,37 @@ class MESHMODELINGTOOLS_API UDrawPolygonToolSnapProperties : public UInteractive
 	GENERATED_BODY()
 
 public:
-	/** If true, enables additional snapping controls. If false, all snapping is disabled. */
+	/** Enables additional snapping controls. If false, all snapping is disabled. */
 	UPROPERTY(EditAnywhere, NonTransactional, Category = Snapping)
 	bool bEnableSnapping = true;
 
-	/** If true, allows snapping to world grid according to editor settings. Enable Snapping needs to be true for this property to be available. */
-	UPROPERTY(EditAnywhere, NonTransactional, Category = Snapping, meta = (EditCondition = "bEnableSnapping"))
+	//~ Not user visible. Mirrors the snapping settings in the viewport and is used in EditConditions
+	UPROPERTY(meta = (TransientToolProperty))
 	bool bSnapToWorldGrid = false;
 
-	/** If true, snap to vertices in the current polygon or in other meshes. Enable Snapping needs to be true for this property to be available. */
+	/** Snap to vertices in other meshes. Requires Enable Snapping to be true. */
 	UPROPERTY(EditAnywhere, NonTransactional, Category = Snapping, meta = (EditCondition = "bEnableSnapping"))
 	bool bSnapToVertices = true;
 
-	/** If true, snap to edges in the current polygon or in other meshes. Enable Snapping needs to be true for this property to be available. */
+	/** Snap to edges in other meshes. Requires Enable Snapping to be true. */
 	UPROPERTY(EditAnywhere, NonTransactional, Category = Snapping, meta = (EditCondition = "bEnableSnapping"))
 	bool bSnapToEdges = false;
 
-	/** If true, snap to axes of the additional drawing grid as well as the axes relative to the last segment. Enable Snapping needs to be true for this property to be available. */
-	UPROPERTY(EditAnywhere, NonTransactional, Category = Snapping, meta = (EditCondition = "bEnableSnapping"))
+	/** Snap to axes of the drawing grid and axes relative to the last segment. Requires grid snapping to be disabled in viewport, and Enable Snapping to be true. */
+	UPROPERTY(EditAnywhere, NonTransactional, Category = Snapping, meta = (EditCondition = "bEnableSnapping && !bSnapToWorldGrid"))
 	bool bSnapToAxes = true;
 
-	/** If true, when snapping to axes, also try to snap to the length of an existing segment in the polygon. Enable Snapping and Snap to Axes needs to both be true for this property to be available. */
-	UPROPERTY(EditAnywhere, NonTransactional, Category = Snapping, meta = (EditCondition = "bEnableSnapping && bSnapToAxes"))
+	/** When snapping to axes, also try to snap to the length of an existing segment in the polygon. Requires grid snapping to be disabled in viewport, and Snap to Axes and Enable Snapping to be true. */
+	UPROPERTY(EditAnywhere, NonTransactional, Category = Snapping, meta = (EditCondition = "bEnableSnapping && !bSnapToWorldGrid && bSnapToAxes"))
 	bool bSnapToLengths = true;
 
-	/** If true, snap to surfaces of existing objects. Enable Snapping needs to be true for this property to be available. */
-	UPROPERTY(EditAnywhere, NonTransactional, Category = Snapping, meta = (EditCondition = "bEnableSnapping"))
+	/** Snap to surfaces of existing objects. Requires grid snapping to be disabled in viewport, and Enable Snapping to be true.  */
+	UPROPERTY(EditAnywhere, NonTransactional, Category = Snapping, meta = (EditCondition = "bEnableSnapping && !bSnapToWorldGrid"))
 	bool bSnapToSurfaces = false;
 
-	/** Offset for snap point on the surface of an existing object in the direction of the surface normal. Enable Snapping and Snap to Surfaces needs to be true for this property to be available. */
-	UPROPERTY(EditAnywhere, NonTransactional, Category = Snapping, meta = (DisplayName = "Surface Offset", EditCondition = "bSnapToSurfaces"))
+	/** Offset for snap point on the surface of an existing object in the direction of the surface normal. Requires grid snapping to be disabled in viewport, and Snap to Surfaces and Enable Snapping to be true. */
+	UPROPERTY(EditAnywhere, NonTransactional, Category = Snapping, meta = (DisplayName = "Surface Offset", 
+		EditCondition = "bEnableSnapping && !bSnapToWorldGrid && bSnapToSurfaces"))
 	float SnapToSurfacesOffset = 0.0f;
 };
 
