Modeling Tools: Provide support for the IMeshDescriptionProvider to provide "empty" MeshDescriptions, for cases where such instances are desired such as changing topology. This change creates extendable infrastructure parallel to the existing MeshDescription provider implementations so clients of the empty MeshDescription instances do not need localized specialization based on ToolTarget subclasses.

This change addresses the errors found with applying certain modeling tools to skeletal meshes due to the old code not using the appropriately setup MeshDescription for skeletal meshes.

#rnx

diff --git a/Engine/Plugins/Runtime/MeshModelingToolset/Source/ModelingComponents/Private/ModelingToolTargetUtil.cpp b/Engine/Plugins/Runtime/MeshModelingToolset/Source/ModelingComponents/Private/ModelingToolTargetUtil.cpp
--- a/Engine/Plugins/Runtime/MeshModelingToolset/Source/ModelingComponents/Private/ModelingToolTargetUtil.cpp
+++ b/Engine/Plugins/Runtime/MeshModelingToolset/Source/ModelingComponents/Private/ModelingToolTargetUtil.cpp
@@ -157,6 +157,18 @@ const FMeshDescription* UE::ToolTarget::GetMeshDescription(UToolTarget* Target,
 	return &EmptyMeshDescription;
 }
 
+FMeshDescription UE::ToolTarget::GetEmptyMeshDescription(UToolTarget* Target)
+{
+	static FMeshDescription EmptyMeshDescription;
+
+	IMeshDescriptionProvider* MeshDescriptionProvider = Cast<IMeshDescriptionProvider>(Target);
+	if (MeshDescriptionProvider)
+	{
+		return MeshDescriptionProvider->GetEmptyMeshDescription();
+	}
+	ensure(false);
+	return EmptyMeshDescription;
+}
 
 FMeshDescription UE::ToolTarget::GetMeshDescriptionCopy(UToolTarget* Target, const FGetMeshParameters& GetMeshParams)
 {
@@ -259,23 +271,24 @@ UE::ToolTarget::EDynamicMeshUpdateResult UE::ToolTarget::CommitMeshDescriptionUp
 		return EDynamicMeshUpdateResult::Failed;
 	}
 
+	FDynamicMeshToMeshDescription Converter;
 	FMeshDescription ConvertedMesh;
 	if (bHaveModifiedTopology)
 	{
-		FStaticMeshAttributes Attributes(ConvertedMesh);
-		Attributes.Register();
-		FDynamicMeshToMeshDescription Converter;
+		ConvertedMesh = UE::ToolTarget::GetEmptyMeshDescription(Target);		
 		Converter.Convert(&UpdatedMesh, ConvertedMesh);
 	}
 	else
 	{
-		// This should never need to be true because we are either (1) ignoring the tangents, in which
-		// case, we don't want to update them anyway (2) replacing them, in which case, we don't need
-		// to compute them
+		// FGetMeshParameters should never need to be initialized to anything other than the 
+		// default because we are either (1) ignoring the tangents, in which case, we don't
+		// want to update them anyway (2) replacing them, in which case, we don't need to
+		// compute them
+
 		ConvertedMesh = UE::ToolTarget::GetMeshDescriptionCopy(Target, FGetMeshParameters());
-		FDynamicMeshToMeshDescription Converter;
 		Converter.Update(&UpdatedMesh, ConvertedMesh);
 	}
+	
 
 	bool bOK = MeshDescriptionCommitter->CommitMeshDescription(MoveTemp(ConvertedMesh));
 	return (bOK) ? EDynamicMeshUpdateResult::Ok : EDynamicMeshUpdateResult::Failed;
diff --git a/Engine/Plugins/Runtime/MeshModelingToolset/Source/ModelingComponents/Public/ModelingToolTargetUtil.h b/Engine/Plugins/Runtime/MeshModelingToolset/Source/ModelingComponents/Public/ModelingToolTargetUtil.h
--- a/Engine/Plugins/Runtime/MeshModelingToolset/Source/ModelingComponents/Public/ModelingToolTargetUtil.h
+++ b/Engine/Plugins/Runtime/MeshModelingToolset/Source/ModelingComponents/Public/ModelingToolTargetUtil.h
@@ -95,6 +95,12 @@ MODELINGCOMPONENTS_API const FMeshDescription* GetMeshDescription(
 	UToolTarget* Target,
 	const FGetMeshParameters& GetMeshParams = FGetMeshParameters());
 
+/**
+ * @return an empty MeshDescription with attributes already registered appropriate to the target.
+ */
+MODELINGCOMPONENTS_API FMeshDescription GetEmptyMeshDescription(
+	UToolTarget* Target);
+
 /**
  * Return a copy of the MeshDescription underlying a ToolTarget
  * @return a new MeshDescription, which may be empty if the Target doesn't have a mesh  
diff --git a/Engine/Plugins/Runtime/MeshModelingToolset/Source/ModelingComponentsEditorOnly/ModelingComponentsEditorOnly.Build.cs b/Engine/Plugins/Runtime/MeshModelingToolset/Source/ModelingComponentsEditorOnly/ModelingComponentsEditorOnly.Build.cs
--- a/Engine/Plugins/Runtime/MeshModelingToolset/Source/ModelingComponentsEditorOnly/ModelingComponentsEditorOnly.Build.cs
+++ b/Engine/Plugins/Runtime/MeshModelingToolset/Source/ModelingComponentsEditorOnly/ModelingComponentsEditorOnly.Build.cs
@@ -29,6 +29,7 @@ public class ModelingComponentsEditorOnly : ModuleRules
                 "InteractiveToolsFramework",
                 "MeshDescription",
 				"StaticMeshDescription",
+				"SkeletalMeshDescription",
 				"GeometryCore",
 				"GeometryFramework",
 				"GeometryAlgorithms",
diff --git a/Engine/Plugins/Runtime/MeshModelingToolset/Source/ModelingComponentsEditorOnly/Private/ToolTargets/DynamicMeshComponentToolTarget.cpp b/Engine/Plugins/Runtime/MeshModelingToolset/Source/ModelingComponentsEditorOnly/Private/ToolTargets/DynamicMeshComponentToolTarget.cpp
--- a/Engine/Plugins/Runtime/MeshModelingToolset/Source/ModelingComponentsEditorOnly/Private/ToolTargets/DynamicMeshComponentToolTarget.cpp
+++ b/Engine/Plugins/Runtime/MeshModelingToolset/Source/ModelingComponentsEditorOnly/Private/ToolTargets/DynamicMeshComponentToolTarget.cpp
@@ -119,6 +119,16 @@ const FMeshDescription* UDynamicMeshComponentToolTarget::GetMeshDescription(cons
 	return nullptr;
 }
 
+FMeshDescription UDynamicMeshComponentToolTarget::GetEmptyMeshDescription()
+{
+	// We use StaticMeshAttributes here because they are the standard used across the engine
+	// with regard to setting up FMeshDescriptions in the majority of cases.
+
+	FMeshDescription EmptyMeshDescription;
+	FStaticMeshAttributes Attributes(EmptyMeshDescription);
+	Attributes.Register();
+	return EmptyMeshDescription;
+}
 
 void UDynamicMeshComponentToolTarget::CommitMeshDescription(const FCommitter& Committer, const FCommitMeshParameters& CommitMeshParams)
 {
diff --git a/Engine/Plugins/Runtime/MeshModelingToolset/Source/ModelingComponentsEditorOnly/Private/ToolTargets/SkeletalMeshComponentToolTarget.cpp b/Engine/Plugins/Runtime/MeshModelingToolset/Source/ModelingComponentsEditorOnly/Private/ToolTargets/SkeletalMeshComponentToolTarget.cpp
--- a/Engine/Plugins/Runtime/MeshModelingToolset/Source/ModelingComponentsEditorOnly/Private/ToolTargets/SkeletalMeshComponentToolTarget.cpp
+++ b/Engine/Plugins/Runtime/MeshModelingToolset/Source/ModelingComponentsEditorOnly/Private/ToolTargets/SkeletalMeshComponentToolTarget.cpp
@@ -8,6 +8,7 @@
 #include "MeshDescriptionToDynamicMesh.h"
 #include "Rendering/SkeletalMeshModel.h"
 #include "ToolTargets/SkeletalMeshToolTarget.h"
+#include "SkeletalMeshAttributes.h"
 
 using namespace UE::Geometry;
 
@@ -121,6 +122,14 @@ const FMeshDescription* USkeletalMeshComponentReadOnlyToolTarget::GetMeshDescrip
 	return CachedMeshDescription.Get();
 }
 
+FMeshDescription USkeletalMeshComponentReadOnlyToolTarget::GetEmptyMeshDescription()
+{
+	FMeshDescription EmptyMeshDescription;
+	FSkeletalMeshAttributes MeshAttributes(EmptyMeshDescription);
+	MeshAttributes.Register();
+
+	return EmptyMeshDescription;
+}
 
 FDynamicMesh3 USkeletalMeshComponentReadOnlyToolTarget::GetDynamicMesh()
 {
diff --git a/Engine/Plugins/Runtime/MeshModelingToolset/Source/ModelingComponentsEditorOnly/Private/ToolTargets/SkeletalMeshToolTarget.cpp b/Engine/Plugins/Runtime/MeshModelingToolset/Source/ModelingComponentsEditorOnly/Private/ToolTargets/SkeletalMeshToolTarget.cpp
--- a/Engine/Plugins/Runtime/MeshModelingToolset/Source/ModelingComponentsEditorOnly/Private/ToolTargets/SkeletalMeshToolTarget.cpp
+++ b/Engine/Plugins/Runtime/MeshModelingToolset/Source/ModelingComponentsEditorOnly/Private/ToolTargets/SkeletalMeshToolTarget.cpp
@@ -7,6 +7,7 @@
 #include "Rendering/SkeletalMeshModel.h"
 #include "Widgets/Notifications/SNotificationList.h"
 #include "Framework/Notifications/NotificationManager.h"
+#include "SkeletalMeshAttributes.h"
 
 using namespace UE::Geometry;
 
@@ -134,6 +135,16 @@ const FMeshDescription* USkeletalMeshReadOnlyToolTarget::GetMeshDescription(cons
 	return CachedMeshDescription.Get();
 }
 
+
+FMeshDescription USkeletalMeshReadOnlyToolTarget::GetEmptyMeshDescription()
+{
+	FMeshDescription EmptyMeshDescription;
+	FSkeletalMeshAttributes MeshAttributes(EmptyMeshDescription);
+	MeshAttributes.Register();
+
+	return EmptyMeshDescription;
+}
+
 void USkeletalMeshReadOnlyToolTarget::GetMeshDescription(const USkeletalMesh* SkeletalMeshIn, FMeshDescription& MeshDescription)
 {
 	using namespace USkeletalMeshToolTargetLocals;
diff --git a/Engine/Plugins/Runtime/MeshModelingToolset/Source/ModelingComponentsEditorOnly/Private/ToolTargets/StaticMeshComponentToolTarget.cpp b/Engine/Plugins/Runtime/MeshModelingToolset/Source/ModelingComponentsEditorOnly/Private/ToolTargets/StaticMeshComponentToolTarget.cpp
--- a/Engine/Plugins/Runtime/MeshModelingToolset/Source/ModelingComponentsEditorOnly/Private/ToolTargets/StaticMeshComponentToolTarget.cpp
+++ b/Engine/Plugins/Runtime/MeshModelingToolset/Source/ModelingComponentsEditorOnly/Private/ToolTargets/StaticMeshComponentToolTarget.cpp
@@ -160,6 +160,13 @@ const FMeshDescription* UStaticMeshComponentToolTarget::GetMeshDescription(const
 	return nullptr;
 }
 
+FMeshDescription UStaticMeshComponentToolTarget::GetEmptyMeshDescription()
+{
+	FMeshDescription EmptyMeshDescription;
+	FStaticMeshAttributes Attributes(EmptyMeshDescription);
+	Attributes.Register();
+	return EmptyMeshDescription;
+}
 
 FMeshDescription UStaticMeshComponentToolTarget::GetMeshDescriptionCopy(const FGetMeshParameters& GetMeshParams)
 {
diff --git a/Engine/Plugins/Runtime/MeshModelingToolset/Source/ModelingComponentsEditorOnly/Private/ToolTargets/StaticMeshToolTarget.cpp b/Engine/Plugins/Runtime/MeshModelingToolset/Source/ModelingComponentsEditorOnly/Private/ToolTargets/StaticMeshToolTarget.cpp
--- a/Engine/Plugins/Runtime/MeshModelingToolset/Source/ModelingComponentsEditorOnly/Private/ToolTargets/StaticMeshToolTarget.cpp
+++ b/Engine/Plugins/Runtime/MeshModelingToolset/Source/ModelingComponentsEditorOnly/Private/ToolTargets/StaticMeshToolTarget.cpp
@@ -212,6 +212,14 @@ const FMeshDescription* UStaticMeshToolTarget::GetMeshDescription(const FGetMesh
 	return nullptr;
 }
 
+FMeshDescription UStaticMeshToolTarget::GetEmptyMeshDescription()
+{
+	FMeshDescription EmptyMeshDescription;
+	FStaticMeshAttributes Attributes(EmptyMeshDescription);
+	Attributes.Register();
+	return EmptyMeshDescription;
+}
+
 void UStaticMeshToolTarget::CommitMeshDescription(const FCommitter& Committer, const FCommitMeshParameters& CommitParams)
 {
 	if (ensure(IsValid()) == false) return;
diff --git a/Engine/Plugins/Runtime/MeshModelingToolset/Source/ModelingComponentsEditorOnly/Private/ToolTargets/VolumeComponentToolTarget.cpp b/Engine/Plugins/Runtime/MeshModelingToolset/Source/ModelingComponentsEditorOnly/Private/ToolTargets/VolumeComponentToolTarget.cpp
--- a/Engine/Plugins/Runtime/MeshModelingToolset/Source/ModelingComponentsEditorOnly/Private/ToolTargets/VolumeComponentToolTarget.cpp
+++ b/Engine/Plugins/Runtime/MeshModelingToolset/Source/ModelingComponentsEditorOnly/Private/ToolTargets/VolumeComponentToolTarget.cpp
@@ -124,6 +124,17 @@ const FMeshDescription* UVolumeComponentToolTarget::GetMeshDescription(const FGe
 	return ConvertedMeshDescription.Get();
 }
 
+FMeshDescription UVolumeComponentToolTarget::GetEmptyMeshDescription()
+{
+	// We use StaticMeshAttributes here because they are the standard used across the engine
+	// with regard to setting up FMeshDescriptions in the majority of cases.
+
+	FMeshDescription EmptyMeshDescription;
+	FStaticMeshAttributes Attributes(EmptyMeshDescription);
+	Attributes.Register();
+	return EmptyMeshDescription;
+}
+
 void UVolumeComponentToolTarget::CommitMeshDescription(const FCommitter& Committer, const FCommitMeshParameters& CommitMeshParams)
 {
 	// no LODs so LODIdentifier is ignored
diff --git a/Engine/Plugins/Runtime/MeshModelingToolset/Source/ModelingComponentsEditorOnly/Public/ToolTargets/DynamicMeshComponentToolTarget.h b/Engine/Plugins/Runtime/MeshModelingToolset/Source/ModelingComponentsEditorOnly/Public/ToolTargets/DynamicMeshComponentToolTarget.h
--- a/Engine/Plugins/Runtime/MeshModelingToolset/Source/ModelingComponentsEditorOnly/Public/ToolTargets/DynamicMeshComponentToolTarget.h
+++ b/Engine/Plugins/Runtime/MeshModelingToolset/Source/ModelingComponentsEditorOnly/Public/ToolTargets/DynamicMeshComponentToolTarget.h
@@ -40,6 +40,7 @@ public:
 public:
 	// IMeshDescriptionProvider implementation
 	const FMeshDescription* GetMeshDescription(const FGetMeshParameters& GetMeshParams = FGetMeshParameters()) override;
+	virtual FMeshDescription GetEmptyMeshDescription() override;
 
 	// IMeshDescritpionCommitter implementation
 	virtual void CommitMeshDescription(const FCommitter& Committer, const FCommitMeshParameters& CommitMeshParams = FCommitMeshParameters()) override;
diff --git a/Engine/Plugins/Runtime/MeshModelingToolset/Source/ModelingComponentsEditorOnly/Public/ToolTargets/SkeletalMeshComponentToolTarget.h b/Engine/Plugins/Runtime/MeshModelingToolset/Source/ModelingComponentsEditorOnly/Public/ToolTargets/SkeletalMeshComponentToolTarget.h
--- a/Engine/Plugins/Runtime/MeshModelingToolset/Source/ModelingComponentsEditorOnly/Public/ToolTargets/SkeletalMeshComponentToolTarget.h
+++ b/Engine/Plugins/Runtime/MeshModelingToolset/Source/ModelingComponentsEditorOnly/Public/ToolTargets/SkeletalMeshComponentToolTarget.h
@@ -39,6 +39,7 @@ public:
 
 	// IMeshDescriptionProvider implementation
 	virtual const FMeshDescription* GetMeshDescription(const FGetMeshParameters& GetMeshParams = FGetMeshParameters()) override;
+	virtual FMeshDescription GetEmptyMeshDescription() override;
 
 	// IMaterialProvider implementation
 	virtual int32 GetNumMaterials() const override;
diff --git a/Engine/Plugins/Runtime/MeshModelingToolset/Source/ModelingComponentsEditorOnly/Public/ToolTargets/SkeletalMeshToolTarget.h b/Engine/Plugins/Runtime/MeshModelingToolset/Source/ModelingComponentsEditorOnly/Public/ToolTargets/SkeletalMeshToolTarget.h
--- a/Engine/Plugins/Runtime/MeshModelingToolset/Source/ModelingComponentsEditorOnly/Public/ToolTargets/SkeletalMeshToolTarget.h
+++ b/Engine/Plugins/Runtime/MeshModelingToolset/Source/ModelingComponentsEditorOnly/Public/ToolTargets/SkeletalMeshToolTarget.h
@@ -39,6 +39,7 @@ public:
 
 	// IMeshDescriptionProvider implementation
 	virtual const FMeshDescription* GetMeshDescription(const FGetMeshParameters& GetMeshParams = FGetMeshParameters()) override;
+	virtual FMeshDescription GetEmptyMeshDescription() override;
 
 	// IMaterialProvider implementation
 	virtual int32 GetNumMaterials() const override;
diff --git a/Engine/Plugins/Runtime/MeshModelingToolset/Source/ModelingComponentsEditorOnly/Public/ToolTargets/StaticMeshComponentToolTarget.h b/Engine/Plugins/Runtime/MeshModelingToolset/Source/ModelingComponentsEditorOnly/Public/ToolTargets/StaticMeshComponentToolTarget.h
--- a/Engine/Plugins/Runtime/MeshModelingToolset/Source/ModelingComponentsEditorOnly/Public/ToolTargets/StaticMeshComponentToolTarget.h
+++ b/Engine/Plugins/Runtime/MeshModelingToolset/Source/ModelingComponentsEditorOnly/Public/ToolTargets/StaticMeshComponentToolTarget.h
@@ -47,6 +47,7 @@ public:
 	// IMeshDescriptionProvider implementation
 	virtual const FMeshDescription* GetMeshDescription(const FGetMeshParameters& GetMeshParams = FGetMeshParameters()) override;
 	virtual FMeshDescription GetMeshDescriptionCopy(const FGetMeshParameters& GetMeshParams) override;
+	virtual FMeshDescription GetEmptyMeshDescription() override;
 
 	// IMeshDescritpionCommitter implementation
 	virtual void CommitMeshDescription(const FCommitter& Committer, const FCommitMeshParameters& CommitParams = FCommitMeshParameters()) override;
diff --git a/Engine/Plugins/Runtime/MeshModelingToolset/Source/ModelingComponentsEditorOnly/Public/ToolTargets/StaticMeshToolTarget.h b/Engine/Plugins/Runtime/MeshModelingToolset/Source/ModelingComponentsEditorOnly/Public/ToolTargets/StaticMeshToolTarget.h
--- a/Engine/Plugins/Runtime/MeshModelingToolset/Source/ModelingComponentsEditorOnly/Public/ToolTargets/StaticMeshToolTarget.h
+++ b/Engine/Plugins/Runtime/MeshModelingToolset/Source/ModelingComponentsEditorOnly/Public/ToolTargets/StaticMeshToolTarget.h
@@ -51,6 +51,7 @@ public:
 
 	// IMeshDescriptionProvider implementation
 	virtual const FMeshDescription* GetMeshDescription(const FGetMeshParameters& GetMeshParams = FGetMeshParameters()) override;
+	virtual FMeshDescription GetEmptyMeshDescription() override;
 
 	// IMeshDescritpionCommitter implementation
 	virtual void CommitMeshDescription(const FCommitter& Committer, const FCommitMeshParameters& CommitParams = FCommitMeshParameters()) override;
diff --git a/Engine/Plugins/Runtime/MeshModelingToolset/Source/ModelingComponentsEditorOnly/Public/ToolTargets/VolumeComponentToolTarget.h b/Engine/Plugins/Runtime/MeshModelingToolset/Source/ModelingComponentsEditorOnly/Public/ToolTargets/VolumeComponentToolTarget.h
--- a/Engine/Plugins/Runtime/MeshModelingToolset/Source/ModelingComponentsEditorOnly/Public/ToolTargets/VolumeComponentToolTarget.h
+++ b/Engine/Plugins/Runtime/MeshModelingToolset/Source/ModelingComponentsEditorOnly/Public/ToolTargets/VolumeComponentToolTarget.h
@@ -59,6 +59,7 @@ public:
 
 	// IMeshDescriptionProvider implementation
 	virtual const FMeshDescription* GetMeshDescription(const FGetMeshParameters& GetMeshParams = FGetMeshParameters()) override;
+	virtual FMeshDescription GetEmptyMeshDescription() override;
 
 	// IMeshDescritpionCommitter implementation
 	virtual void CommitMeshDescription(const FCommitter& Committer, const FCommitMeshParameters& CommitMeshParams = FCommitMeshParameters()) override;
diff --git a/Engine/Source/Runtime/InteractiveToolsFramework/Public/TargetInterfaces/MeshDescriptionProvider.h b/Engine/Source/Runtime/InteractiveToolsFramework/Public/TargetInterfaces/MeshDescriptionProvider.h
--- a/Engine/Source/Runtime/InteractiveToolsFramework/Public/TargetInterfaces/MeshDescriptionProvider.h
+++ b/Engine/Source/Runtime/InteractiveToolsFramework/Public/TargetInterfaces/MeshDescriptionProvider.h
@@ -27,6 +27,16 @@ public:
 	 */
 	virtual const FMeshDescription* GetMeshDescription(const FGetMeshParameters& GetMeshParams = FGetMeshParameters()) = 0;
 
+	/**
+	 * Returns an empty mesh description appropriate for the provider, i.e. configured with appropriate mesh
+	   attributes but otherwise devoid of topology or element data.
+
+	   Note: Some of our code expects at least FStaticMeshAttributes to be registered with the provided
+	   mesh description, and will break if some of the attributes that FStaticMeshAttributes uses are not present.
+	   The only reason we don't provide a default implementation here that is identical to the one used in
+	   StaticMeshToolTarget.cpp is to avoid a dependency on the static mesh description module in ITF.
+	 */
+	virtual FMeshDescription GetEmptyMeshDescription() = 0;
 
 	/**
 	 * Get a copy of the MeshDescription available through this Provider. 
