ModelingComponents: handle situation where StaticMesh SectionInfoMap is empty or partially populated in AssetUtils::GetStaticMeshLODAssetMaterials()
GeometryScript: explicitly clear SectionInfoMap in CopyMeshToStaticMesh() instead of clearing entries one at a time, which is wrong if only some sections are in the map
#rnx

diff --git a/Engine/Plugins/Experimental/GeometryScripting/Source/GeometryScriptingCore/Private/MeshAssetFunctions.cpp b/Engine/Plugins/Experimental/GeometryScripting/Source/GeometryScriptingCore/Private/MeshAssetFunctions.cpp
--- a/Engine/Plugins/Experimental/GeometryScripting/Source/GeometryScriptingCore/Private/MeshAssetFunctions.cpp
+++ b/Engine/Plugins/Experimental/GeometryScripting/Source/GeometryScriptingCore/Private/MeshAssetFunctions.cpp
@@ -231,11 +231,7 @@ UDynamicMesh*  UGeometryScriptLibrary_StaticMeshFunctions::CopyMeshToStaticMesh(
 		ToStaticMeshAsset->SetStaticMaterials(NewMaterials);
 
 		// Reset the section info map
-		int32 NumSections = ToStaticMeshAsset->GetSectionInfoMap().GetSectionNumber(UseLODIndex);
-		for (int32 SectionIndex = 0; SectionIndex < NumSections; ++SectionIndex)
-		{
-			ToStaticMeshAsset->GetSectionInfoMap().Remove(UseLODIndex, SectionIndex);
-		}
+		ToStaticMeshAsset->GetSectionInfoMap().Clear();
 	}
 
 	ToStaticMeshAsset->CommitMeshDescription(UseLODIndex);
diff --git a/Engine/Plugins/Experimental/GeometryScripting/Source/GeometryScriptingCore/Private/MeshMaterialFunctions.cpp b/Engine/Plugins/Experimental/GeometryScripting/Source/GeometryScriptingCore/Private/MeshMaterialFunctions.cpp
--- a/Engine/Plugins/Experimental/GeometryScripting/Source/GeometryScriptingCore/Private/MeshMaterialFunctions.cpp
+++ b/Engine/Plugins/Experimental/GeometryScripting/Source/GeometryScriptingCore/Private/MeshMaterialFunctions.cpp
@@ -415,58 +415,6 @@ UDynamicMesh* UGeometryScriptLibrary_MeshMaterialFunctions::CompactMaterialIDs(
 					SourceMaterialList[OldIndex] : nullptr;
 			}
 
-			//TArray<int32> KnownMaterialIDs;
-			//int32 MaxMaterialID = -1, MinMaterialID = TNumericLimits<int32>::Max();
-			//for (int32 TriangleID : EditMesh.TriangleIndicesItr())
-			//{
-			//	int32 MaterialID = MaterialIDs->GetValue(TriangleID);
-			//	MaxMaterialID = FMath::Max(MaterialID, MaxMaterialID);
-			//	MinMaterialID = FMath::Min(MaterialID, MinMaterialID);
-			//	KnownMaterialIDs.AddUnique(MaterialID);
-			//}
-			//if (MinMaterialID < 0)
-			//{
-			//	UE::Geometry::AppendError(Debug, EGeometryScriptErrorType::InvalidInputs, LOCTEXT("CompactMaterialIDs_InvalidMaterialID", "CompactMaterialIDs: Negative MaterialID found, unsafe to Compact"));
-			//	return;
-			//}
-			//if ( MaxMaterialID == (KnownMaterialIDs.Num()-1) )
-			//{
-			//	// set is already compact, nothing to do
-			//	return;
-			//}
-
-			//TArray<int32> MaterialIDMap;
-			//MaterialIDMap.Init(-1, MaxMaterialID+1);
-
-			//int32 NewMaterialIDCounter = 0;
-			//int32 NewMaxID = 0;
-			//for (int32 TriangleID : EditMesh.TriangleIndicesItr())
-			//{
-			//	int32 MaterialID = MaterialIDs->GetValue(TriangleID);
-			//	if (MaterialIDMap[MaterialID] == -1)
-			//	{
-			//		MaterialIDMap[MaterialID] = NewMaterialIDCounter;
-			//		NewMaxID = FMath::Max(NewMaxID, NewMaterialIDCounter);
-			//		NewMaterialIDCounter++;
-			//	}
-			//	int32 NewMaterialID = MaterialIDMap[MaterialID];
-			//	MaterialIDs->SetValue(TriangleID, NewMaterialID);
-			//}
-
-			//CompactedMaterialList.Reset();
-			//if ( SourceMaterialList.Num() == (MaxMaterialID+1) )
-			//{
-			//	CompactedMaterialList.SetNum(NewMaterialIDCounter);
-			//	for (int32 k = 0; k <= MaxMaterialID; ++k)
-			//	{
-			//		if (MaterialIDMap[k] >= 0)
-			//		{
-			//			CompactedMaterialList[MaterialIDMap[k]] = SourceMaterialList[k];
-			//		}
-			//		
-			//	}
-			//}
-
 		}, EDynamicMeshChangeType::GeneralEdit, EDynamicMeshAttributeChangeFlags::Unknown, false);
 	}
 	return TargetMesh;
diff --git a/Engine/Plugins/Runtime/MeshModelingToolset/Source/ModelingComponents/Private/AssetUtils/StaticMeshMaterialUtil.cpp b/Engine/Plugins/Runtime/MeshModelingToolset/Source/ModelingComponents/Private/AssetUtils/StaticMeshMaterialUtil.cpp
--- a/Engine/Plugins/Runtime/MeshModelingToolset/Source/ModelingComponents/Private/AssetUtils/StaticMeshMaterialUtil.cpp
+++ b/Engine/Plugins/Runtime/MeshModelingToolset/Source/ModelingComponents/Private/AssetUtils/StaticMeshMaterialUtil.cpp
@@ -26,6 +26,12 @@ bool UE::AssetUtils::GetStaticMeshLODAssetMaterials(
 		return false;
 	}
 
+	// Need to access the mesh here because # Sections == # PolygonGroups and this info doesn't seem to be anywhere else.
+	// Otherwise do not use the mesh in this function (would be nice to avoid this call if possible)
+	const FMeshDescription* SourceMesh = StaticMeshAsset->GetMeshDescription(LODIndex);
+	int32 NumSections = SourceMesh->PolygonGroups().Num();
+
+
 	TArray<FStaticMaterial> StaticMaterials = StaticMeshAsset->GetStaticMaterials();
 	MaterialInfoOut.MaterialSlots.Reset();
 	for (FStaticMaterial Mat : StaticMaterials)
@@ -33,25 +39,48 @@ bool UE::AssetUtils::GetStaticMeshLODAssetMaterials(
 		MaterialInfoOut.MaterialSlots.Add( FStaticMeshMaterialSlot{ Mat.MaterialInterface, Mat.MaterialSlotName } );
 	}
 
+	// This is complicated. A UStaticMesh has N MaterialSlots and each LOD has M Sections.
+	// Each Section can have any MaterialSlot assigned to it, ie it is not necessarily 1-1 or in-order.
+	// The SectionInfoMap is a TMap that will contain the SectionIndex-to-SlotIndex mapping
+	// *if* the mapping is not (SectionIndex == SlotIndex), or has ever been edited.
+	// So if the SectionIndex is not found in the SectionInfoMap, then it should be used as the SlotIndex directly.
+
 	const FMeshSectionInfoMap& SectionInfoMap = StaticMeshAsset->GetSectionInfoMap();
 	MaterialInfoOut.LODIndex = LODIndex;
-	MaterialInfoOut.NumSections = SectionInfoMap.GetSectionNumber(LODIndex);
+	MaterialInfoOut.NumSections = NumSections;
 
 	MaterialInfoOut.SectionSlotIndexes.SetNum(MaterialInfoOut.NumSections);
 	MaterialInfoOut.SectionMaterials.SetNum(MaterialInfoOut.NumSections);
-	for (int32 k = 0; k < MaterialInfoOut.NumSections; ++k)
+	for (int32 SectionIndex = 0; SectionIndex < MaterialInfoOut.NumSections; ++SectionIndex)
 	{
-		FMeshSectionInfo SectionInfo = SectionInfoMap.Get(LODIndex, k);
-		if (SectionInfo.MaterialIndex >= 0 && SectionInfo.MaterialIndex < MaterialInfoOut.MaterialSlots.Num())
+		MaterialInfoOut.SectionSlotIndexes[SectionIndex] = -1;
+		MaterialInfoOut.SectionMaterials[SectionIndex] = nullptr;
+
+		if (SectionInfoMap.IsValidSection(LODIndex, SectionIndex))
 		{
-			MaterialInfoOut.SectionSlotIndexes[k] = SectionInfo.MaterialIndex;
-			MaterialInfoOut.SectionMaterials[k] = MaterialInfoOut.MaterialSlots[SectionInfo.MaterialIndex].Material;
+			// did not find this section 
+			if ( StaticMaterials.IsValidIndex(SectionIndex) )
+			{
+				MaterialInfoOut.SectionSlotIndexes[SectionIndex] = SectionIndex;
+				MaterialInfoOut.SectionMaterials[SectionIndex] = MaterialInfoOut.MaterialSlots[SectionIndex].Material;
+			}
+			else
+			{
+				ensure(false);		// material list is broken? use default material.
+			}
 		}
 		else
 		{
-			ensure(false);		// this is *not* supposed to be able to happen! SectionMap is broken...
-			MaterialInfoOut.SectionSlotIndexes[k] = -1;
-			MaterialInfoOut.SectionMaterials[k] = nullptr;
+			FMeshSectionInfo SectionInfo = SectionInfoMap.Get(LODIndex, SectionIndex);
+			if ( StaticMaterials.IsValidIndex(SectionInfo.MaterialIndex) )
+			{
+				MaterialInfoOut.SectionSlotIndexes[SectionIndex] = SectionInfo.MaterialIndex;
+				MaterialInfoOut.SectionMaterials[SectionIndex] = MaterialInfoOut.MaterialSlots[SectionInfo.MaterialIndex].Material;
+			}
+			else
+			{
+				ensure(false);		// this is *not* supposed to be able to happen! SectionMap is broken...
+			}
 		}
 	}
 
