MeshModelingToolset: Log warning about UVAtlas not available under Linux only once

diff --git a/Engine/Plugins/Runtime/MeshModelingToolset/Source/ModelingOperatorsEditorOnly/Private/ParameterizationOps/ParameterizeMeshOp.cpp b/Engine/Plugins/Runtime/MeshModelingToolset/Source/ModelingOperatorsEditorOnly/Private/ParameterizationOps/ParameterizeMeshOp.cpp
--- a/Engine/Plugins/Runtime/MeshModelingToolset/Source/ModelingOperatorsEditorOnly/Private/ParameterizationOps/ParameterizeMeshOp.cpp
+++ b/Engine/Plugins/Runtime/MeshModelingToolset/Source/ModelingOperatorsEditorOnly/Private/ParameterizationOps/ParameterizeMeshOp.cpp
@@ -19,7 +19,7 @@
 #endif	// WITH_PROXYLOD
 
 // We do not have UVAtlas available if ProxyLOD is not available.
-#define NO_UVATLAS WITH_PROXYLOD ? 0 : 1
+#define NO_UVATLAS (WITH_PROXYLOD ? 0 : 1)
 
 using namespace UE::Geometry;
 
@@ -426,7 +426,11 @@ void FParameterizeMeshOp::CalculateResult(FProgressCancel* Progress)
 
 #if NO_UVATLAS
 	constexpr bool bUseXAtlas = true;
-	UE_LOG(LogGeometry, Warning, TEXT("AutoUV method UVAtlas not available; falling back to XAtlas instead."));
+	auto LogUVAtlasWarning = []
+	{
+		UE_LOG(LogGeometry, Warning, TEXT("AutoUV method UVAtlas not available; falling back to XAtlas instead."));
+	};
+	UE_CALL_ONCE(LogUVAtlasWarning);
 #else	// NO_UVATLAS
 	const bool bUseXAtlas = (Method == EParamOpBackend::XAtlas);
 #endif	// NO_UVATLAS
