MeshModelingToolset: Minor change to tooltip in PolyGroup layer properties

diff --git a/Engine/Plugins/Runtime/MeshModelingToolset/Source/ModelingComponents/Public/PropertySets/PolygroupLayersProperties.h b/Engine/Plugins/Runtime/MeshModelingToolset/Source/ModelingComponents/Public/PropertySets/PolygroupLayersProperties.h
--- a/Engine/Plugins/Runtime/MeshModelingToolset/Source/ModelingComponents/Public/PropertySets/PolygroupLayersProperties.h
+++ b/Engine/Plugins/Runtime/MeshModelingToolset/Source/ModelingComponents/Public/PropertySets/PolygroupLayersProperties.h
@@ -19,7 +19,7 @@ class MODELINGCOMPONENTS_API UPolygroupLayersProperties : public UInteractiveToo
 	GENERATED_BODY()
 public:
 
-	/** Select polygroup layer. */
+	/** Select PolyGroup layer to use. */
 	UPROPERTY(EditAnywhere, Category = "PolyGroup Layer", meta = (DisplayName = "Active PolyGroup", GetOptions = GetGroupLayersFunc))
 	FName ActiveGroupLayer = "Default";
 
