MeshModelingTools: Make CollectSurfacePathMechanic remove preview vertex if there isn't currently a valid one.

#rnx

diff --git a/Engine/Plugins/Runtime/MeshModelingToolset/Source/ModelingComponents/Private/Mechanics/CollectSurfacePathMechanic.cpp b/Engine/Plugins/Runtime/MeshModelingToolset/Source/ModelingComponents/Private/Mechanics/CollectSurfacePathMechanic.cpp
--- a/Engine/Plugins/Runtime/MeshModelingToolset/Source/ModelingComponents/Private/Mechanics/CollectSurfacePathMechanic.cpp
+++ b/Engine/Plugins/Runtime/MeshModelingToolset/Source/ModelingComponents/Private/Mechanics/CollectSurfacePathMechanic.cpp
@@ -54,9 +54,12 @@ void UCollectSurfacePathMechanic::Render(IToolsContextRenderAPI* RenderAPI)
 
 			if (!bGeometricCloseOccurred)
 			{
-				// This draws the line to the current hover point as a preview
-				const FLinearColor& DrawPreviewColor = (bCurrentPreviewWillComplete || bGeometricCloseOccurred) ? PathCompleteColor : PreviewColor;
-				PathDrawer.DrawLine(HitPath[NumPoints].Origin, PreviewPathPoint.Origin, DrawPreviewColor);
+				if (bPreviewPathPointValid)
+				{
+					// This draws the line to the current hover point as a preview
+					const FLinearColor& DrawPreviewColor = (bCurrentPreviewWillComplete || bGeometricCloseOccurred) ? PathCompleteColor : PreviewColor;
+					PathDrawer.DrawLine(HitPath[NumPoints].Origin, PreviewPathPoint.Origin, DrawPreviewColor);
+				}
 			}
 			else if (LoopWasClosed())
 			{
@@ -65,7 +68,7 @@ void UCollectSurfacePathMechanic::Render(IToolsContextRenderAPI* RenderAPI)
 			}
 		}
 
-		if (!bGeometricCloseOccurred)
+		if (!bGeometricCloseOccurred && bPreviewPathPointValid)
 		{
 			PathDrawer.DrawPoint(PreviewPathPoint.Origin, PreviewColor, PathDrawer.PointSize, PathDrawer.bDepthTested);
 		}
@@ -121,7 +124,8 @@ bool UCollectSurfacePathMechanic::IsHitByRay(const FRay3d& Ray, FFrame3d& HitPoi
 bool UCollectSurfacePathMechanic::UpdatePreviewPoint(const FRay3d& Ray)
 {
 	FFrame3d PreviewPoint;
-	if ( RayToPathPoint(Ray, PreviewPoint, true ) == false)
+	bPreviewPathPointValid = RayToPathPoint(Ray, PreviewPoint, true);
+	if (!bPreviewPathPointValid)
 	{
 		return false;
 	}
diff --git a/Engine/Plugins/Runtime/MeshModelingToolset/Source/ModelingComponents/Public/Mechanics/CollectSurfacePathMechanic.h b/Engine/Plugins/Runtime/MeshModelingToolset/Source/ModelingComponents/Public/Mechanics/CollectSurfacePathMechanic.h
--- a/Engine/Plugins/Runtime/MeshModelingToolset/Source/ModelingComponents/Public/Mechanics/CollectSurfacePathMechanic.h
+++ b/Engine/Plugins/Runtime/MeshModelingToolset/Source/ModelingComponents/Public/Mechanics/CollectSurfacePathMechanic.h
@@ -44,7 +44,7 @@ public:
 	TArray<FFrame3d> HitPath;
 
 	FFrame3d PreviewPathPoint;
-
+	bool bPreviewPathPointValid = false;
 
 	FToolDataVisualizer PathDrawer;
 	FLinearColor PathColor;
