MeshModelingTools: Make first point in PolyEditCutFacesActivity undoable.

#rnx

diff --git a/Engine/Plugins/Runtime/MeshModelingToolset/Source/MeshModelingTools/Private/ToolActivities/PolyEditCutFacesActivity.cpp b/Engine/Plugins/Runtime/MeshModelingToolset/Source/MeshModelingTools/Private/ToolActivities/PolyEditCutFacesActivity.cpp
--- a/Engine/Plugins/Runtime/MeshModelingToolset/Source/MeshModelingTools/Private/ToolActivities/PolyEditCutFacesActivity.cpp
+++ b/Engine/Plugins/Runtime/MeshModelingToolset/Source/MeshModelingTools/Private/ToolActivities/PolyEditCutFacesActivity.cpp
@@ -233,8 +233,14 @@ void UPolyEditCutFacesActivity::OnClicked(const FInputDeviceRay& ClickPos)
 	{
 		if (SurfacePathMechanic->IsDone())
 		{
+			++ActivityStamp;
 			ApplyCutFaces();
 		}
+		else
+		{
+			ParentTool->GetToolManager()->EmitObjectChange(this, MakeUnique<FPolyEditCutFacesActivityFirstPointChange>(ActivityStamp),
+				LOCTEXT("CutLineStart", "Cut Line Start"));
+		}
 	}
 }
 
@@ -251,4 +257,14 @@ bool UPolyEditCutFacesActivity::OnUpdateHover(const FInputDeviceRay& DevicePos)
 	return bIsRunning;
 }
 
+void FPolyEditCutFacesActivityFirstPointChange::Revert(UObject* Object)
+{
+	UPolyEditCutFacesActivity* Activity = Cast<UPolyEditCutFacesActivity>(Object);
+	if (ensure(Activity->SurfacePathMechanic))
+	{
+		Activity->SurfacePathMechanic->PopLastPoint();
+	}
+	bHaveDoneUndo = true;
+}
+
 #undef LOCTEXT_NAMESPACE
diff --git a/Engine/Plugins/Runtime/MeshModelingToolset/Source/MeshModelingTools/Public/ToolActivities/PolyEditCutFacesActivity.h b/Engine/Plugins/Runtime/MeshModelingToolset/Source/MeshModelingTools/Public/ToolActivities/PolyEditCutFacesActivity.h
--- a/Engine/Plugins/Runtime/MeshModelingToolset/Source/MeshModelingTools/Public/ToolActivities/PolyEditCutFacesActivity.h
+++ b/Engine/Plugins/Runtime/MeshModelingToolset/Source/MeshModelingTools/Public/ToolActivities/PolyEditCutFacesActivity.h
@@ -6,6 +6,7 @@
 
 #include "BaseBehaviors/BehaviorTargetInterfaces.h"
 #include "InteractiveToolActivity.h"
+#include "InteractiveToolChange.h" //FToolCommandChange
 #include "ToolContextInterfaces.h" // FViewCameraState
 
 #include "PolyEditCutFacesActivity.generated.h"
@@ -84,6 +85,35 @@ protected:
 	TObjectPtr<UPolyEditActivityContext> ActivityContext;
 
 	bool bIsRunning = false;
+	int32 ActivityStamp = 1;
 
 	FViewCameraState CameraState;
+
+	friend class FPolyEditCutFacesActivityFirstPointChange;
 };
+
+/**
+ * This should get emitted when setting the first point so that we can undo it.
+ */
+class MESHMODELINGTOOLS_API FPolyEditCutFacesActivityFirstPointChange : public FToolCommandChange
+{
+public:
+	FPolyEditCutFacesActivityFirstPointChange(int32 CurrentActivityStamp)
+		: ActivityStamp(CurrentActivityStamp)
+	{};
+
+	virtual void Apply(UObject* Object) override {};
+	virtual void Revert(UObject* Object) override;
+	virtual bool HasExpired(UObject* Object) const override
+	{
+		return bHaveDoneUndo || Cast<UPolyEditCutFacesActivity>(Object)->ActivityStamp != ActivityStamp;
+	}
+	virtual FString ToString() const override
+	{
+		return TEXT("FPolyEditCutFacesActivityFirstPointChange");
+	}
+
+protected:
+	int32 ActivityStamp;
+	bool bHaveDoneUndo = false;
+};
\ No newline at end of file
