ModelingTools: Fix StaticMesh component and asset material set discrepancy when using MshMrg tool on multiple meshes with different materials and writing to the last input asset.

#rnx

diff --git a/Engine/Plugins/Runtime/MeshModelingToolset/Source/MeshModelingTools/Private/CombineMeshesTool.cpp b/Engine/Plugins/Runtime/MeshModelingToolset/Source/MeshModelingTools/Private/CombineMeshesTool.cpp
--- a/Engine/Plugins/Runtime/MeshModelingToolset/Source/MeshModelingTools/Private/CombineMeshesTool.cpp
+++ b/Engine/Plugins/Runtime/MeshModelingToolset/Source/MeshModelingTools/Private/CombineMeshesTool.cpp
@@ -425,6 +425,10 @@ void UCombineMeshesTool::UpdateExistingAsset()
 		NewMaterialSet.Materials = AllMaterials;
 		UE::ToolTarget::CommitDynamicMeshUpdate(Targets[SkipIndex], AccumulateDMesh, true, FConversionToMeshDescriptionOptions(), &NewMaterialSet);
 
+		// CommitDynamicMeshUpdate updates the materials for the underlying asset. However,
+		// it does not update the component itself, so address that now.
+		UE::ToolTarget::CommitMaterialSetUpdate(Targets[SkipIndex], NewMaterialSet, false);
+
 		if (bOutputComponentSupportsCollision)
 		{
 			UE::Geometry::SetSimpleCollision(UpdateComponent, &SimpleCollision, CollisionSettings);
