MeshModelingTools: Add missing operations to TriEdit.

#rnx

diff --git a/Engine/Plugins/Runtime/MeshModelingToolset/Source/MeshModelingTools/Public/EditMeshPolygonsTool.h b/Engine/Plugins/Runtime/MeshModelingToolset/Source/MeshModelingTools/Public/EditMeshPolygonsTool.h
--- a/Engine/Plugins/Runtime/MeshModelingToolset/Source/MeshModelingTools/Public/EditMeshPolygonsTool.h
+++ b/Engine/Plugins/Runtime/MeshModelingToolset/Source/MeshModelingTools/Public/EditMeshPolygonsTool.h
@@ -268,12 +268,20 @@ public:
 	UFUNCTION(CallInEditor, Category = TriangleEdits, meta = (DisplayName = "Extrude", DisplayPriority = 1))
 	void Extrude() { PostAction(EEditMeshPolygonsToolActions::Extrude); }
 
+	/** Like Extrude/Offset, but performed in a boolean way, meaning that the faces can cut away the mesh or bridge mesh parts. */
+	UFUNCTION(CallInEditor, Category = TriangleEdits, meta = (DisplayName = "Push/Pull", DisplayPriority = 1))
+	void PushPull() { PostAction(EEditMeshPolygonsToolActions::PushPull); }
+
+	/** Like Extrude, but defaults to moving verts along vertex normals instead of a single direction. */
+	UFUNCTION(CallInEditor, Category = TriangleEdits, meta = (DisplayName = "Offset", DisplayPriority = 1))
+	void Offset() { PostAction(EEditMeshPolygonsToolActions::Offset); }
+
 	/** Inset the current set of selected faces. Click in viewport to confirm inset distance. */
-	UFUNCTION(CallInEditor, Category = FaceEdits, meta = (DisplayName = "Inset", DisplayPriority = 2))
+	UFUNCTION(CallInEditor, Category = TriangleEdits, meta = (DisplayName = "Inset", DisplayPriority = 2))
 	void Inset() { PostAction(EEditMeshPolygonsToolActions::Inset);	}
 
 	/** Outset the current set of selected faces. Click in viewport to confirm outset distance. */
-	UFUNCTION(CallInEditor, Category = FaceEdits, meta = (DisplayName = "Outset", DisplayPriority = 3))
+	UFUNCTION(CallInEditor, Category = TriangleEdits, meta = (DisplayName = "Outset", DisplayPriority = 3))
 	void Outset() { PostAction(EEditMeshPolygonsToolActions::Outset);	}
 
 	/** Delete the current set of selected faces */
