MeshModelingTools: Fix height mechanic not falling back to line distance.

#rnx

diff --git a/Engine/Plugins/Runtime/MeshModelingToolset/Source/ModelingComponents/Private/Mechanics/PlaneDistanceFromHitMechanic.cpp b/Engine/Plugins/Runtime/MeshModelingToolset/Source/ModelingComponents/Private/Mechanics/PlaneDistanceFromHitMechanic.cpp
--- a/Engine/Plugins/Runtime/MeshModelingToolset/Source/ModelingComponents/Private/Mechanics/PlaneDistanceFromHitMechanic.cpp
+++ b/Engine/Plugins/Runtime/MeshModelingToolset/Source/ModelingComponents/Private/Mechanics/PlaneDistanceFromHitMechanic.cpp
@@ -95,6 +95,7 @@ void UPlaneDistanceFromHitMechanic::UpdateCurrentDistance(const FRay& WorldRay)
 	if (bFoundHit == false && bFallbackToLineAxisPoint)
 	{
 		FDistLine3Ray3d Distance(FLine3d(PreviewHeightFrame.Origin, PreviewHeightFrame.Z()), FRay3d(WorldRay));
+		Distance.ComputeResult();
 		float WorldHitDist = WorldRay.GetParameter( (FVector)Distance.RayClosestPoint );
 		NearestHitFrameWorld = FFrame3d(Distance.RayClosestPoint, UE::Geometry::Normalized(Distance.RayClosestPoint - Distance.LineClosestPoint) );
 		FVector3d HitPosLocal = PreviewHeightFrame.ToFramePoint(Distance.RayClosestPoint);
