MeshModelingTools: Make the duplicate, combine, transfer, and convert tools require explicit accept/cancel to avoid accidental asset creation.

#rnx

diff --git a/Engine/Plugins/Experimental/MeshModelingToolsetExp/Source/MeshModelingToolsExp/Public/ConvertMeshesTool.h b/Engine/Plugins/Experimental/MeshModelingToolsetExp/Source/MeshModelingToolsExp/Public/ConvertMeshesTool.h
--- a/Engine/Plugins/Experimental/MeshModelingToolsetExp/Source/MeshModelingToolsExp/Public/ConvertMeshesTool.h
+++ b/Engine/Plugins/Experimental/MeshModelingToolsetExp/Source/MeshModelingToolsExp/Public/ConvertMeshesTool.h
@@ -5,6 +5,7 @@
 #include "CoreMinimal.h"
 #include "BaseTools/MultiSelectionMeshEditingTool.h"
 #include "InteractiveToolBuilder.h"
+#include "InteractiveToolQueryInterfaces.h" // IInteractiveToolExclusiveToolAPI
 #include "PropertySets/CreateMeshObjectTypeProperties.h"
 #include "ConvertMeshesTool.generated.h"
 
@@ -37,7 +38,9 @@ public:
 
 
 UCLASS()
-class MESHMODELINGTOOLSEXP_API UConvertMeshesTool : public UMultiSelectionMeshEditingTool
+class MESHMODELINGTOOLSEXP_API UConvertMeshesTool : public UMultiSelectionMeshEditingTool,
+	// Disallow auto-accept switch-away for the tool
+	public IInteractiveToolExclusiveToolAPI
 {
 	GENERATED_BODY()
 
diff --git a/Engine/Plugins/Experimental/MeshModelingToolsetExp/Source/MeshModelingToolsExp/Public/TransferMeshTool.h b/Engine/Plugins/Experimental/MeshModelingToolsetExp/Source/MeshModelingToolsExp/Public/TransferMeshTool.h
--- a/Engine/Plugins/Experimental/MeshModelingToolsetExp/Source/MeshModelingToolsExp/Public/TransferMeshTool.h
+++ b/Engine/Plugins/Experimental/MeshModelingToolsetExp/Source/MeshModelingToolsExp/Public/TransferMeshTool.h
@@ -5,6 +5,7 @@
 #include "CoreMinimal.h"
 #include "BaseTools/MultiSelectionMeshEditingTool.h"
 #include "InteractiveToolBuilder.h"
+#include "InteractiveToolQueryInterfaces.h" // IInteractiveToolExclusiveToolAPI
 #include "TargetInterfaces/MeshTargetInterfaceTypes.h"
 #include "TransferMeshTool.generated.h"
 
@@ -72,7 +73,9 @@ public:
 
 
 UCLASS()
-class MESHMODELINGTOOLSEXP_API UTransferMeshTool : public UMultiSelectionMeshEditingTool
+class MESHMODELINGTOOLSEXP_API UTransferMeshTool : public UMultiSelectionMeshEditingTool,
+	// Disallow auto-accept switch-away for the tool
+	public IInteractiveToolExclusiveToolAPI
 {
 	GENERATED_BODY()
 
diff --git a/Engine/Plugins/Runtime/MeshModelingToolset/Source/MeshModelingTools/Public/CombineMeshesTool.h b/Engine/Plugins/Runtime/MeshModelingToolset/Source/MeshModelingTools/Public/CombineMeshesTool.h
--- a/Engine/Plugins/Runtime/MeshModelingToolset/Source/MeshModelingTools/Public/CombineMeshesTool.h
+++ b/Engine/Plugins/Runtime/MeshModelingToolset/Source/MeshModelingTools/Public/CombineMeshesTool.h
@@ -5,6 +5,7 @@
 #include "CoreMinimal.h"
 #include "BaseTools/MultiSelectionMeshEditingTool.h"
 #include "InteractiveToolBuilder.h"
+#include "InteractiveToolQueryInterfaces.h" // IInteractiveToolExclusiveToolAPI
 #include "MeshOpPreviewHelpers.h"
 #include "DynamicMesh/DynamicMesh3.h"
 #include "BaseTools/BaseCreateFromSelectedTool.h"
@@ -64,7 +65,10 @@ public:
  * Simple tool to combine multiple meshes into a single mesh asset
  */
 UCLASS()
-class MESHMODELINGTOOLS_API UCombineMeshesTool : public UMultiSelectionMeshEditingTool
+class MESHMODELINGTOOLS_API UCombineMeshesTool : public UMultiSelectionMeshEditingTool, 
+	// Disallow auto-accept switch-away because it's easy to accidentally make an extra asset in duplicate mode,
+	// and it's not great in combine mode either.
+	public IInteractiveToolExclusiveToolAPI
 {
 	GENERATED_BODY()
 
