ModelingTools: Fix handling of UV Unwrap tool to correctly take into account user selection for PolyGroup layer. Additionally fixes issue for the UV Unwrap tool in the UV Editor as well as cleaning up tooltip language about the option.

#rnx

diff --git a/Engine/Plugins/Experimental/UVEditor/Source/UVEditorTools/Private/UVEditorRecomputeUVsTool.cpp b/Engine/Plugins/Experimental/UVEditor/Source/UVEditorTools/Private/UVEditorRecomputeUVsTool.cpp
--- a/Engine/Plugins/Experimental/UVEditor/Source/UVEditorTools/Private/UVEditorRecomputeUVsTool.cpp
+++ b/Engine/Plugins/Experimental/UVEditor/Source/UVEditorTools/Private/UVEditorRecomputeUVsTool.cpp
@@ -52,6 +52,7 @@ void UUVEditorRecomputeUVsTool::Setup()
 	Settings = NewObject<URecomputeUVsToolProperties>(this);
 	Settings->RestoreProperties(this);
 	AddToolPropertySource(Settings);
+	Factories.SetNum(Targets.Num());
 
 	if (Targets.Num() == 1)
 	{
@@ -62,8 +63,13 @@ void UUVEditorRecomputeUVsTool::Setup()
 		AddToolPropertySource(PolygroupLayerProperties);
 		UpdateActiveGroupLayer();
 	}
+	else
+	{
+		Settings->bEnablePolygroupSupport = false;
+		Settings->IslandGeneration = ERecomputeUVsPropertiesIslandMode::ExistingUVs;
+	}
 
-	Factories.SetNum(Targets.Num());
+	
 	for (int32 TargetIndex = 0; TargetIndex < Targets.Num(); ++TargetIndex)
 	{
 		TObjectPtr<UUVEditorToolMeshInput> Target = Targets[TargetIndex];
@@ -211,6 +217,12 @@ void UUVEditorRecomputeUVsTool::UpdateActiveGroupLayer()
 			ActiveGroupSet = MakeShared<UE::Geometry::FPolygroupSet, ESPMode::ThreadSafe>(Targets[0]->AppliedCanonical.Get(), FoundAttrib);
 		}
 	}
+	for (int32 TargetIndex = 0; TargetIndex < Targets.Num(); ++TargetIndex)
+	{
+		if (Factories[TargetIndex]) {
+			Factories[TargetIndex]->InputGroups = ActiveGroupSet;
+		}
+	}
 }
 
 
diff --git a/Engine/Plugins/Runtime/MeshModelingToolset/Source/MeshModelingTools/Private/RecomputeUVsTool.cpp b/Engine/Plugins/Runtime/MeshModelingToolset/Source/MeshModelingTools/Private/RecomputeUVsTool.cpp
--- a/Engine/Plugins/Runtime/MeshModelingToolset/Source/MeshModelingTools/Private/RecomputeUVsTool.cpp
+++ b/Engine/Plugins/Runtime/MeshModelingToolset/Source/MeshModelingTools/Private/RecomputeUVsTool.cpp
@@ -75,6 +75,7 @@ void URecomputeUVsTool::Setup()
 
 	RecomputeUVsOpFactory = NewObject<URecomputeUVsOpFactory>();
 	RecomputeUVsOpFactory->OriginalMesh = InputMesh;
+	RecomputeUVsOpFactory->InputGroups = ActiveGroupSet;
 	RecomputeUVsOpFactory->Settings = Settings;
 	RecomputeUVsOpFactory->TargetTransform = UE::ToolTarget::GetLocalToWorldTransform(Target);
 	RecomputeUVsOpFactory->GetSelectedUVChannel = [this]() { return GetSelectedUVChannel(); };
@@ -217,6 +218,10 @@ void URecomputeUVsTool::UpdateActiveGroupLayer()
 		ensureMsgf(FoundAttrib, TEXT("Selected Attribute Not Found! Falling back to Default group layer."));
 		ActiveGroupSet = MakeShared<UE::Geometry::FPolygroupSet, ESPMode::ThreadSafe>(InputMesh.Get(), FoundAttrib);
 	}
+	if (RecomputeUVsOpFactory)
+	{
+		RecomputeUVsOpFactory->InputGroups = ActiveGroupSet;
+	}
 }
 
 
diff --git a/Engine/Plugins/Runtime/MeshModelingToolset/Source/ModelingOperators/Public/Properties/RecomputeUVsProperties.h b/Engine/Plugins/Runtime/MeshModelingToolset/Source/ModelingOperators/Public/Properties/RecomputeUVsProperties.h
--- a/Engine/Plugins/Runtime/MeshModelingToolset/Source/ModelingOperators/Public/Properties/RecomputeUVsProperties.h
+++ b/Engine/Plugins/Runtime/MeshModelingToolset/Source/ModelingOperators/Public/Properties/RecomputeUVsProperties.h
@@ -28,7 +28,7 @@ enum class ERecomputeUVsPropertiesIslandMode
 {
 	// Values must match UE::Geometry::ERecomputeUVsIslandMode
 
-	/** Use Polygroups defined by Active PolyGroup property to define initial UV islands */
+	/** Use Polygroups defined by Active PolyGroup property to define initial UV islands. */
 	PolyGroups = 0 UMETA(DisplayName = "PolyGroups"),
 	/** Use existing UV Layer to define UV islands, i.e. re-solve UV flattening based on existing UVs */
 	ExistingUVs = 1
@@ -65,8 +65,11 @@ class MODELINGOPERATORS_API URecomputeUVsToolProperties : public UInteractiveToo
 	GENERATED_BODY()
 public:
 
-	/** Generation method for initial UV islands */
-	UPROPERTY(EditAnywhere, Category = "UV Unwrap")
+	UPROPERTY(meta = (TransientToolProperty))
+	bool bEnablePolygroupSupport = true;
+
+	/** Generation method for initial UV islands.*/
+	UPROPERTY(EditAnywhere, Category = "UV Unwrap", meta = (EditCondition = "bEnablePolygroupSupport", EditConditionHides, HideEditConditionToggle = true))
 	ERecomputeUVsPropertiesIslandMode IslandGeneration = ERecomputeUVsPropertiesIslandMode::ExistingUVs;
 
 	/** Type of UV flattening algorithm to use */
