MeshModelingTools: Add marquee face selection in PolyEdit. Fix minor bug where shift+clicking away from the mesh would clear selection.

#rnx

diff --git a/Engine/Plugins/Runtime/MeshModelingToolset/Source/MeshModelingTools/Private/EditMeshPolygonsTool.cpp b/Engine/Plugins/Runtime/MeshModelingToolset/Source/MeshModelingTools/Private/EditMeshPolygonsTool.cpp
--- a/Engine/Plugins/Runtime/MeshModelingToolset/Source/MeshModelingTools/Private/EditMeshPolygonsTool.cpp
+++ b/Engine/Plugins/Runtime/MeshModelingToolset/Source/MeshModelingTools/Private/EditMeshPolygonsTool.cpp
@@ -300,6 +300,9 @@ void UEditMeshPolygonsTool::Setup()
 	SelectionMechanic->Setup(this);
 	SelectionMechanic->Properties->RestoreProperties(this, GetPropertyCacheIdentifier(bTriangleMode));
 	SelectionMechanic->OnSelectionChanged.AddUObject(this, &UEditMeshPolygonsTool::OnSelectionModifiedEvent);
+	SelectionMechanic->OnFaceSelectionPreviewChanged.AddWeakLambda(this, [this]() {
+		Preview->PreviewMesh->FastNotifySecondaryTrianglesChanged();
+	});
 	if (bTriangleMode)
 	{
 		SelectionMechanic->PolyEdgesRenderer.LineThickness = 1.0;
diff --git a/Engine/Plugins/Runtime/MeshModelingToolset/Source/ModelingComponents/Private/Selection/GroupTopologySelector.cpp b/Engine/Plugins/Runtime/MeshModelingToolset/Source/ModelingComponents/Private/Selection/GroupTopologySelector.cpp
--- a/Engine/Plugins/Runtime/MeshModelingToolset/Source/ModelingComponents/Private/Selection/GroupTopologySelector.cpp
+++ b/Engine/Plugins/Runtime/MeshModelingToolset/Source/ModelingComponents/Private/Selection/GroupTopologySelector.cpp
@@ -2,7 +2,9 @@
 
 
 #include "Selection/GroupTopologySelector.h"
-#include "Mechanics/RectangleMarqueeMechanic.h"
+
+#include "Engine/Polys.h"
+#include "Mechanics/RectangleMarqueeMechanic.h" // FCameraRectangle
 #include "MeshQueries.h"
 #include "ToolDataVisualizer.h"
 #include "ToolSceneQueriesUtil.h"
@@ -605,7 +607,8 @@ bool IsOccluded(const FVector3d& Point, const FVector3d& ViewOrigin,
 }
 
 bool FGroupTopologySelector::FindSelectedElement(const FSelectionSettings& Settings, 
-	const FCameraRectangle& CameraRectangle, FTransform3d TargetTransform, FGroupTopologySelection& ResultOut)
+	const FCameraRectangle& CameraRectangle, FTransform3d TargetTransform, FGroupTopologySelection& ResultOut,
+	TMap<int32, bool>* TriIsOccludedCache)
 {
 	// One minor easy thing we can do to speed up the below is to detect cases where transform does not have a non-uniform
 	// scale, and in those cases we transform the rectangle once rather than transforming all the query points/curves. 
@@ -666,6 +669,125 @@ bool FGroupTopologySelector::FindSelectedElement(const FSelectionSettings& Setti
 			ResultOut.SelectedEdgeIDs);
 	}
 	
+	if (Settings.bEnableFaceHits && ResultOut.IsEmpty())
+	{
+		// Get a frustum volume in local space, then use it to figure out if it contains or intersects
+		// the aabb tree boxes in the traversal. 
+		// TODO: Frustum transformation is done the same way in FractureEditorMode.cpp- should this be placed
+		// somewhere common? Not sure where to put it though.
+		FConvexVolume WorldSpaceFrustum = CameraRectangle.FrustumAsConvexVolume();
+		FMatrix InverseTargetTransform(FTransform(TargetTransform).ToInverseMatrixWithScale());
+
+		FConvexVolume LocalFrustum;
+		LocalFrustum.Planes.Empty(6);
+		for (const FPlane& Plane : WorldSpaceFrustum.Planes)
+		{
+			LocalFrustum.Planes.Add(Plane.TransformBy(InverseTargetTransform));
+		}
+		LocalFrustum.Init();
+
+		FDynamicMeshAABBTree3::FTreeTraversal Traversal;
+		// Used as additional state for the traversal, relying on the fact that our traversal is depth-first
+		// and serial. When we descend to a box that is fully contained, we use CurrentSelectAllDepth to mark
+		// all deeper boxes as contained, and reset it once we return back out of that subtree (done similarly
+		// in UVEditorMeshSelectionMechanic).
+		int CurrentDepth = -1;
+		int CurrentSelectAllDepth = TNumericLimits<int>::Max();
+
+		Traversal.NextBoxF = [&LocalFrustum, &CurrentDepth, &CurrentSelectAllDepth](const FAxisAlignedBox3d& Box, int Depth)
+		{
+			CurrentDepth = Depth;
+			if (Depth > CurrentSelectAllDepth)
+			{
+				// We're currently inside a fully-selected box
+				return true;
+			}
+			// Reset CurrentSelectAllDepth because we must have left the cotained subtree (if it was set to begin with)
+			CurrentSelectAllDepth = TNumericLimits<int>::Max();
+
+			bool bFullyContained = false;
+			bool bIntersects = LocalFrustum.IntersectBox(Box.Center(), Box.Extents(), bFullyContained);
+			
+			if (bFullyContained)
+			{
+				CurrentSelectAllDepth = Depth;
+			}
+			return bIntersects;
+		};
+
+		Traversal.NextTriangleF = 
+			[this, &ResultOut, &CurrentDepth, &CurrentSelectAllDepth, &LocalFrustum, 
+			&LocalCameraOrigin, &Spatial, &Settings, TriIsOccludedCache](int Tid)
+		{
+			// If this group already got selected, no need to do anything else. Helps avoid some occlusion tests.
+			if (ResultOut.SelectedGroupIDs.Contains(Mesh->GetTriangleGroup(Tid)))
+			{
+				return;
+			}
+
+			// Apply frustum check. Skip if box was fully in frustum
+			if (!(CurrentDepth >= CurrentSelectAllDepth))
+			{
+				// TODO: This isn't an ideal way to check whether the triangle is inside the frustum because FPoly is
+				// an old class that we probably want to avoid using. It's fine enough here, but it should be replaced.
+				// Related note: It may be tempting to just check for containment of the centroid, but we probably don't
+				// want to do that because it makes selection behavior dependent on the tesselation of the underlying
+				// group. I.e. a remeshed cube face ends up being much more reliable to select than one made of two triangles 
+				// where a marquee select of even the group center could fail if it's between the tri centroids.
+				FIndex3i TriVids = Mesh->GetTriangle(Tid);
+				FPoly TrianglePolygon;
+				TrianglePolygon.Init();
+				for (int i = 0; i < 3; ++i)
+				{
+					TrianglePolygon.Vertices.Add(Mesh->GetVertex(TriVids[i]));
+				}
+
+				if (!LocalFrustum.ClipPolygon(TrianglePolygon))
+				{
+					// Triangle was fully outside the frustum
+					return;
+				}
+			}
+
+			// Apply occlusion filter. Note that the back face filter is applied via TriangleFilterF in the traversal.
+			if (!Settings.bIgnoreOcclusion)
+			{
+				bool bTriangleIsOccluded = false;
+
+				// See if we've already calculated occlusion for this tri.
+				bool* CachedOcclusion = TriIsOccludedCache->Find(Tid);
+				if (CachedOcclusion)
+				{
+					bTriangleIsOccluded = *CachedOcclusion;
+				}
+				else
+				{
+					bTriangleIsOccluded = IsOccluded(Mesh->GetTriCentroid(Tid), LocalCameraOrigin, Spatial, Settings.bHitBackFaces);
+					TriIsOccludedCache->Add(Tid, bTriangleIsOccluded);
+				}
+
+				if (bTriangleIsOccluded)
+				{
+					return;
+				}
+			}
+
+			ResultOut.SelectedGroupIDs.Add(Mesh->GetTriangleGroup(Tid));
+		};
+
+		// Add the back face filter, if applicable
+		IMeshSpatial::FQueryOptions Options;
+		if (!Settings.bHitBackFaces)
+		{
+			Options.TriangleFilterF = [this, &LocalCameraOrigin](int32 Tid)
+			{
+				return Mesh->GetTriNormal(Tid).Dot(Mesh->GetVertex(Mesh->GetTriangle(Tid)[0]) - LocalCameraOrigin) < 0;
+			};
+		}
+		
+		Spatial->DoTraversal(Traversal, Options);
+	}
+
 	return ResultOut.IsEmpty();
 }
 
diff --git a/Engine/Plugins/Runtime/MeshModelingToolset/Source/ModelingComponents/Private/Selection/PolygonSelectionMechanic.cpp b/Engine/Plugins/Runtime/MeshModelingToolset/Source/ModelingComponents/Private/Selection/PolygonSelectionMechanic.cpp
--- a/Engine/Plugins/Runtime/MeshModelingToolset/Source/ModelingComponents/Private/Selection/PolygonSelectionMechanic.cpp
+++ b/Engine/Plugins/Runtime/MeshModelingToolset/Source/ModelingComponents/Private/Selection/PolygonSelectionMechanic.cpp
@@ -59,6 +59,9 @@ void UPolygonSelectionMechanic::Setup(UInteractiveTool* ParentToolIn)
 	Properties->WatchProperty(Properties->bSelectEdges, [this](bool bSelectVertices) { 
 		UpdateMarqueeEnabled(); 
 	});
+	Properties->WatchProperty(Properties->bSelectFaces, [this](bool bSelectFaces) {
+		UpdateMarqueeEnabled();
+	});
 	Properties->WatchProperty(Properties->bEnableMarquee, [this](bool bEnableMarquee) { 
 		UpdateMarqueeEnabled(); 
 	});
@@ -473,30 +476,26 @@ bool UPolygonSelectionMechanic::UpdateSelection(const FRay& WorldRay, FVector3d&
 		{
 			TopoSelector->ExpandSelectionByEdgeLoops(Selection);
 		}
+	}
 
-		if (ShouldAddToSelectionFunc())
-		{
-			if (ShouldRemoveFromSelectionFunc())
-			{
-				PersistentSelection.Toggle(Selection);
-			}
-			else
-			{
-				PersistentSelection.Append(Selection);
-			}
-		}
-		else if (ShouldRemoveFromSelectionFunc())
+	if (ShouldAddToSelectionFunc())
+	{
+		if (ShouldRemoveFromSelectionFunc())
 		{
-			PersistentSelection.Remove(Selection);
+			PersistentSelection.Toggle(Selection);
 		}
 		else
 		{
-			PersistentSelection = Selection;
+			PersistentSelection.Append(Selection);
 		}
 	}
+	else if (ShouldRemoveFromSelectionFunc())
+	{
+		PersistentSelection.Remove(Selection);
+	}
 	else
 	{
-		PersistentSelection.Clear();
+		PersistentSelection = Selection;
 	}
 
 	if (PersistentSelection != PreviousSelection)
@@ -605,6 +604,7 @@ void UPolygonSelectionMechanic::OnDragRectangleStarted()
 	BeginChange();
 
 	PreDragPersistentSelection = PersistentSelection;
+	LastUpdateRectangleSelection = PersistentSelection;
 	PreDragTopoSelectorSettings = GetTopoSelectorSettings(false);
 	PreDragTopoSelectorSettings.bIgnoreOcclusion = Properties->bMarqueeIgnoreOcclusion; // uses a separate setting for marquee
 }
@@ -614,7 +614,7 @@ void UPolygonSelectionMechanic::OnDragRectangleChanged(const FCameraRectangle& C
 	FGroupTopologySelection RectangleSelection;
 
 	TopoSelector->FindSelectedElement(PreDragTopoSelectorSettings, CurrentRectangle, TargetTransform,
-		RectangleSelection);
+		RectangleSelection, &TriIsOccludedCache);
 
 	if (ShouldAddToSelectionFunc())
 	{
@@ -638,10 +638,23 @@ void UPolygonSelectionMechanic::OnDragRectangleChanged(const FCameraRectangle& C
 		// Neither key pressed.
 		PersistentSelection = RectangleSelection;
 	}
+
+	// If we modified the currently selected edges/vertices, they will be properly displayed in our
+	// Render() call. However, the mechanic is not responsible for face highlighting, so if we modified
+	// that, we need to notify the user so that they can update the highlighting (since OnSelectionChanged
+	// only gets broadcast at rectangle end).
+	if ((!PersistentSelection.SelectedGroupIDs.IsEmpty() || !LastUpdateRectangleSelection.SelectedGroupIDs.IsEmpty()) // if groups are involved
+		&& PersistentSelection != LastUpdateRectangleSelection)
+	{
+		LastUpdateRectangleSelection = PersistentSelection;
+		OnFaceSelectionPreviewChanged.Broadcast();
+	}
 }
 
 void UPolygonSelectionMechanic::OnDragRectangleFinished(const FCameraRectangle& Rectangle, bool bCancelled)
 {
+	TriIsOccludedCache.Reset();
+
 	if (PersistentSelection != PreDragPersistentSelection)
 	{
 		SelectionTimestamp++;
@@ -656,8 +669,8 @@ void UPolygonSelectionMechanic::UpdateMarqueeEnabled()
 {
 	MarqueeMechanic->SetIsEnabled(
 		bIsEnabled 
-		&& Properties->bEnableMarquee 
-		&& (Properties->bSelectVertices || Properties->bSelectEdges));
+		&& Properties->bEnableMarquee
+		&& (Properties->bSelectVertices || Properties->bSelectEdges || Properties->bSelectFaces));
 }
 
 
diff --git a/Engine/Plugins/Runtime/MeshModelingToolset/Source/ModelingComponents/Public/Selection/GroupTopologySelector.h b/Engine/Plugins/Runtime/MeshModelingToolset/Source/ModelingComponents/Public/Selection/GroupTopologySelector.h
--- a/Engine/Plugins/Runtime/MeshModelingToolset/Source/ModelingComponents/Public/Selection/GroupTopologySelector.h
+++ b/Engine/Plugins/Runtime/MeshModelingToolset/Source/ModelingComponents/Public/Selection/GroupTopologySelector.h
@@ -123,21 +123,26 @@ public:
 		FVector3d& SelectedPositionOut, FVector3d& SelectedNormalOut, int32* EdgeSegmentIdOut = nullptr);
 
 	/**
-	 * Given a camera rectangle (from a marquee selection) gives the group corners or edges that are contained in
-	 * the rectangle. Currently does not select faces.
+	 * Given a camera rectangle (from a marquee selection) gives the corners/edges/groups that are contained in
+	 * the rectangle. Only selects one of the types, and prefers corners to edges to groups.
 	 * 
-	 * Uses the bEnableEdgeHits, bEnableCornerHits, and bIgnoreOcclusion members of the FSelectionSettings argument.
 	 * For group edges, occlusion is checked at endpoints, and if either is occluded, the edge is considered occluded.
+	 * For faces, occlusion is checked at triangle centroids.
 	 *
 	 * @param Settings Settings that determine what is selected and whether occlusion is tested.
 	 * @param CameraRectangle The camera rectangle.
 	 * @param TargetTransform Transform from points to world space, since the rectangle is in world space.
 	 * @param ResultOut Output selection (cleared before use).
+	 * @param TriIsOccludedCache Cached values of whether a given Tid has an occluded centroid, since occlusion 
+	 *  values will stay the same if the rectangle is resized as long as the camera is not moved. This should
+	 *  be used when doing occlusion-based marquee face selection on meshes with many triangles, because the 
+	 *  occlusion checks are expensive.
 	 *
 	 * @return true if something was selected.
 	 */
 	bool FindSelectedElement(const FSelectionSettings& Settings, const FCameraRectangle& CameraRectangle, 
-		FTransform3d TargetTransform, FGroupTopologySelection& ResultOut);
+		FTransform3d TargetTransform, FGroupTopologySelection& ResultOut, 
+		TMap<int32, bool>* TriIsOccludedCache = nullptr);
 
 	/**
 	 * Using the edges in the given selection as starting points, add any "edge loops" containing the edges. An 
diff --git a/Engine/Plugins/Runtime/MeshModelingToolset/Source/ModelingComponents/Public/Selection/PolygonSelectionMechanic.h b/Engine/Plugins/Runtime/MeshModelingToolset/Source/ModelingComponents/Public/Selection/PolygonSelectionMechanic.h
--- a/Engine/Plugins/Runtime/MeshModelingToolset/Source/ModelingComponents/Public/Selection/PolygonSelectionMechanic.h
+++ b/Engine/Plugins/Runtime/MeshModelingToolset/Source/ModelingComponents/Public/Selection/PolygonSelectionMechanic.h
@@ -325,6 +325,13 @@ public:
 	 */
 	FSimpleMulticastDelegate OnSelectionChanged;
 
+	/**
+	 * This is broadcast during marquee rectangle dragging if selected faces change, to allow user to
+	 * dynamically update face highlighting if desired (needed because OnSelectionChanged is only 
+	 * fired once the rectangle is completed, not while it is being updated).
+	 */
+	FSimpleMulticastDelegate OnFaceSelectionPreviewChanged;
+
 	// TODO: Is it worth issuing separate callbacks in normal selection changes and in FChange ones, to
 	// allow the user to bundle in some FChanges into the normal callback?
 
@@ -381,7 +388,9 @@ protected:
 
 	// Used for box selection
 	FGroupTopologySelection PreDragPersistentSelection;
+	FGroupTopologySelection LastUpdateRectangleSelection;
 	FGroupTopologySelector::FSelectionSettings PreDragTopoSelectorSettings;
+	TMap<int32, bool> TriIsOccludedCache;
 
 	FVector3d LastClickedHitPosition;
 	FVector3d LastClickedHitNormal;
