PolyExtrude tool: fix an ensure when a polygon self-intersection passes exactly through vertex

#rnx

diff --git a/Engine/Plugins/Runtime/MeshModelingToolset/Source/MeshModelingTools/Private/DrawPolygonTool.cpp b/Engine/Plugins/Runtime/MeshModelingToolset/Source/MeshModelingTools/Private/DrawPolygonTool.cpp
--- a/Engine/Plugins/Runtime/MeshModelingToolset/Source/MeshModelingTools/Private/DrawPolygonTool.cpp
+++ b/Engine/Plugins/Runtime/MeshModelingToolset/Source/MeshModelingTools/Private/DrawPolygonTool.cpp
@@ -759,6 +759,12 @@ bool UDrawPolygonTool::OnNextSequenceClick(const FInputDeviceRay& ClickPos)
 
 		if (bHaveSelfIntersection)
 		{
+			// If the self-intersection point is coincident with a polygon vertex, don't add that point twice (it would produce a degenerate polygon edge)
+			if (SelfIntersectSegmentIdx < PolygonVertices.Num()-1 && FVector3d::PointsAreSame(SelfIntersectionPoint, PolygonVertices[SelfIntersectSegmentIdx+1]))
+			{
+				++SelfIntersectSegmentIdx;
+			}
+
 			// discard vertex in segments before intersection (this is redundant if idx is 0)
 			for (int j = SelfIntersectSegmentIdx; j < PolygonVertices.Num(); ++j)
 			{
