Non Unity Fixes

diff --git a/Engine/Plugins/Runtime/MeshModelingToolset/Source/ModelingComponents/Private/Selection/GroupTopologySelector.cpp b/Engine/Plugins/Runtime/MeshModelingToolset/Source/ModelingComponents/Private/Selection/GroupTopologySelector.cpp
--- a/Engine/Plugins/Runtime/MeshModelingToolset/Source/ModelingComponents/Private/Selection/GroupTopologySelector.cpp
+++ b/Engine/Plugins/Runtime/MeshModelingToolset/Source/ModelingComponents/Private/Selection/GroupTopologySelector.cpp
@@ -8,6 +8,7 @@
 #include "MeshQueries.h"
 #include "ToolDataVisualizer.h"
 #include "ToolSceneQueriesUtil.h"
+#include "ConvexVolume.h"
 
 using namespace UE::Geometry;
 
