MeshModelingTools: Fix triangle marquee select not working in TriEdit because it was querying group id from the mesh instead of topology.

#rnx

diff --git a/Engine/Plugins/Runtime/MeshModelingToolset/Source/ModelingComponents/Private/Selection/GroupTopologySelector.cpp b/Engine/Plugins/Runtime/MeshModelingToolset/Source/ModelingComponents/Private/Selection/GroupTopologySelector.cpp
--- a/Engine/Plugins/Runtime/MeshModelingToolset/Source/ModelingComponents/Private/Selection/GroupTopologySelector.cpp
+++ b/Engine/Plugins/Runtime/MeshModelingToolset/Source/ModelingComponents/Private/Selection/GroupTopologySelector.cpp
@@ -3,12 +3,12 @@
 
 #include "Selection/GroupTopologySelector.h"
 
+#include "ConvexVolume.h"
 #include "Engine/Polys.h"
 #include "Mechanics/RectangleMarqueeMechanic.h" // FCameraRectangle
 #include "MeshQueries.h"
 #include "ToolDataVisualizer.h"
 #include "ToolSceneQueriesUtil.h"
-#include "ConvexVolume.h"
 
 using namespace UE::Geometry;
 
@@ -720,8 +720,10 @@ bool FGroupTopologySelector::FindSelectedElement(const FSelectionSettings& Setti
 			[this, &ResultOut, &CurrentDepth, &CurrentSelectAllDepth, &LocalFrustum, 
 			&LocalCameraOrigin, &Spatial, &Settings, TriIsOccludedCache](int Tid)
 		{
+			int32 TriGroupID = Topology->GetGroupID(Tid);
+
 			// If this group already got selected, no need to do anything else. Helps avoid some occlusion tests.
-			if (ResultOut.SelectedGroupIDs.Contains(Mesh->GetTriangleGroup(Tid)))
+			if (ResultOut.SelectedGroupIDs.Contains(TriGroupID))
 			{
 				return;
 			}
@@ -773,7 +775,7 @@ bool FGroupTopologySelector::FindSelectedElement(const FSelectionSettings& Setti
 				}
 			}
 
-			ResultOut.SelectedGroupIDs.Add(Mesh->GetTriangleGroup(Tid));
+			ResultOut.SelectedGroupIDs.Add(TriGroupID);
 		};
 
 		// Add the back face filter, if applicable
