MeshModelingTools: Reorder selection filter in PolyEd to have vertices before faces.

#rnx

diff --git a/Engine/Plugins/Runtime/MeshModelingToolset/Source/ModelingComponents/Public/Selection/PolygonSelectionMechanic.h b/Engine/Plugins/Runtime/MeshModelingToolset/Source/ModelingComponents/Public/Selection/PolygonSelectionMechanic.h
--- a/Engine/Plugins/Runtime/MeshModelingToolset/Source/ModelingComponents/Public/Selection/PolygonSelectionMechanic.h
+++ b/Engine/Plugins/Runtime/MeshModelingToolset/Source/ModelingComponents/Public/Selection/PolygonSelectionMechanic.h
@@ -31,13 +31,13 @@ class MODELINGCOMPONENTS_API UPolygonSelectionMechanicProperties : public UInter
 
 public:
 	UPROPERTY(EditAnywhere, Category = SelectionFilter)
-	bool bSelectFaces = true;
+	bool bSelectVertices = true;
 
 	UPROPERTY(EditAnywhere, Category = SelectionFilter)
 	bool bSelectEdges = true;
 
 	UPROPERTY(EditAnywhere, Category = SelectionFilter)
-	bool bSelectVertices = true;
+	bool bSelectFaces = true;
 
 	/** When true, will select edge loops. Edge loops are paths along a string of valence-4 vertices. */
 	UPROPERTY(EditAnywhere, Category = SelectionFilter, meta = (EditCondition = "bSelectEdges"))
