MeshModelingTools: Hide selection options in PolyEd while an activity is active (since they are not used).

#rnx

diff --git a/Engine/Plugins/Runtime/MeshModelingToolset/Source/MeshModelingTools/Private/EditMeshPolygonsTool.cpp b/Engine/Plugins/Runtime/MeshModelingToolset/Source/MeshModelingTools/Private/EditMeshPolygonsTool.cpp
--- a/Engine/Plugins/Runtime/MeshModelingToolset/Source/MeshModelingTools/Private/EditMeshPolygonsTool.cpp
+++ b/Engine/Plugins/Runtime/MeshModelingToolset/Source/MeshModelingTools/Private/EditMeshPolygonsTool.cpp
@@ -1002,6 +1002,7 @@ void UEditMeshPolygonsTool::StartActivity(TObjectPtr<UInteractiveToolActivity> A
 	{
 		TransformGizmo->SetVisibility(false);
 		SelectionMechanic->SetIsEnabled(false);
+		SetToolPropertySourceEnabled(SelectionMechanic->Properties, false);
 		CurrentActivity = Activity;
 		if (CurrentActivity == BevelEdgeActivity)
 		{
@@ -1035,6 +1036,7 @@ void UEditMeshPolygonsTool::EndCurrentActivity(EToolShutdownType ShutdownType)
 		SetToolPropertySourceEnabled(AcceptCancelAction, false);
 		SetActionButtonPanelsVisible(true);
 		SelectionMechanic->SetIsEnabled(true);
+		SetToolPropertySourceEnabled(SelectionMechanic->Properties, true);
 		UpdateGizmoVisibility();
 	}
 }
