Fix non-unity error in SingleSelectionMeshEditingTool & MultiSelectionMeshEditingTool

#rnx

diff --git a/Engine/Plugins/Runtime/MeshModelingToolset/Source/ModelingComponents/Private/BaseTools/MultiSelectionMeshEditingTool.cpp b/Engine/Plugins/Runtime/MeshModelingToolset/Source/ModelingComponents/Private/BaseTools/MultiSelectionMeshEditingTool.cpp
--- a/Engine/Plugins/Runtime/MeshModelingToolset/Source/ModelingComponents/Private/BaseTools/MultiSelectionMeshEditingTool.cpp
+++ b/Engine/Plugins/Runtime/MeshModelingToolset/Source/ModelingComponents/Private/BaseTools/MultiSelectionMeshEditingTool.cpp
@@ -2,6 +2,8 @@
 
 #include "BaseTools/MultiSelectionMeshEditingTool.h"
 
+#include "Engine/World.h"
+
 #include "TargetInterfaces/MaterialProvider.h"
 #include "TargetInterfaces/MeshDescriptionCommitter.h"
 #include "TargetInterfaces/MeshDescriptionProvider.h"
@@ -47,6 +49,17 @@ void UMultiSelectionMeshEditingToolBuilder::InitializeNewTool(UMultiSelectionMes
  * Tool
  */
 
+void UMultiSelectionMeshEditingTool::SetWorld(UWorld* World)
+{
+	TargetWorld = World;
+}
+
+UWorld* UMultiSelectionMeshEditingTool::GetTargetWorld()
+{
+	return TargetWorld.Get();
+}
+
+
 bool UMultiSelectionMeshEditingTool::GetMapToSharedSourceData(TArray<int32>& MapToFirstOccurrences)
 {
 	bool bSharesSources = false;
diff --git a/Engine/Plugins/Runtime/MeshModelingToolset/Source/ModelingComponents/Private/BaseTools/SingleSelectionMeshEditingTool.cpp b/Engine/Plugins/Runtime/MeshModelingToolset/Source/ModelingComponents/Private/BaseTools/SingleSelectionMeshEditingTool.cpp
--- a/Engine/Plugins/Runtime/MeshModelingToolset/Source/ModelingComponents/Private/BaseTools/SingleSelectionMeshEditingTool.cpp
+++ b/Engine/Plugins/Runtime/MeshModelingToolset/Source/ModelingComponents/Private/BaseTools/SingleSelectionMeshEditingTool.cpp
@@ -2,6 +2,8 @@
 
 #include "BaseTools/SingleSelectionMeshEditingTool.h"
 
+#include "Engine/World.h"
+
 #include "TargetInterfaces/MaterialProvider.h"
 #include "TargetInterfaces/MeshDescriptionCommitter.h"
 #include "TargetInterfaces/MeshDescriptionProvider.h"
@@ -52,4 +54,16 @@ void USingleSelectionMeshEditingToolBuilder::InitializeNewTool(USingleSelectionM
 	}
 }
 
+void USingleSelectionMeshEditingTool::SetWorld(UWorld* World)
+{
+	TargetWorld = World;
+}
+
+UWorld* USingleSelectionMeshEditingTool::GetTargetWorld()
+{
+	return TargetWorld.Get();
+}
+
+
+
 
diff --git a/Engine/Plugins/Runtime/MeshModelingToolset/Source/ModelingComponents/Public/BaseTools/MultiSelectionMeshEditingTool.h b/Engine/Plugins/Runtime/MeshModelingToolset/Source/ModelingComponents/Public/BaseTools/MultiSelectionMeshEditingTool.h
--- a/Engine/Plugins/Runtime/MeshModelingToolset/Source/ModelingComponents/Public/BaseTools/MultiSelectionMeshEditingTool.h
+++ b/Engine/Plugins/Runtime/MeshModelingToolset/Source/ModelingComponents/Public/BaseTools/MultiSelectionMeshEditingTool.h
@@ -44,8 +44,8 @@ class MODELINGCOMPONENTS_API UMultiSelectionMeshEditingTool : public UMultiSelec
 {
 	GENERATED_BODY()
 public:
-	virtual void SetWorld(UWorld* World) { TargetWorld = World; }
-	virtual UWorld* GetTargetWorld() { return TargetWorld.Get(); }
+	virtual void SetWorld(UWorld* World);
+	virtual UWorld* GetTargetWorld();
 
 	/**
 	 * Helper to find which targets share source data.
diff --git a/Engine/Plugins/Runtime/MeshModelingToolset/Source/ModelingComponents/Public/BaseTools/SingleSelectionMeshEditingTool.h b/Engine/Plugins/Runtime/MeshModelingToolset/Source/ModelingComponents/Public/BaseTools/SingleSelectionMeshEditingTool.h
--- a/Engine/Plugins/Runtime/MeshModelingToolset/Source/ModelingComponents/Public/BaseTools/SingleSelectionMeshEditingTool.h
+++ b/Engine/Plugins/Runtime/MeshModelingToolset/Source/ModelingComponents/Public/BaseTools/SingleSelectionMeshEditingTool.h
@@ -49,9 +49,8 @@ class MODELINGCOMPONENTS_API USingleSelectionMeshEditingTool : public USingleSel
 	GENERATED_BODY()
 
 public:
-	virtual void SetWorld(UWorld* World) { TargetWorld = World; }
-
-	virtual UWorld* GetTargetWorld() { return TargetWorld.Get(); }
+	virtual void SetWorld(UWorld* World);
+	virtual UWorld* GetTargetWorld();
 
 protected:
 	UPROPERTY()
