MeshModelingTools: Fix there being two separate "Shape Edits" categories in PolyEd detail panel.

#rnx

diff --git a/Engine/Plugins/Runtime/MeshModelingToolset/Source/MeshModelingTools/Public/EditMeshPolygonsTool.h b/Engine/Plugins/Runtime/MeshModelingToolset/Source/MeshModelingTools/Public/EditMeshPolygonsTool.h
--- a/Engine/Plugins/Runtime/MeshModelingToolset/Source/MeshModelingTools/Public/EditMeshPolygonsTool.h
+++ b/Engine/Plugins/Runtime/MeshModelingToolset/Source/MeshModelingTools/Public/EditMeshPolygonsTool.h
@@ -251,8 +251,16 @@ public:
 	UFUNCTION(CallInEditor, Category = FaceEdits, meta = (DisplayName = "Duplicate", DisplayPriority = 12))
 	void Duplicate() { PostAction(EEditMeshPolygonsToolActions::Duplicate); }
 
+	//~ TODO: add tooltip, especially explaining limitations, for this and InsertEdge. Can't currently do that
+	//~ without cutting down another tooltip until UE-124608 is fixed...
+	UFUNCTION(CallInEditor, Category = ShapeEdits, meta = (DisplayName = "InsertEdgeLoop", DisplayPriority = 13))
+	void InsertEdgeLoop() { PostAction(EEditMeshPolygonsToolActions::InsertEdgeLoop); }
+
+	UFUNCTION(CallInEditor, Category = ShapeEdits, meta = (DisplayName = "Insert Edge", DisplayPriority = 14))
+	void InsertEdge() { PostAction(EEditMeshPolygonsToolActions::InsertEdge); }
+
 	/** Simplify every polygon group by removing vertices on shared straight edges and retriangulating */
-	UFUNCTION(CallInEditor, Category = ShapeEdits, meta = (DisplayName = "SimplifyByGroups", DisplayPriority = 13))
+	UFUNCTION(CallInEditor, Category = ShapeEdits, meta = (DisplayName = "SimplifyByGroups", DisplayPriority = 15))
 	void SimplifyByGroups() { PostAction(EEditMeshPolygonsToolActions::SimplifyByGroups); }
 
 };
@@ -341,23 +349,17 @@ class MESHMODELINGTOOLS_API UEditMeshPolygonsToolEdgeActions : public UEditMeshP
 {
 	GENERATED_BODY()
 public:
-	UFUNCTION(CallInEditor, Category = ShapeEdits, meta = (DisplayName = "InsertEdgeLoop", DisplayPriority = 1))
-	void InsertEdgeLoop() { PostAction(EEditMeshPolygonsToolActions::InsertEdgeLoop); }
-
-	UFUNCTION(CallInEditor, Category = ShapeEdits, meta = (DisplayName = "Insert Edge", DisplayPriority = 2))
-	void InsertEdge() { PostAction(EEditMeshPolygonsToolActions::InsertEdge); }
-
-	UFUNCTION(CallInEditor, Category = EdgeEdits, meta = (DisplayName = "Weld", DisplayPriority = 3))
+	UFUNCTION(CallInEditor, Category = EdgeEdits, meta = (DisplayName = "Weld", DisplayPriority = 1))
 	void Weld() { PostAction(EEditMeshPolygonsToolActions::WeldEdges); }
 
-	UFUNCTION(CallInEditor, Category = EdgeEdits, meta = (DisplayName = "Straighten", DisplayPriority = 4))
+	UFUNCTION(CallInEditor, Category = EdgeEdits, meta = (DisplayName = "Straighten", DisplayPriority = 2))
 	void Straighten() { PostAction(EEditMeshPolygonsToolActions::StraightenEdge); }
 
 	/** Fill the adjacent hole for any selected boundary edges */
-	UFUNCTION(CallInEditor, Category = EdgeEdits, meta = (DisplayName = "Fill Hole", DisplayPriority = 5))
+	UFUNCTION(CallInEditor, Category = EdgeEdits, meta = (DisplayName = "Fill Hole", DisplayPriority = 3))
 	void FillHole()	{ PostAction(EEditMeshPolygonsToolActions::FillHole); }
 
-	UFUNCTION(CallInEditor, Category = EdgeEdits, meta = (DisplayName = "Bevel", DisplayPriority = 6))
+	UFUNCTION(CallInEditor, Category = EdgeEdits, meta = (DisplayName = "Bevel", DisplayPriority = 4))
 	void Bevel() { PostAction(EEditMeshPolygonsToolActions::BevelEdges); }
 
 };
