MeshModelingTools: Add SelectAll and InvertSelection actions to PolyEd.

#rnx

diff --git a/Engine/Plugins/Runtime/MeshModelingToolset/Source/ModelingComponents/Private/Selection/PolygonSelectionMechanic.cpp b/Engine/Plugins/Runtime/MeshModelingToolset/Source/ModelingComponents/Private/Selection/PolygonSelectionMechanic.cpp
--- a/Engine/Plugins/Runtime/MeshModelingToolset/Source/ModelingComponents/Private/Selection/PolygonSelectionMechanic.cpp
+++ b/Engine/Plugins/Runtime/MeshModelingToolset/Source/ModelingComponents/Private/Selection/PolygonSelectionMechanic.cpp
@@ -14,6 +14,22 @@ using namespace UE::Geometry;
 
 #define LOCTEXT_NAMESPACE "UPolygonSelectionMechanic"
 
+void UPolygonSelectionMechanicProperties::InvertSelection()
+{
+	if (Mechanic.IsValid())
+	{
+		Mechanic->InvertSelection();
+	}
+}
+
+void UPolygonSelectionMechanicProperties::SelectAll()
+{
+	if (Mechanic.IsValid())
+	{
+		Mechanic->SelectAll();
+	}
+}
+
 UPolygonSelectionMechanic::~UPolygonSelectionMechanic()
 {
 	checkf(PreviewGeometryActor == nullptr, TEXT("Shutdown() should be called before UPolygonSelectionMechanic is destroyed."));
@@ -49,6 +65,7 @@ void UPolygonSelectionMechanic::Setup(UInteractiveTool* ParentToolIn)
 	ParentTool->AddInputBehavior(ClickOrDragBehavior, this);
 
 	Properties = NewObject<UPolygonSelectionMechanicProperties>(this);
+	Properties->Initialize(this);
 	if (bAddSelectionFilterPropertiesToParentTool)
 	{
 		AddToolPropertySource(Properties);
@@ -527,6 +544,105 @@ void UPolygonSelectionMechanic::ClearSelection()
 	OnSelectionChanged.Broadcast();
 }
 
+void UPolygonSelectionMechanic::InvertSelection()
+{
+	if (PersistentSelection.IsEmpty())
+	{
+		SelectAll();
+		return;
+	}
+
+	ParentTool->GetToolManager()->BeginUndoTransaction(LOCTEXT("SelectionChange", "Selection"));
+	BeginChange();
+
+	const FGroupTopologySelection PreviousSelection = PersistentSelection;
+	PersistentSelection.Clear();
+
+	if (!PreviousSelection.SelectedCornerIDs.IsEmpty())
+	{
+		for (int32 CornerID = 0; CornerID < Topology->Corners.Num(); ++CornerID)
+		{
+			if (!PreviousSelection.SelectedCornerIDs.Contains(CornerID))
+			{
+				PersistentSelection.SelectedCornerIDs.Add(CornerID);
+			}
+		}
+	}
+	else if (!PreviousSelection.SelectedEdgeIDs.IsEmpty())
+	{
+		for (int32 EdgeID = 0; EdgeID < Topology->Edges.Num(); ++EdgeID)
+		{
+			if (!PreviousSelection.SelectedEdgeIDs.Contains(EdgeID))
+			{
+				PersistentSelection.SelectedEdgeIDs.Add(EdgeID);
+			}
+		}
+	}
+	else if (!PreviousSelection.SelectedGroupIDs.IsEmpty())
+	{
+		for (const FGroupTopology::FGroup& Group : Topology->Groups)
+		{
+			if (!PreviousSelection.SelectedGroupIDs.Contains(Group.GroupID))
+			{
+				PersistentSelection.SelectedGroupIDs.Add(Group.GroupID);
+			}
+		}
+	}
+
+	SelectionTimestamp++;
+	OnSelectionChanged.Broadcast();
+	EndChangeAndEmitIfModified();
+	ParentTool->GetToolManager()->EndUndoTransaction();
+}
+
+void UPolygonSelectionMechanic::SelectAll()
+{
+	const FGroupTopologySelection PreviousSelection = PersistentSelection;
+
+	ParentTool->GetToolManager()->BeginUndoTransaction(LOCTEXT("SelectionChange", "Selection"));
+	BeginChange();
+
+	auto SelectAllIndices = [](int32 MaxExclusiveIndex, TSet<int32>& ContainerOut)
+	{
+		for (int32 i = 0; i < MaxExclusiveIndex; ++i)
+		{
+			ContainerOut.Add(i);
+		}
+	};
+
+	PersistentSelection.Clear();
+
+	// Select based on settings, prefering corners to edges to groups (since this is the preference we have
+	// elsewhere, eg in marquee).
+	if (Properties->bSelectVertices)
+	{
+		SelectAllIndices(Topology->Corners.Num(), PersistentSelection.SelectedCornerIDs);
+	}
+	else if (Properties->bSelectEdges)
+	{
+		SelectAllIndices(Topology->Edges.Num(), PersistentSelection.SelectedEdgeIDs);
+	}
+	else if (Properties->bSelectFaces)
+	{
+		for (const FGroupTopology::FGroup& Group : Topology->Groups)
+		{
+			PersistentSelection.SelectedGroupIDs.Add(Group.GroupID);
+		}
+	}
+
+	SelectionTimestamp++;
+	OnSelectionChanged.Broadcast();
+	
+	if (PreviousSelection != PersistentSelection)
+	{
+		SelectionTimestamp++;
+		OnSelectionChanged.Broadcast();
+	}
+	
+	EndChangeAndEmitIfModified();
+	ParentTool->GetToolManager()->EndUndoTransaction();
+}
+
 FInputRayHit UPolygonSelectionMechanic::IsHitByClick(const FInputDeviceRay& ClickPos)
 {
 	if (!bIsEnabled)
diff --git a/Engine/Plugins/Runtime/MeshModelingToolset/Source/ModelingComponents/Public/Selection/PolygonSelectionMechanic.h b/Engine/Plugins/Runtime/MeshModelingToolset/Source/ModelingComponents/Public/Selection/PolygonSelectionMechanic.h
--- a/Engine/Plugins/Runtime/MeshModelingToolset/Source/ModelingComponents/Public/Selection/PolygonSelectionMechanic.h
+++ b/Engine/Plugins/Runtime/MeshModelingToolset/Source/ModelingComponents/Public/Selection/PolygonSelectionMechanic.h
@@ -18,6 +18,7 @@
 class FPolygonSelectionMechanicSelectionChange;
 class UPersistentMeshSelection;
 class UMouseHoverBehavior;
+class UPolygonSelectionMechanic;
 class URectangleMarqueeMechanic;
 class USingleClickOrDragInputBehavior;
 
@@ -73,6 +74,22 @@ public:
 	/** Do not check whether the closest element is occluded from the current view. */
 	UPROPERTY(EditAnywhere, Category = "SelectionFilter|Ortho Viewport Behavior")
 	bool bIgnoreOcclusion = false;
+
+	/** Invert current selection. If selection is empty, has same effect as Select All, and is similarly dependent on selection filter. */
+	UFUNCTION(CallInEditor, Category = SelectionActions)
+	void InvertSelection();
+
+	/** Select all elements. Depends on selection filter, where vertices are preferred to edges to faces. */
+	UFUNCTION(CallInEditor, Category = SelectionActions)
+	void SelectAll();
+
+	void Initialize(UPolygonSelectionMechanic* MechanicIn)
+	{
+		Mechanic = MechanicIn;
+	}
+
+protected:
+	TWeakObjectPtr<UPolygonSelectionMechanic> Mechanic;
 };
 
 
@@ -239,6 +256,9 @@ public:
 	 */
 	void ClearSelection();
 
+	void InvertSelection();
+	void SelectAll();
+
 	/** 
 	 * @return true if the current selection is non-empty 
 	 */
