MeshModelingTools: Disable PolyEdit and TriEdit on high res meshes to prevent them from hanging the editor.

#rnx

diff --git a/Engine/Plugins/Runtime/MeshModelingToolset/Source/MeshModelingTools/Private/EditMeshPolygonsTool.cpp b/Engine/Plugins/Runtime/MeshModelingToolset/Source/MeshModelingTools/Private/EditMeshPolygonsTool.cpp
--- a/Engine/Plugins/Runtime/MeshModelingToolset/Source/MeshModelingTools/Private/EditMeshPolygonsTool.cpp
+++ b/Engine/Plugins/Runtime/MeshModelingToolset/Source/MeshModelingTools/Private/EditMeshPolygonsTool.cpp
@@ -67,6 +67,13 @@ namespace EditMeshPolygonsToolLocals
 	{
 		return bTriangleMode ? TEXT("TriEditTool") : TEXT("PolyEditTool");
 	}
+
+	TAutoConsoleVariable<int32> CVarEdgeLimit(
+		TEXT("modeling.PolyEdit.EdgeLimit"),
+		60000,
+		TEXT("Maximal number of edges that PolyEd and TriEd support. Meshes that would require "
+			"more than this number of edges to be rendered in PolyEd or TriEd force the tools to "
+			"be disabled to avoid hanging the editor."));
 }
 
 /*
@@ -185,6 +192,55 @@ void UEditMeshPolygonsTool::Setup()
 {
 	using namespace EditMeshPolygonsToolLocals;
 
+	// TODO: Currently we draw all the edges in the tool with PDI and can lock up the editor on high-res meshes. 
+	// As a hack, disable everything if the number of edges is too high, so that user doesn't lose work accidentally
+	// if they start the tool on the wrong thing.
+	int32 MaxEdges = CVarEdgeLimit.GetValueOnGameThread();
+
+	CurrentMesh = MakeShared<FDynamicMesh3>(UE::ToolTarget::GetDynamicMeshCopy(Target));
+	if (bTriangleMode)
+	{
+		bToolDisabled = CurrentMesh->EdgeCount() > MaxEdges;
+		if (bToolDisabled)
+		{
+			GetToolManager()->DisplayMessage(FText::Format(
+				LOCTEXT("TriEditTooManyEdges", 
+					"This tool is currently disallowed from operating on a mesh of this resolution. "
+					"Current limit set by \"modeling.PolyEdit.EdgeLimit\" is {0} edges, and mesh has "
+					"{1}. Limit can be changed but exists to avoid hanging the editor when trying to "
+					"render too many edges using the current system, so make sure to save your work "
+					"if you change the upper limit and try to edit a very dense mesh."),
+				MaxEdges, CurrentMesh->EdgeCount()), EToolMessageLevel::UserError);
+			return;
+		}
+	}
+
+	Topology = (bTriangleMode) ? MakeShared<FTriangleGroupTopology, ESPMode::ThreadSafe>(CurrentMesh.Get(), true)
+		: MakeShared<FGroupTopology, ESPMode::ThreadSafe>(CurrentMesh.Get(), true);
+	
+	if (!bTriangleMode)
+	{
+		int32 NumEdgesToRender = 0;
+		for (const FGroupTopology::FGroupEdge& Edge : Topology->Edges)
+		{
+			NumEdgesToRender += Edge.Span.Edges.Num();
+		}
+
+		bToolDisabled = NumEdgesToRender > MaxEdges;
+		if (bToolDisabled)
+		{
+			GetToolManager()->DisplayMessage(FText::Format(
+				LOCTEXT("PolyEditTooManyEdges",
+					"This tool is currently disallowed from operating on a group topology of this resolution. "
+					"Current limit set by \"modeling.PolyEdit.EdgeLimit\" is {0} displayed edges, and topology has "
+					"{1} edge segments to display. Limit can be changed, but it exists to avoid hanging the editor "
+					"when trying to render too many edges using the current system, so make sure to save your work "
+					"if you change the upper limit and try to edit a very complicated topology."),
+				MaxEdges, NumEdgesToRender), EToolMessageLevel::UserError);
+			return;
+		}
+	}
+
 	// Start by adding the actions, because we want them at the top.
 	if (bTriangleMode)
 	{
@@ -214,7 +270,6 @@ void UEditMeshPolygonsTool::Setup()
 		AddToolPropertySource(EditUVActions);
 
 		DefaultMessage = TriEditDefaultMessage;
-		
 	}
 
 	GetToolManager()->DisplayMessage(DefaultMessage,
@@ -248,8 +303,6 @@ void UEditMeshPolygonsTool::Setup()
 	// and cause UpdateGizmoFrame() to be called emitting a spurious Transform change. 
 	CommonProps->SilentUpdateWatched();
 
-	CurrentMesh = MakeShared<FDynamicMesh3>(UE::ToolTarget::GetDynamicMeshCopy(Target));
-
 	// TODO: Do we need this?
 	FMeshNormals::QuickComputeVertexNormals(*CurrentMesh);
 
@@ -291,9 +344,6 @@ void UEditMeshPolygonsTool::Setup()
 	MeshSpatial = MakeShared<FDynamicMeshAABBTree3>();
 	MeshSpatial->SetMesh(CurrentMesh.Get());
 
-	Topology = (bTriangleMode) ? MakeShared<FTriangleGroupTopology, ESPMode::ThreadSafe>(CurrentMesh.Get(), true)
-		: MakeShared<FGroupTopology, ESPMode::ThreadSafe>(CurrentMesh.Get(), true);
-
 	// set up SelectionMechanic
 	SelectionMechanic = NewObject<UPolygonSelectionMechanic>(this);
 	SelectionMechanic->bAddSelectionFilterPropertiesToParentTool = false; // We'll do this ourselves later
@@ -451,6 +501,13 @@ void UEditMeshPolygonsTool::OnShutdown(EToolShutdownType ShutdownType)
 {
 	using namespace EditMeshPolygonsToolLocals;
 
+	if (bToolDisabled)
+	{
+		CurrentMesh.Reset();
+		Topology.Reset();
+		return;
+	}
+
 	if (CurrentActivity)
 	{
 		CurrentActivity->End(ShutdownType);
@@ -822,6 +879,11 @@ void UEditMeshPolygonsTool::ComputeUpdate_Gizmo()
 
 void UEditMeshPolygonsTool::OnTick(float DeltaTime)
 {
+	if (bToolDisabled)
+	{
+		return;
+	}
+
 	Preview->Tick(DeltaTime);
 
 	if (CurrentActivity)
@@ -1065,6 +1127,11 @@ void UEditMeshPolygonsTool::UpdateGizmoVisibility()
 
 void UEditMeshPolygonsTool::Render(IToolsContextRenderAPI* RenderAPI)
 {
+	if (bToolDisabled)
+	{
+		return;
+	}
+
 	Preview->PreviewMesh->EnableWireframe(CommonProps->bShowWireframe);
 	SelectionMechanic->Render(RenderAPI);
 	DragAlignmentMechanic->Render(RenderAPI);
@@ -1077,6 +1144,11 @@ void UEditMeshPolygonsTool::Render(IToolsContextRenderAPI* RenderAPI)
 
 void UEditMeshPolygonsTool::DrawHUD(FCanvas* Canvas, IToolsContextRenderAPI* RenderAPI)
 {
+	if (bToolDisabled)
+	{
+		return;
+	}
+
 	SelectionMechanic->DrawHUD(Canvas, RenderAPI);
 }
 
diff --git a/Engine/Plugins/Runtime/MeshModelingToolset/Source/MeshModelingTools/Public/EditMeshPolygonsTool.h b/Engine/Plugins/Runtime/MeshModelingToolset/Source/MeshModelingTools/Public/EditMeshPolygonsTool.h
--- a/Engine/Plugins/Runtime/MeshModelingToolset/Source/MeshModelingTools/Public/EditMeshPolygonsTool.h
+++ b/Engine/Plugins/Runtime/MeshModelingToolset/Source/MeshModelingTools/Public/EditMeshPolygonsTool.h
@@ -482,6 +482,11 @@ protected:
 	// This is a fundamental mode switch, must be set before ::Setup() is called!
 	bool bTriangleMode;		
 
+	// TODO: This is a hack to allow us to disallow any actions inside the tool after Setup() is called. We
+	// use it if the user tries to run the tool on a mesh that has too many edges for us to render, to avoid
+	// hanging the editor.
+	bool bToolDisabled = false;
+
 	TObjectPtr<UWorld> TargetWorld = nullptr;
 
 	UPROPERTY()
