ApproximateActors: explicitly release any allocated MeshDescriptions on destruction of TStaticMeshSpatialWrapper in FMeshSceneAdapter. This reduces memory usage peaks in Approximate Actors generation (Merge Actors and HLOD).
#rnx

diff --git a/Engine/Plugins/Runtime/MeshModelingToolset/Source/ModelingComponents/Private/Scene/MeshSceneAdapter.cpp b/Engine/Plugins/Runtime/MeshModelingToolset/Source/ModelingComponents/Private/Scene/MeshSceneAdapter.cpp
--- a/Engine/Plugins/Runtime/MeshModelingToolset/Source/ModelingComponents/Private/Scene/MeshSceneAdapter.cpp
+++ b/Engine/Plugins/Runtime/MeshModelingToolset/Source/ModelingComponents/Private/Scene/MeshSceneAdapter.cpp
@@ -632,6 +632,14 @@ public:
 	TUniquePtr<TMeshAABBTree3<TriangleMeshType>> AABBTree;
 	TUniquePtr<TFastWindingTree<TriangleMeshType>> FWNTree;
 
+	virtual ~TStaticMeshSpatialWrapper()
+	{
+#if WITH_EDITOR
+		// Release any MeshDescriptions that were loaded. This is potentially not ideal and should probably be controlled by the build options...
+		StaticMesh->ClearMeshDescriptions();
+#endif
+	}
+
 	virtual bool Build(const FMeshSceneAdapterBuildOptions& BuildOptions) override
 	{
 		check(StaticMesh);
