update FTransform3d to FTransformSRT3d

diff --git a/Engine/Plugins/Runtime/MeshModelingToolset/Source/ModelingComponents/Private/Scene/SceneGeometrySpatialCache.cpp b/Engine/Plugins/Runtime/MeshModelingToolset/Source/ModelingComponents/Private/Scene/SceneGeometrySpatialCache.cpp
--- a/Engine/Plugins/Runtime/MeshModelingToolset/Source/ModelingComponents/Private/Scene/SceneGeometrySpatialCache.cpp
+++ b/Engine/Plugins/Runtime/MeshModelingToolset/Source/ModelingComponents/Private/Scene/SceneGeometrySpatialCache.cpp
@@ -6,6 +6,7 @@
 #include "GameFramework/Volume.h"
 #include "Components/BrushComponent.h"
 #include "ConversionUtils/VolumeToDynamicMesh.h"
+#include "TransformTypes.h"
 
 using namespace UE::Geometry;
 
