ApproximateActors: reduce memory footprint by using compressed or lower-resolution LOD0 meshes
- add FStaticMeshLODResourcesAdapter in MeshConversions module, wrapper that presents StaticMesh LODResources mesh (ie Section buffers) with an API compatible with GeometryCore mesh templates
- add TMeshWrapperAdapterd in GeometryCore, this is a shim for getting any template-API-compatible mesh into a FTriangleMeshAdapterd specifically (required to call some non-template mesh processing functions)

- add support for using StaticMesh LODResources meshes in FMeshSceneAdapter instead of source meshes, controlled by build options flag
- add FMeshSceneAdapter build options flags specifying whether UV/Normal queries are required (default true)
  - when not required, the loaded MeshDescriptions and created DynamicMeshes are packed into FColliderMesh which has a smaller memory footprint, allowing the larger meshes to be freed

- to support above, the SpatialWrappers in MeshSceneAdapter.cpp have been refactored extensively
  - added FCompressedMeshSpatialWrapper which Builds from a temporary FDynamicMesh3 into a FColliderMesh, supports all the same options as FDynamicMeshSpatialWrapper (so a drop-in replacement with less memory usage, but no UV/Normal queries)
  - added FBaseMeshSpatialWrapper, base class for existing FDynamicMeshSpatialWrapper and new FCompressedMeshSpatialWrapper that has shared config settings (can more more here in the future)
  - added FStaticMeshLODResourcesMeshSurfaceAdapter, similar toFMeshDescriptionTriangleMeshSurfaceAdapter, a mesh adapter that filters out any geo that doesn't have a Surface-domain material
  - refactored TStaticMeshSpatialWrapper into TStaticMeshSpatialWrapperBase and subclasses FStaticMeshSourceDataSpatialWrapper and FStaticMeshRenderDataSpatialWrapper, for source mesh vs render mesh
  - added FCompressedStaticMeshSpatialWrapper, variant of FStaticMeshSourceDataSpatialWrapper that stores to a FColliderMesh and releases the source MeshDescription
  - SpatialWrapperFactory() now has logic to build from render mesh vs source mesh vs compressed source mesh
  - moved initial SpatialWrapper construction from AddActors() phase to Build() phase, because now it depends on build settings

- add EMeshDataSourceLODPolicy to IGeometryProcessing::FOptions, with LOD0 Source and LOD0 RenderMesh options
- ApproximateActorsImpl now provides/handles this policy setting, and configures FMeshSceneAdapter to not need UV/Normal queries (not required in this context)
- add bUseRenderLODMeshes option in FMeshApproximationSettings to expose render vs source mesh control at user level

#rnx

diff --git a/Engine/Plugins/Experimental/MeshModelingToolsetExp/Source/GeometryProcessingAdapters/Private/ApproximateActorsImpl.cpp b/Engine/Plugins/Experimental/MeshModelingToolsetExp/Source/GeometryProcessingAdapters/Private/ApproximateActorsImpl.cpp
--- a/Engine/Plugins/Experimental/MeshModelingToolsetExp/Source/GeometryProcessingAdapters/Private/ApproximateActorsImpl.cpp
+++ b/Engine/Plugins/Experimental/MeshModelingToolsetExp/Source/GeometryProcessingAdapters/Private/ApproximateActorsImpl.cpp
@@ -423,7 +423,6 @@ static TSharedPtr<FApproximationMeshData> GenerateApproximationMesh(
 	}
 
 	Progress.EnterProgressFrame(1.f, LOCTEXT("SolidifyMesh", "Approximating Mesh..."));
-
 	TRACE_BOOKMARK(TEXT("ApproximateActors-Solidify"));
 
 	FDynamicMesh3 SolidMesh;
@@ -494,7 +493,6 @@ static TSharedPtr<FApproximationMeshData> GenerateApproximationMesh(
 	}
 
 	Progress.EnterProgressFrame(1.f, LOCTEXT("SimplifyingMesh", "Simplifying Mesh..."));
-
 	TRACE_BOOKMARK(TEXT("ApproximateActors-Simplify"));
 
 	FVolPresMeshSimplification Simplifier(CurResultMesh);
@@ -522,7 +520,7 @@ static TSharedPtr<FApproximationMeshData> GenerateApproximationMesh(
 			// (this does not seem to help perf and probably makes the results slightly worse)
 			//{
 			//	TRACE_CPUPROFILER_EVENT_SCOPE(ApproximateActorsImpl_Generate_Simplification_PrePass);
-			//	Simplifier.FastCollapsePass(0.1 * UseTargetTolerance, 5);
+			//	Simplifier.FastCollapsePass(ApproxAccuracy, 5);
 			//}
 
 			// now simplify down to a reasonable tri count, as geometric metric is (relatively) expensive
@@ -862,6 +860,10 @@ IGeometryProcessing_ApproximateActors::FOptions FApproximateActorsImpl::Construc
 	Options.BasePolicy = (UseSettings.OutputType == EMeshApproximationType::MeshShapeOnly) ?
 		IGeometryProcessing_ApproximateActors::EApproximationPolicy::CollisionMesh :
 		IGeometryProcessing_ApproximateActors::EApproximationPolicy::MeshAndGeneratedMaterial;
+	
+	Options.MeshDataLODPolicy = (UseSettings.bUseRenderLODMeshes) ?
+		EMeshDataSourceLODPolicy::LOD0RenderMeshes : EMeshDataSourceLODPolicy::LOD0SourceMeshes;
+	
 	Options.WorldSpaceApproximationAccuracyMeters = UseSettings.ApproximationAccuracy;
 
 	Options.bAutoThickenThinParts = UseSettings.bAttemptAutoThickening;
@@ -1015,12 +1017,14 @@ void FApproximateActorsImpl::GenerateApproximationForActorSet(const TArray<AActo
 
 	TUniquePtr<FMeshSceneAdapter> Scene = MakeUnique<FMeshSceneAdapter>();
 	FMeshSceneAdapterBuildOptions SceneBuildOptions;
+	SceneBuildOptions.bIgnoreStaticMeshSourceData = (Options.MeshDataLODPolicy == EMeshDataSourceLODPolicy::LOD0RenderMeshes);		// default is false
 	SceneBuildOptions.bThickenThinMeshes = Options.bAutoThickenThinParts;
 	SceneBuildOptions.DesiredMinThickness = Options.AutoThickenThicknessMeters * 100.0;		// convert to cm (UE Units)
 	// filter out objects smaller than 10% of voxel size
 	SceneBuildOptions.bFilterTinyObjects = Options.bIgnoreTinyParts;
 	SceneBuildOptions.TinyObjectBoxMaxDimension = Options.TinyPartMaxDimensionMeters;
-	SceneBuildOptions.bOnlySurfaceMaterials = true;
+	SceneBuildOptions.bOnlySurfaceMaterials = true;				// don't include decal geometry in 3D mesh scene (will be included in renderings)
+	SceneBuildOptions.bEnableUVQueries = SceneBuildOptions.bEnableNormalsQueries = false;		// not required in this context, will reduce memory usage
 	SceneBuildOptions.bPrintDebugMessages = Options.bVerbose;
 	{
 		TRACE_CPUPROFILER_EVENT_SCOPE(ApproximateActorsImpl_Generate_BuildScene);
diff --git a/Engine/Plugins/Runtime/MeshModelingToolset/Source/ModelingComponents/Private/Scene/MeshSceneAdapter.cpp b/Engine/Plugins/Runtime/MeshModelingToolset/Source/ModelingComponents/Private/Scene/MeshSceneAdapter.cpp
--- a/Engine/Plugins/Runtime/MeshModelingToolset/Source/ModelingComponents/Private/Scene/MeshSceneAdapter.cpp
+++ b/Engine/Plugins/Runtime/MeshModelingToolset/Source/ModelingComponents/Private/Scene/MeshSceneAdapter.cpp
@@ -4,9 +4,11 @@
 #include "Scene/MeshSceneAdapter.h"
 
 #include "DynamicMesh/DynamicMesh3.h"
+#include "DynamicMesh/ColliderMesh.h"
 #include "Spatial/MeshAABBTree3.h"
 #include "Spatial/FastWinding.h"
 #include "MeshDescriptionAdapter.h"
+#include "StaticMeshLODResourcesAdapter.h"
 #include "MeshDescriptionToDynamicMesh.h"
 #include "DynamicMesh/MeshTransforms.h"
 #include "BoxTypes.h"
@@ -240,12 +242,12 @@ static bool IsClosedRegion(const FDynamicMesh3& Mesh, const TArray<int32>& Trian
 
 
 
-
-class FDynamicMeshSpatialWrapper : public IMeshSpatialWrapper
+/**
+ * Base type for various spatial wrappers, to provide common config variables
+ */
+class FBaseMeshSpatialWrapper : public IMeshSpatialWrapper
 {
 public:
-	FDynamicMesh3 Mesh;
-
 	// if true, Mesh is in world space (whatever that means)
 	bool bHasBakedTransform = false;
 	// if true, Mesh is only translated and rotated (allows some assumptions to be made)
@@ -259,9 +261,30 @@ public:
 	// query would only create a shell instead of a full solid
 	bool bRequiresWindingQueryFallback = false;
 
+};
+
+
+
+/**
+ * IMeshSpatialWrapper for a FDynamicMesh3, with optional AABBTree and FWNTree.
+ * 
+ */
+class FDynamicMeshSpatialWrapper : public FBaseMeshSpatialWrapper
+{
+public:
+	FDynamicMesh3 Mesh;
 	TUniquePtr<TMeshAABBTree3<FDynamicMesh3>> AABBTree;
 	TUniquePtr<TFastWindingTree<FDynamicMesh3>> FWNTree;
 
+	FDynamicMeshSpatialWrapper()
+	{
+	}
+
+	FDynamicMeshSpatialWrapper(FDynamicMesh3&& SourceMesh)
+	{
+		Mesh = MoveTemp(SourceMesh);
+	}
+
 	virtual bool Build(const FMeshSceneAdapterBuildOptions& BuildOptions) override
 	{
 		ensure(Mesh.TriangleCount() > 0);
@@ -491,6 +514,255 @@ public:
 
 
 
+
+
+/**
+* IMeshSpatialWrapper that stores a minimal ("compressed") version of an input mesh,
+* basically just a vertex array and indexed triangle list. This reduces the memory
+* footprint of large MeshDescription/FDynamicMeshes by dropping expensive attribute 
+* overlays and connectivity storage. In addition, for very small meshes, FDynamicMesh3
+* is not an efficient representation because each internal buffer has a minimum size,
+* so using this format can significantly reduce the overhead of all those tiny meshes.
+* 
+* Generally supports the same options and behavior as FDynamicMeshSpatialWrapper. 
+* It might be useful to try and combine the various queries in a template base class?
+* 
+* Note that this wrapper currently does not support the barycentric-interpolation queries 
+* of IMeshSpatialWrapper, and cannot be used if they are requested in the build options.
+*/
+class FCompressedMeshSpatialWrapper : public FBaseMeshSpatialWrapper
+{
+public:
+	FDynamicMesh3 BuildInputMesh;		// source mesh data - only used in build, then discarded
+
+	// compressed representation
+	FColliderMesh ColliderMesh;
+	TUniquePtr<TFastWindingTree<FColliderMesh>> FWNTree;
+
+	FCompressedMeshSpatialWrapper()
+	{
+	}
+
+	FCompressedMeshSpatialWrapper(FDynamicMesh3&& SourceMesh)
+	{
+		BuildInputMesh = MoveTemp(SourceMesh);
+	}
+
+
+	virtual bool Build(const FMeshSceneAdapterBuildOptions& BuildOptions) override
+	{
+		ensure(BuildOptions.bEnableUVQueries == false && BuildOptions.bEnableNormalsQueries == false);
+		ensure(BuildInputMesh.TriangleCount() > 0);
+
+		FColliderMesh::FBuildOptions ColliderBuildOptions;
+		ColliderBuildOptions.bBuildAABBTree = BuildOptions.bBuildSpatialDataStructures;
+		ColliderMesh.Initialize(BuildInputMesh);
+		BuildInputMesh = FDynamicMesh3();		// discard input mesh data
+
+		if (BuildOptions.bBuildSpatialDataStructures)
+		{
+			if (bUseDistanceShellForWinding == false || bRequiresWindingQueryFallback)
+			{
+				TRACE_CPUPROFILER_EVENT_SCOPE(MeshScene_WrapperBuild_DMesh_FWNTree);
+				FWNTree = MakeUnique<TFastWindingTree<FColliderMesh>>(&ColliderMesh.GetAABBTree(), true);
+			}
+		}
+		return true;
+	}
+
+	virtual int32 GetTriangleCount() const override
+	{
+		return ColliderMesh.TriangleCount();
+	}
+
+	virtual bool IsTriangle(int32 TriId) const override
+	{
+		return ColliderMesh.IsTriangle(TriId);
+	}
+
+	virtual FIndex3i GetTriangle(int32 TriId) const override
+	{
+		return ColliderMesh.GetTriangle(TriId);
+	}
+
+	virtual bool HasNormals() const override
+	{
+		check(false);
+		return false;
+	}
+	virtual bool HasUVs(const int UVLayer = 0) const override
+	{
+		check(false);
+		return false;
+	}
+
+	virtual FVector3d TriBaryInterpolatePoint(int32 TriId, const FVector3d& BaryCoords) const override
+	{
+		check(false);
+		return FVector3d::Zero();
+	}
+	virtual bool TriBaryInterpolateNormal(int32 TriId, const FVector3d& BaryCoords, FVector3f& NormalOut) const override
+	{
+		check(false);
+		return false;
+	}
+	virtual bool TriBaryInterpolateUV(const int32 TriId, const FVector3d& BaryCoords, const int UVLayer, FVector2f& UVOut) const override
+	{
+		check(false);
+		return false;
+	}
+
+	virtual FAxisAlignedBox3d GetWorldBounds(TFunctionRef<FVector3d(const FVector3d&)> LocalToWorldFunc) override
+	{
+		FAxisAlignedBox3d Bounds = bHasBakedTransform ?
+			GetTransformedVertexBounds<FColliderMesh>(ColliderMesh, [](const FVector3d& P) { return P;}) :
+			GetTransformedVertexBounds<FColliderMesh>(ColliderMesh, LocalToWorldFunc);
+		if (bUseDistanceShellForWinding)
+		{
+			Bounds.Expand(WindingShellThickness);
+		}
+		return Bounds;
+	}
+
+	virtual void CollectSeedPoints(TArray<FVector3d>& WorldPoints, TFunctionRef<FVector3d(const FVector3d&)> LocalToWorldFunc) override
+	{
+		return bHasBakedTransform ?
+			CollectSeedPointsFromMeshVertices<FColliderMesh>(ColliderMesh, [](const FVector3d& P) {return P;}, WorldPoints) :
+			CollectSeedPointsFromMeshVertices<FColliderMesh>(ColliderMesh, LocalToWorldFunc, WorldPoints);
+	}
+
+	virtual double FastWindingNumber(const FVector3d& P, const FTransformSequence3d& LocalToWorldTransform) override
+	{
+		if (bUseDistanceShellForWinding)
+		{
+			if (bRequiresWindingQueryFallback)
+			{
+				double WindingNumber = bHasBakedTransform ?
+					FWNTree->FastWindingNumber(P) :
+					FWNTree->FastWindingNumber(LocalToWorldTransform.InverseTransformPosition(P));
+				if (WindingNumber > 0.99)
+				{
+					return 1.0;
+				}
+			}
+
+			if (bHasBakedTransform || bHasBakedScale)
+			{
+				FVector3d UseP = (bHasBakedTransform) ? P : LocalToWorldTransform.InverseTransformPosition(P);
+				double NearestDistSqr;
+				int32 NearTriID = ColliderMesh.GetAABBTree().FindNearestTriangle(UseP, NearestDistSqr, IMeshSpatial::FQueryOptions(WindingShellThickness));
+				if (NearTriID != IndexConstants::InvalidID)
+				{
+					// Do we even need to do this? won't we return InvalidID if we don't find point within distance?
+					// (also technically we can early-out as soon as we find any point, not the nearest point - might be worth a custom query)
+					FDistPoint3Triangle3d Query = TMeshQueries<FColliderMesh>::TriangleDistance(ColliderMesh, NearTriID, UseP);
+					if (Query.GetSquared() < WindingShellThickness * WindingShellThickness)
+					{
+						return 1.0;
+					}
+				}
+			}
+			else
+			{
+				ensure(false);		// not supported yet
+			}
+			return 0.0;
+		}
+		else
+		{
+			return bHasBakedTransform ?
+				FWNTree->FastWindingNumber(P) :
+				FWNTree->FastWindingNumber(LocalToWorldTransform.InverseTransformPosition(P));
+		}
+	}
+
+
+	virtual bool RayIntersection(const FRay3d& WorldRay, const FTransformSequence3d& LocalToWorldTransform, FMeshSceneRayHit& WorldHitResultOut) override
+	{
+		FRay3d UseRay = WorldRay;
+		if (bHasBakedTransform == false)
+		{
+			FVector3d LocalOrigin = LocalToWorldTransform.InverseTransformPosition(WorldRay.Origin);
+			FVector3d LocalDir = Normalized(LocalToWorldTransform.InverseTransformPosition(WorldRay.PointAt(1.0)) - LocalOrigin);
+			UseRay = FRay3d(LocalOrigin, LocalDir);
+		}
+
+		double RayHitT; int32 HitTID; FVector3d HitBaryCoords;
+		if (ColliderMesh.GetAABBTree().FindNearestHitTriangle(UseRay, RayHitT, HitTID, HitBaryCoords))
+		{
+			WorldHitResultOut.HitMeshTriIndex = HitTID;
+			WorldHitResultOut.HitMeshSpatialWrapper = this;
+			if (bHasBakedTransform)
+			{
+				WorldHitResultOut.RayDistance = RayHitT;
+			}
+			else
+			{
+				FVector3d WorldPos = LocalToWorldTransform.TransformPosition(UseRay.PointAt(RayHitT));
+				WorldHitResultOut.RayDistance = WorldRay.GetParameter(WorldPos);
+			}
+			WorldHitResultOut.HitMeshBaryCoords = HitBaryCoords;
+			return true;
+		}
+		return false;
+	}
+
+
+	virtual bool ProcessVerticesInWorld(TFunctionRef<bool(const FVector3d&)> ProcessFunc, const FTransformSequence3d& LocalToWorldTransform) override
+	{
+		bool bContinue = true;
+		if (bHasBakedTransform)
+		{
+			for ( int32 k = 0; k < ColliderMesh.VertexCount(); ++k)
+			{
+				bContinue = ProcessFunc(ColliderMesh.GetVertex(k));
+				if (!bContinue)
+				{
+					break;
+				}
+			}
+		}
+		else
+		{
+			for ( int32 k = 0; k < ColliderMesh.VertexCount(); ++k )
+			{
+				bContinue = ProcessFunc( LocalToWorldTransform.TransformPosition(ColliderMesh.GetVertex(k)) );
+				if (!bContinue)
+				{
+					break;
+				}
+			}
+		}
+		return bContinue;
+	}
+
+
+	virtual void AppendMesh(FDynamicMesh3& AppendTo, const FTransformSequence3d& TransformSeq) override
+	{
+		// this is necessary because the Editor.AppendMesh function is not a template and explicitly takes a FTriangleMeshAdapterd
+		TMeshWrapperAdapterd<FColliderMesh> AdapterWrapper(&ColliderMesh);
+
+		FDynamicMeshEditor Editor(&AppendTo);
+		FMeshIndexMappings Mappings;
+		if (bHasBakedTransform)
+		{
+			Editor.AppendMesh(&AdapterWrapper, Mappings,
+				[](int, const FVector3d& Pos) { return Pos; });
+		}
+		else
+		{
+			Editor.AppendMesh(&AdapterWrapper, Mappings,
+				[&](int, const FVector3d& Pos) { return TransformSeq.TransformPosition(Pos); });
+		}
+	}
+
+};
+
+
+
+
+
+
 /**
  * Utility function for handling static mesh material queries
  */
@@ -598,115 +870,134 @@ private:
 	bool bHasInvalidPolyGroup;
 };
 
-template <class TriangleMeshType>
-struct TMeshDescriptionMeshAdapterd : public UE::Geometry::TTriangleMeshAdapter<double>
+/**
+* Extension of Mesh adapter for StaticMesh LODResources (ie LOD0 rendering mesh) to filter out all mesh sections that 
+* are using materials for which the material domain isn't "Surface". This excludes decals for example.
+*/
+struct FStaticMeshLODResourcesMeshSurfaceAdapter : public FStaticMeshLODResourcesMeshAdapter
 {
-	TriangleMeshType ParentAdapter;
-
-	TMeshDescriptionMeshAdapterd(TriangleMeshType ParentAdapterIn) : ParentAdapter(ParentAdapterIn)
+	FStaticMeshLODResourcesMeshSurfaceAdapter(const FStaticMeshLODResources* MeshIn, UStaticMesh* StaticMeshIn, bool bOnlySurfaceMaterials)
 	{
-		IsTriangle = [this](int index) { return ParentAdapter.IsTriangle(index); };
-		IsVertex = [this](int index) { return ParentAdapter.IsVertex(index); };
-		MaxTriangleID = [this]() { return ParentAdapter.MaxTriangleID(); };
-		MaxVertexID = [this]() { return ParentAdapter.MaxVertexID(); };
-		TriangleCount = [this]() { return ParentAdapter.TriangleCount(); };
-		VertexCount = [this]() { return ParentAdapter.VertexCount(); };
-		GetChangeStamp = [this]() { return ParentAdapter.GetChangeStamp(); };
-		GetTriangle = [this](int32 TriangleID) { return ParentAdapter.GetTriangle(TriangleID); };
-		GetVertex = [this](int32 VertexID) { return ParentAdapter.GetVertex(VertexID); };
+		this->Mesh = MeshIn;
+		this->NumTriangles = 0;
+
+		TriangleOffsetArray.Reserve(MeshIn->Sections.Num() + 1);
+		ValidSections.Reserve(MeshIn->Sections.Num());
+
+		int32 SectionOffset = 0;
+		for (const FStaticMeshSection& Section : MeshIn->Sections)
+		{
+			auto IsValidMaterial = [StaticMeshIn, Section]()
+			{
+				UMaterialInterface* MaterialInterface = StaticMeshIn->GetMaterial(Section.MaterialIndex);
+				const UMaterial* Material = MaterialInterface != nullptr ? MaterialInterface->GetMaterial_Concurrent() : nullptr;
+
+				return Material != nullptr && Material->MaterialDomain == EMaterialDomain::MD_Surface;
+			};
+
+			bool bIsValidMaterial = bOnlySurfaceMaterials ? IsValidMaterial() : true;
+			if (bIsValidMaterial)
+			{
+				TriangleOffsetArray.Add(SectionOffset);
+				NumTriangles += Section.NumTriangles;
+				ValidSections.Add(&Section);
+			}
+
+			SectionOffset += Section.NumTriangles;
+		}
+
+		TriangleOffsetArray.Add(SectionOffset);
 	}
 };
 
-template <class TriangleMeshType>
-class TStaticMeshSpatialWrapper : public IMeshSpatialWrapper
+
+
+/**
+ * TStaticMeshSpatialWrapperBase is an IMeshSpatialWrapper implementation that can be used
+ * with various different mesh adapter types, depending on whether the "source" mesh or built/render
+ * mesh should be used. This is an abstract base class, calling code should use FStaticMeshSourceDataSpatialWrapper 
+ * or FStaticMeshRenderDataSpatialWrapper as needed
+ */
+template <class MeshAdapterType>
+class TStaticMeshSpatialWrapperBase : public IMeshSpatialWrapper
 {
 public:
 	UStaticMesh* StaticMesh = nullptr;
 	int32 LODIndex = 0;
+
 	FVector3d BuildScale = FVector3d::One();
-	
-	FMeshDescription* SourceMesh = nullptr;
+	TUniquePtr<MeshAdapterType> Adapter;
 
-	TUniquePtr<TriangleMeshType> Adapter;
+	TUniquePtr<TMeshAABBTree3<MeshAdapterType>> AABBTree;
+	TUniquePtr<TFastWindingTree<MeshAdapterType>> FWNTree;
 
-	TUniquePtr<TMeshAABBTree3<TriangleMeshType>> AABBTree;
-	TUniquePtr<TFastWindingTree<TriangleMeshType>> FWNTree;
-
-	virtual ~TStaticMeshSpatialWrapper()
-	{
-#if WITH_EDITOR
-		// Release any MeshDescriptions that were loaded. This is potentially not ideal and should probably be controlled by the build options...
-		StaticMesh->ClearMeshDescriptions();
-#endif
-	}
+	/**
+	 * Create the MeshAdapterType instance based on the build options and source mesh.
+	 * Subclasses must implement this function.
+	 */
+	virtual TUniquePtr<MeshAdapterType> MakeMeshAdapter(const FMeshSceneAdapterBuildOptions& BuildOptions) = 0;
 
 	virtual bool Build(const FMeshSceneAdapterBuildOptions& BuildOptions) override
 	{
 		check(StaticMesh);
 
-#if !WITH_EDITOR
-		SourceMesh = nullptr;
-#else
-		SourceMesh = StaticMesh->GetMeshDescription(LODIndex);
-#endif
-		Adapter = MakeUnique<FMeshDescriptionTriangleMeshSurfaceAdapter>(SourceMesh, StaticMesh, BuildOptions.bOnlySurfaceMaterials);
-
-		if (SourceMesh && Adapter)
+		Adapter = MakeMeshAdapter(BuildOptions);
+		if (!Adapter)
 		{
-			BuildScale = FVector3d::One();
+			return false;
+		}
+
+		BuildScale = FVector3d::One();
 #if WITH_EDITOR
-			// respect BuildScale build setting
-			const FMeshBuildSettings& LODBuildSettings = StaticMesh->GetSourceModel(LODIndex).BuildSettings;
-			BuildScale = (FVector3d)LODBuildSettings.BuildScale3D;
-			Adapter->SetBuildScale(BuildScale, false);
+		// respect BuildScale build setting
+		const FMeshBuildSettings& LODBuildSettings = StaticMesh->GetSourceModel(LODIndex).BuildSettings;
+		BuildScale = (FVector3d)LODBuildSettings.BuildScale3D;
+		Adapter->SetBuildScale(BuildScale, false);
 #endif
-			if (BuildOptions.bBuildSpatialDataStructures)
+
+		if (BuildOptions.bBuildSpatialDataStructures)
+		{
 			{
-				{
-					TRACE_CPUPROFILER_EVENT_SCOPE(MeshScene_WrapperBuild_StaticMesh_AABBTree);
-					AABBTree = MakeUnique<TMeshAABBTree3<TriangleMeshType>>(Adapter.Get(), true);
-				}
-				{
-					TRACE_CPUPROFILER_EVENT_SCOPE(MeshScene_WrapperBuild_StaticMesh_FWNTree);
-					FWNTree = MakeUnique<TFastWindingTree<TriangleMeshType>>(AABBTree.Get(), true);
-				}
+				TRACE_CPUPROFILER_EVENT_SCOPE(MeshScene_WrapperBuild_StaticMesh_AABBTree);
+				AABBTree = MakeUnique<TMeshAABBTree3<MeshAdapterType>>(Adapter.Get(), true);
+			}
+			{
+				TRACE_CPUPROFILER_EVENT_SCOPE(MeshScene_WrapperBuild_StaticMesh_FWNTree);
+				FWNTree = MakeUnique<TFastWindingTree<MeshAdapterType>>(AABBTree.Get(), true);
 			}
-			return true;
 		}
 
-		SourceMesh = nullptr;
-		return false;
+		return true;
 	}
 
 	virtual int32 GetTriangleCount() const override
 	{
-		if (!SourceMesh) return 0;
-		return Adapter->TriangleCount();
+		return Adapter ? Adapter->TriangleCount() : 0;
 	}
 
 	virtual bool IsTriangle(int32 TriId) const override
 	{
-		return SourceMesh ? Adapter->IsTriangle(TriId) : false;
+		return Adapter ? Adapter->IsTriangle(TriId) : false;
 	}
 
 	virtual FIndex3i GetTriangle(int32 TriId) const override
 	{
-		return SourceMesh ? Adapter->GetTriangle(TriId) : FIndex3i::Invalid();
+		return Adapter ? Adapter->GetTriangle(TriId) : FIndex3i::Invalid();
 	}
 
 	virtual bool HasNormals() const override
 	{
-		return SourceMesh ? Adapter->HasNormals() : false;
+		return Adapter ? Adapter->HasNormals() : false;
 	}
 
 	virtual bool HasUVs(const int UVLayer = 0) const override
 	{
-		return SourceMesh ? Adapter->HasUVs(UVLayer) : false;
+		return Adapter ? Adapter->HasUVs(UVLayer) : false;
 	}
 
 	virtual FVector3d TriBaryInterpolatePoint(int32 TriId, const FVector3d& BaryCoords) const override
 	{
-		if (SourceMesh)
+		if (Adapter)
 		{
 			FVector3d P0, P1, P2;
 			Adapter->GetTriVertices(TriId, P0, P1, P2);
@@ -717,7 +1008,7 @@ public:
 
 	virtual bool TriBaryInterpolateNormal(int32 TriId, const FVector3d& BaryCoords, FVector3f& NormalOut) const override
 	{
-		if (SourceMesh && Adapter->HasNormals())
+		if (Adapter && Adapter->HasNormals())
 		{
 			FIndex3i VertIds = Adapter->GetTriangle(TriId);
 			
@@ -731,7 +1022,7 @@ public:
 
 	virtual bool TriBaryInterpolateUV(const int32 TriId, const FVector3d& BaryCoords, const int UVLayer, FVector2f& UVOut) const override
 	{
-		if (SourceMesh && Adapter->HasUVs(UVLayer))
+		if (Adapter && Adapter->HasUVs(UVLayer))
 		{
 			FVector2f UV0, UV1, UV2;
 			Adapter->GetTriUVs(TriId, UVLayer, UV0, UV1, UV2);
@@ -743,19 +1034,19 @@ public:
 
 	virtual FAxisAlignedBox3d GetWorldBounds(TFunctionRef<FVector3d(const FVector3d&)> LocalToWorldFunc) override
 	{
-		if (!SourceMesh) return FAxisAlignedBox3d::Empty();
-		return GetTransformedVertexBounds<TriangleMeshType>(*Adapter, LocalToWorldFunc);
+		if (!Adapter) return FAxisAlignedBox3d::Empty();
+		return GetTransformedVertexBounds<MeshAdapterType>(*Adapter, LocalToWorldFunc);
 	}
 
 	virtual void CollectSeedPoints(TArray<FVector3d>& WorldPoints, TFunctionRef<FVector3d(const FVector3d&)> LocalToWorldFunc) override
 	{
-		if (!SourceMesh) return;
-		CollectSeedPointsFromMeshVertices<TriangleMeshType>(*Adapter, LocalToWorldFunc, WorldPoints);
+		if (!Adapter) return;
+		CollectSeedPointsFromMeshVertices<MeshAdapterType>(*Adapter, LocalToWorldFunc, WorldPoints);
 	}
 
 	virtual double FastWindingNumber(const FVector3d& P, const FTransformSequence3d& LocalToWorldTransform) override
 	{
-		return (SourceMesh) ? FWNTree->FastWindingNumber(LocalToWorldTransform.InverseTransformPosition(P)) : 0.0;
+		return (Adapter) ? FWNTree->FastWindingNumber(LocalToWorldTransform.InverseTransformPosition(P)) : 0.0;
 	}
 
 	virtual bool RayIntersection(const FRay3d& WorldRay, const FTransformSequence3d& LocalToWorldTransform, FMeshSceneRayHit& WorldHitResultOut) override
@@ -781,7 +1072,7 @@ public:
 	virtual bool ProcessVerticesInWorld(TFunctionRef<bool(const FVector3d&)> ProcessFunc, const FTransformSequence3d& LocalToWorldTransform) override
 	{
 		bool bContinue = true;
-		int32 NumVertices = (SourceMesh) ? Adapter->VertexCount() : 0;
+		int32 NumVertices = (Adapter) ? Adapter->VertexCount() : 0;
 		for (int32 vi = 0; vi < NumVertices && bContinue; ++vi)
 		{
 			if (Adapter->IsVertex(vi))
@@ -794,28 +1085,30 @@ public:
 
 	virtual void AppendMesh(FDynamicMesh3& AppendTo, const FTransformSequence3d& TransformSeq) override
 	{
-		if (!SourceMesh) return;
-
-#if WITH_EDITOR
-		// Fast path only works on non-filtered meshes
-		// as it relies on the mesh description directly rather than the adapter
-		bool bFilteredMesh = SourceMesh->Triangles().Num() != Adapter->TriangleCount();
-		if (!bFilteredMesh && AppendTo.TriangleCount() == 0 && TransformSeq.Num() == 0)
-		{
-			// this is somewhat faster in profiling
-			FMeshDescription* UseMeshDescription = StaticMesh->GetMeshDescription(LODIndex);
-			FMeshDescriptionToDynamicMesh Converter;
-			Converter.bEnableOutputGroups = false; Converter.bCalculateMaps = false;
-			Converter.bDisableAttributes = true;
-			Converter.Convert(UseMeshDescription, AppendTo);
-			MeshTransforms::Scale(AppendTo, BuildScale, FVector3d::Zero());
-			return;
-		}
-#endif
+		if (!Adapter) return;
+
+		// TODO: need to move this to FStaticMeshSourceDataSpatialWrapper?
+//#if WITH_EDITOR
+//		// Fast path only works on non-filtered meshes
+//		// as it relies on the mesh description directly rather than the adapter
+//		bool bFilteredMesh = SourceMesh->Triangles().Num() != Adapter->TriangleCount();
+//		if (!bFilteredMesh && AppendTo.TriangleCount() == 0 && TransformSeq.Num() == 0)
+//		{
+//			// this is somewhat faster in profiling
+//			FMeshDescription* UseMeshDescription = StaticMesh->GetMeshDescription(LODIndex);
+//			FMeshDescriptionToDynamicMesh Converter;
+//			Converter.bEnableOutputGroups = false; Converter.bCalculateMaps = false;
+//			Converter.bDisableAttributes = true;
+//			Converter.Convert(UseMeshDescription, AppendTo);
+//			MeshTransforms::Scale(AppendTo, BuildScale, FVector3d::Zero());
+//			return;
+//		}
+//#endif
 
 		FDynamicMeshEditor Editor(&AppendTo);
 		FMeshIndexMappings Mappings;
-		TMeshDescriptionMeshAdapterd<TriangleMeshType> AdapterWrapper(*Adapter);
+		// this is necessary because the Editor.AppendMesh function is not a template and explicitly takes a FTriangleMeshAdapterd
+		TMeshWrapperAdapterd<MeshAdapterType> AdapterWrapper(Adapter.Get());
 		Editor.AppendMesh(&AdapterWrapper, Mappings,
 			[&](int, const FVector3d& Pos) { return TransformSeq.TransformPosition(Pos); });
 	}
@@ -824,25 +1117,246 @@ public:
 
 
 
+class FStaticMeshSourceDataSpatialWrapper : public TStaticMeshSpatialWrapperBase<FMeshDescriptionTriangleMeshSurfaceAdapter>
+{
+public:
+	FMeshDescription* CachedSourceMeshDescription = nullptr;
+
+	virtual ~FStaticMeshSourceDataSpatialWrapper()
+	{
+#if WITH_EDITOR
+		// Release any MeshDescriptions that were loaded. This is potentially not ideal and should probably be controlled by the build options...
+		if (CachedSourceMeshDescription != nullptr)
+		{
+			CachedSourceMeshDescription = nullptr;
+			StaticMesh->ClearMeshDescriptions();
+		}
+#endif
+	}
+
+
+	virtual TUniquePtr<FMeshDescriptionTriangleMeshSurfaceAdapter> MakeMeshAdapter(const FMeshSceneAdapterBuildOptions& BuildOptions) override
+	{
+#if WITH_EDITOR
+		CachedSourceMeshDescription = StaticMesh->GetMeshDescription(LODIndex);
+		return MakeUnique<FMeshDescriptionTriangleMeshSurfaceAdapter>(CachedSourceMeshDescription, StaticMesh, BuildOptions.bOnlySurfaceMaterials);
+#else
+		// cannot use this path in non-editor builds, but this should have been handled higher up by instantiating 
+		// the runtime subclass
+		check(false);
+		return nullptr;
+#endif
+	}
+
+};
+
+
+class FStaticMeshRenderDataSpatialWrapper : public TStaticMeshSpatialWrapperBase<FStaticMeshLODResourcesMeshSurfaceAdapter>
+{
+public:
+	virtual TUniquePtr<FStaticMeshLODResourcesMeshSurfaceAdapter> MakeMeshAdapter(const FMeshSceneAdapterBuildOptions& BuildOptions) override
+	{
+		return MakeUnique<FStaticMeshLODResourcesMeshSurfaceAdapter>(&StaticMesh->GetRenderData()->LODResources[0], StaticMesh, BuildOptions.bOnlySurfaceMaterials);
+	}
+};
+
 
-static TUniquePtr<IMeshSpatialWrapper> SpatialWrapperFactory( const FMeshTypeContainer& MeshContainer )
+
+/**
+ * FCompressedStaticMeshSpatialWrapper mirrors the behavior of FStaticMeshSourceDataSpatialWrapper, but instead of
+ * using the underlying MeshDescription directly, it copies the mesh into a minimal indexed mesh (vert list and indexed tris)
+ * and then releases the MeshDescription memory. This is useful to limit the total memory usage of the MeshSceneAdapter
+ * in cases where only certain queries are needed. In particular, the UV and Normal queries
+ * are not available with this MeshSpatialWrapper. 
+ */
+class FCompressedStaticMeshSpatialWrapper : public IMeshSpatialWrapper
 {
+public:
+	UStaticMesh* StaticMesh = nullptr;
+
+	FColliderMesh ColliderMesh;
+	TUniquePtr<TMeshAABBTree3<FColliderMesh>> AABBTree;
+	TUniquePtr<TFastWindingTree<FColliderMesh>> FWNTree;
+
+	virtual bool Build(const FMeshSceneAdapterBuildOptions& BuildOptions) override
+	{
+		// this is because Build() will be called twice
+		if (ColliderMesh.TriangleCount() == 0)
+		{
+			TUniquePtr<FStaticMeshSourceDataSpatialWrapper> TempWrapper = MakeUnique<FStaticMeshSourceDataSpatialWrapper>();
+			TempWrapper->SourceContainer = SourceContainer;
+			TempWrapper->StaticMesh = StaticMesh;
+
+			FMeshSceneAdapterBuildOptions TempBuildOptions;
+			TempBuildOptions.bBuildSpatialDataStructures = false;	// we are just copying from this so do not build spatials
+			TempWrapper->Build(TempBuildOptions);
+
+			FDynamicMesh3 TempMesh(EMeshComponents::None);
+			TempWrapper->AppendMesh(TempMesh, FTransformSequence3d());
+
 #if WITH_EDITOR
-	if (MeshContainer.MeshType == ESceneMeshType::StaticMeshAsset)
+			// dump mesh description memory
+			StaticMesh->ClearMeshDescriptions();
+#endif
+
+			FColliderMesh::FBuildOptions ColliderBuildOptions;
+			ColliderBuildOptions.bBuildAABBTree = false;		// will use our own AABBTree so that we can defer build
+			ColliderMesh.Initialize(TempMesh, ColliderBuildOptions);
+		}
+
+		if (BuildOptions.bBuildSpatialDataStructures)
+		{
+			{
+				TRACE_CPUPROFILER_EVENT_SCOPE(MeshScene_WrapperBuild_StaticMeshCompressed_AABBTree);
+				AABBTree = MakeUnique<TMeshAABBTree3<FColliderMesh>>(&ColliderMesh, true);
+			}
+			{
+				TRACE_CPUPROFILER_EVENT_SCOPE(MeshScene_WrapperBuild_StaticMeshCompressed_FWNTree);
+				FWNTree = MakeUnique<TFastWindingTree<FColliderMesh>>(AABBTree.Get(), true);
+			}
+		}
+
+		return true;
+	}
+
+	virtual int32 GetTriangleCount() const override
 	{
-		TUniquePtr<TStaticMeshSpatialWrapper<FMeshDescriptionTriangleMeshSurfaceAdapter>> SMWrapper = MakeUnique<TStaticMeshSpatialWrapper<FMeshDescriptionTriangleMeshSurfaceAdapter>>();
+		return ColliderMesh.TriangleCount();
+	}
 
-		SMWrapper->SourceContainer = MeshContainer;
-		SMWrapper->StaticMesh = MeshContainer.GetStaticMesh();
-		if (ensure(SMWrapper->StaticMesh != nullptr))
+	virtual bool IsTriangle(int32 TriId) const override
+	{
+		return ColliderMesh.IsTriangle(TriId);
+	}
+
+	virtual FIndex3i GetTriangle(int32 TriId) const override
+	{
+		return ColliderMesh.GetTriangle(TriId);
+	}
+
+	virtual bool HasNormals() const override { return false; }
+	virtual bool HasUVs(const int UVLayer = 0) const override { return false; }
+	virtual FVector3d TriBaryInterpolatePoint(int32 TriId, const FVector3d& BaryCoords) const override { check(false); return FVector3d::Zero(); }
+	virtual bool TriBaryInterpolateNormal(int32 TriId, const FVector3d& BaryCoords, FVector3f& NormalOut) const override { check(false); return false; }
+	virtual bool TriBaryInterpolateUV(const int32 TriId, const FVector3d& BaryCoords, const int UVLayer, FVector2f& UVOut) const override { check(false); return false; }
+
+	virtual FAxisAlignedBox3d GetWorldBounds(TFunctionRef<FVector3d(const FVector3d&)> LocalToWorldFunc) override
+	{
+		return GetTransformedVertexBounds<FColliderMesh>(ColliderMesh, LocalToWorldFunc);
+	}
+
+	virtual void CollectSeedPoints(TArray<FVector3d>& WorldPoints, TFunctionRef<FVector3d(const FVector3d&)> LocalToWorldFunc) override
+	{
+		CollectSeedPointsFromMeshVertices<FColliderMesh>(ColliderMesh, LocalToWorldFunc, WorldPoints);
+	}
+
+	virtual double FastWindingNumber(const FVector3d& P, const FTransformSequence3d& LocalToWorldTransform) override
+	{
+		return FWNTree->FastWindingNumber(LocalToWorldTransform.InverseTransformPosition(P));
+	}
+
+	virtual bool RayIntersection(const FRay3d& WorldRay, const FTransformSequence3d& LocalToWorldTransform, FMeshSceneRayHit& WorldHitResultOut) override
+	{
+		FVector3d LocalOrigin = LocalToWorldTransform.InverseTransformPosition(WorldRay.Origin);
+		FVector3d LocalDir = Normalized(LocalToWorldTransform.InverseTransformPosition(WorldRay.PointAt(1.0)) - LocalOrigin);
+		FRay3d LocalRay(LocalOrigin, LocalDir);
+
+		double LocalHitT; int32 HitTID; FVector3d HitBaryCoords;
+		if (AABBTree->FindNearestHitTriangle(LocalRay, LocalHitT, HitTID, HitBaryCoords))
+		{
+			WorldHitResultOut.HitMeshTriIndex = HitTID;
+			WorldHitResultOut.HitMeshSpatialWrapper = this;
+			FVector3d WorldPos = LocalToWorldTransform.TransformPosition(LocalRay.PointAt(LocalHitT));
+			WorldHitResultOut.RayDistance = WorldRay.GetParameter(WorldPos);
+			WorldHitResultOut.HitMeshBaryCoords = HitBaryCoords;
+			return true;
+		}
+		return false;
+	}
+
+
+	virtual bool ProcessVerticesInWorld(TFunctionRef<bool(const FVector3d&)> ProcessFunc, const FTransformSequence3d& LocalToWorldTransform) override
+	{
+		bool bContinue = true;
+		int32 NumVertices = ColliderMesh.VertexCount();
+		for (int32 vi = 0; vi < NumVertices && bContinue; ++vi)
+		{
+			if (ColliderMesh.IsVertex(vi))
+			{
+				bContinue = ProcessFunc(LocalToWorldTransform.TransformPosition(ColliderMesh.GetVertex(vi)));
+			}
+		}
+		return bContinue;
+	}
+
+	virtual void AppendMesh(FDynamicMesh3& AppendTo, const FTransformSequence3d& TransformSeq) override
+	{
+		TArray<int32> NewVertIDs;
+		NewVertIDs.SetNum(ColliderMesh.VertexCount());
+		for (int32 k = 0; k < ColliderMesh.VertexCount(); ++k)
+		{
+			NewVertIDs[k] = AppendTo.AppendVertex(ColliderMesh.GetVertex(k));
+		}
+		for (int32 k = 0; k < ColliderMesh.TriangleCount(); ++k)
 		{
-			return SMWrapper;
+			FIndex3i Tri = ColliderMesh.GetTriangle(k);
+			AppendTo.AppendTriangle(NewVertIDs[Tri.A], NewVertIDs[Tri.B], NewVertIDs[Tri.C]);
 		}
 	}
+
+};
+
+
+
+
+
+
+
+
+static TUniquePtr<IMeshSpatialWrapper> SpatialWrapperFactory( const FMeshTypeContainer& MeshContainer, const FMeshSceneAdapterBuildOptions& BuildOptions )
+{
+
+	if (MeshContainer.MeshType == ESceneMeshType::StaticMeshAsset)
+	{
+#if WITH_EDITOR
+		bool bUseRenderMeshes = BuildOptions.bIgnoreStaticMeshSourceData;
 #else
-	checkf(false, TEXT("Not currently supported - to build at Runtime it is necessary to read from the StaticMesh RenderBuffers"));
+		bool bUseRenderMeshes = true;
 #endif
 
+		if (bUseRenderMeshes)
+		{
+			TUniquePtr<FStaticMeshRenderDataSpatialWrapper> SMWrapper = MakeUnique<FStaticMeshRenderDataSpatialWrapper>();
+			SMWrapper->SourceContainer = MeshContainer;
+			SMWrapper->StaticMesh = MeshContainer.GetStaticMesh();
+			if (ensure(SMWrapper->StaticMesh != nullptr))
+			{
+				return SMWrapper;
+			}
+		}
+		else if (BuildOptions.bEnableUVQueries == false && BuildOptions.bEnableUVQueries == false)
+		{
+			TUniquePtr<FCompressedStaticMeshSpatialWrapper> SMWrapper = MakeUnique<FCompressedStaticMeshSpatialWrapper>();
+			SMWrapper->SourceContainer = MeshContainer;
+			SMWrapper->StaticMesh = MeshContainer.GetStaticMesh();
+			if (ensure(SMWrapper->StaticMesh != nullptr))
+			{
+				return SMWrapper;
+			}
+		}
+		else
+		{
+			TUniquePtr<FStaticMeshSourceDataSpatialWrapper> SMWrapper = MakeUnique<FStaticMeshSourceDataSpatialWrapper>();
+			SMWrapper->SourceContainer = MeshContainer;
+			SMWrapper->StaticMesh = MeshContainer.GetStaticMesh();
+			if (ensure(SMWrapper->StaticMesh != nullptr))
+			{
+				return SMWrapper;
+			}
+		}
+
+	}
+
 	return TUniquePtr<IMeshSpatialWrapper>();
 }
 
@@ -909,6 +1423,7 @@ static void CollectComponentChildMeshes(UActorComponent* Component, FActorAdapte
 }	// end namespace Geometry
 
 
+
 void FMeshSceneAdapter::AddActors(const TArray<AActor*>& ActorsSetIn)
 {
 	TRACE_CPUPROFILER_EVENT_SCOPE(MeshScene_AddActors);
@@ -985,7 +1500,7 @@ void FMeshSceneAdapter::InitializeSpatialWrappers(const TArray<FActorAdapter*>&
 				SpatialAdapters.Add(MeshKey, NewWrapperInfo);
 
 				NewWrapperInfo->SourceContainer = ChildMesh->MeshContainer;
-				NewWrapperInfo->SpatialWrapper = SpatialWrapperFactory(ChildMesh->MeshContainer);
+				NewWrapperInfo->SpatialWrapper = nullptr;	// these are now initialized at beginning of Build() function
 				MeshInfo = NewWrapperInfo.Get();
 			}
 			else
@@ -1008,6 +1523,12 @@ void FMeshSceneAdapter::Build(const FMeshSceneAdapterBuildOptions& BuildOptions)
 {
 	TRACE_CPUPROFILER_EVENT_SCOPE(MeshScene_Build);
 
+	// Initialize all Spatial Wrappers using the default factory
+	for (TPair<void*, TSharedPtr<FSpatialWrapperInfo>> Pair : SpatialAdapters)
+	{
+		Pair.Value->SpatialWrapper = SpatialWrapperFactory(Pair.Value->SourceContainer, BuildOptions);
+	}
+
 	if (BuildOptions.bThickenThinMeshes)
 	{
 		Build_FullDecompose(BuildOptions);
@@ -1144,6 +1665,8 @@ void FMeshSceneAdapter::Build_FullDecompose(const FMeshSceneAdapterBuildOptions&
 	EParallelForFlags ParallelFlags = (CVarMeshSceneAdapterDisableMultiThreading.GetValueOnAnyThread() != 0) ?
 		EParallelForFlags::ForceSingleThread : EParallelForFlags::Unbalanced;
 
+	bool bCanUseCompressedMeshes = (BuildOptions.bEnableUVQueries == false && BuildOptions.bEnableNormalsQueries == false);
+
 	// initial list of spatial wrappers that need to be built
 	TArray<FSpatialWrapperInfo*> ToBuild;
 	for (TPair<void*, TSharedPtr<FSpatialWrapperInfo>> Pair : SpatialAdapters)
@@ -1432,18 +1955,25 @@ void FMeshSceneAdapter::Build_FullDecompose(const FMeshSceneAdapterBuildOptions&
 		ItemStats.TotalNumUniqueTris += ItemStats.InstancedSubmeshTris;
 
 		// First handle the LocalSpaceParts mesh, which can still be shared between the original FActorChildMesh instances
-		if (LocalSpaceParts.TriangleCount() > 0)
+		int32 LocalSpacePartsTriangleCount = LocalSpaceParts.TriangleCount();
+		if (LocalSpacePartsTriangleCount > 0)
 		{
-			TUniquePtr<FDynamicMeshSpatialWrapper> LocalSpaceMeshWrapper = MakeUnique<FDynamicMeshSpatialWrapper>();
-			LocalSpaceMeshWrapper->Mesh = MoveTemp(LocalSpaceParts);
-			AddBuildJob(LocalSpaceMeshWrapper.Get(), LocalSpaceMeshWrapper->Mesh.TriangleCount());
-			WrapperInfo->SpatialWrapper = MoveTemp(LocalSpaceMeshWrapper);
+			if (bCanUseCompressedMeshes)
+			{
+				WrapperInfo->SpatialWrapper = MakeUnique<FCompressedMeshSpatialWrapper>(MoveTemp(LocalSpaceParts));
+			}
+			else
+			{
+				WrapperInfo->SpatialWrapper = MakeUnique<FDynamicMeshSpatialWrapper>(MoveTemp(LocalSpaceParts));
+			}
+
+			AddBuildJob(WrapperInfo->SpatialWrapper.Get(), LocalSpacePartsTriangleCount);
 			for (FActorChildMesh* MeshInstance : MeshesToDecompose)
 			{
 				MeshInstance->MeshSpatial = WrapperInfo->SpatialWrapper.Get();
 			}
 			InternalListsLock.Lock();
-			InstancedTrisCount += LocalSpaceParts.TriangleCount() * ParentTransforms.Num();
+			InstancedTrisCount += LocalSpacePartsTriangleCount * ParentTransforms.Num();
 			InternalListsLock.Unlock();
 		}
 		else
@@ -1476,29 +2006,41 @@ void FMeshSceneAdapter::Build_FullDecompose(const FMeshSceneAdapterBuildOptions&
 			TUniquePtr<FActorAdapter> NewActor = MakeUnique<FActorAdapter>();
 			NewActor->SourceActor = nullptr;		// not a "real" actor
 
+			int32 SubmeshTriangleCount = Submesh.Mesh->TriangleCount();
+			ensure(SubmeshTriangleCount > 0);
+
 			// make new spatialwrapper for this instanced mesh
-			TUniquePtr<FDynamicMeshSpatialWrapper> NewInstancedMesh = MakeUnique<FDynamicMeshSpatialWrapper>();
-			int32 TriangleCount = Submesh.Mesh->TriangleCount();
-			NewInstancedMesh->Mesh = MoveTemp(*Submesh.Mesh);
-			NewInstancedMesh->bHasBakedTransform = false;
-			NewInstancedMesh->bHasBakedScale = true;
-			// queue up build job
-			AddBuildJob(NewInstancedMesh.Get(), TriangleCount);
+			TUniquePtr<FBaseMeshSpatialWrapper> NewSpatialWrapper;
+			if (bCanUseCompressedMeshes)
+			{
+				NewSpatialWrapper = MakeUnique<FCompressedMeshSpatialWrapper>(MoveTemp(*Submesh.Mesh));
+			}
+			else
+			{
+				NewSpatialWrapper = MakeUnique<FDynamicMeshSpatialWrapper>(MoveTemp(*Submesh.Mesh));
+			}
+
+			// configure flags
+			NewSpatialWrapper->bHasBakedTransform = false;
+			NewSpatialWrapper->bHasBakedScale = true;
 			// if mesh is too thin, configure the extra shell offset based on 'missing' thickness
 			if (Submesh.ComputedThickness < BuildOptions.DesiredMinThickness)
 			{
-				NewInstancedMesh->bUseDistanceShellForWinding = true;
-				NewInstancedMesh->WindingShellThickness = 0.5 * (BuildOptions.DesiredMinThickness - Submesh.ComputedThickness);
-				NewInstancedMesh->bRequiresWindingQueryFallback = (NewInstancedMesh->WindingShellThickness < 0.6*BuildOptions.DesiredMinThickness);
+				NewSpatialWrapper->bUseDistanceShellForWinding = true;
+				NewSpatialWrapper->WindingShellThickness = 0.5 * (BuildOptions.DesiredMinThickness - Submesh.ComputedThickness);
+				NewSpatialWrapper->bRequiresWindingQueryFallback = (NewSpatialWrapper->WindingShellThickness < 0.6*BuildOptions.DesiredMinThickness);
 			}
 
 			TSharedPtr<FSpatialWrapperInfo> NewWrapperInfo = MakeShared<FSpatialWrapperInfo>();
-			NewWrapperInfo->SpatialWrapper = MoveTemp(NewInstancedMesh);
+			NewWrapperInfo->SpatialWrapper = MoveTemp(NewSpatialWrapper);
+
+			// queue up build job
+			AddBuildJob(NewWrapperInfo->SpatialWrapper.Get(), SubmeshTriangleCount);
 
 			// add to internal lists
 			InternalListsLock.Lock();
-			AddedUniqueTrisCount += TriangleCount;
-			InstancedTrisCount += TriangleCount * Submesh.NewTransforms.Num();
+			AddedUniqueTrisCount += SubmeshTriangleCount;
+			InstancedTrisCount += SubmeshTriangleCount * Submesh.NewTransforms.Num();
 			void* UseKey = (void*)NewWrapperInfo->SpatialWrapper.Get();
 			SpatialAdapters.Add(UseKey, NewWrapperInfo);
 			InternalListsLock.Unlock();
@@ -1520,10 +2062,10 @@ void FMeshSceneAdapter::Build_FullDecompose(const FMeshSceneAdapterBuildOptions&
 				NewInstancesCount++;
 			}
 
-			ItemStats.SubmeshSizeCount.Add(FVector2i(TriangleCount, Submesh.NewTransforms.Num()));
+			ItemStats.SubmeshSizeCount.Add(FVector2i(SubmeshTriangleCount, Submesh.NewTransforms.Num()));
 			ItemStats.NewInstances += Submesh.NewTransforms.Num();
-			ItemStats.NewUniqueTris += TriangleCount;
-			ItemStats.TotalNumUniqueTris += TriangleCount;
+			ItemStats.NewUniqueTris += SubmeshTriangleCount;
+			ItemStats.TotalNumUniqueTris += SubmeshTriangleCount;
 
 			// add actor our actor set
 			InternalListsLock.Lock();
diff --git a/Engine/Plugins/Runtime/MeshModelingToolset/Source/ModelingComponents/Public/Scene/MeshSceneAdapter.h b/Engine/Plugins/Runtime/MeshModelingToolset/Source/ModelingComponents/Public/Scene/MeshSceneAdapter.h
--- a/Engine/Plugins/Runtime/MeshModelingToolset/Source/ModelingComponents/Public/Scene/MeshSceneAdapter.h
+++ b/Engine/Plugins/Runtime/MeshModelingToolset/Source/ModelingComponents/Public/Scene/MeshSceneAdapter.h
@@ -81,6 +81,20 @@ struct FMeshSceneAdapterBuildOptions
 
 	/** If true, AABBTree and Fast Winding data structures will be built for unique scene meshes */
 	bool bBuildSpatialDataStructures = true;
+
+	/** If true, UV queries will be supported on mesh wrappers */
+	bool bEnableUVQueries = true;
+
+	/** If true, Normals queries at hit locations will be supported */
+	bool bEnableNormalsQueries = true;
+
+
+	/** 
+	* If true, the source meshes in StaticMesh Assets are ignored in favor of the render meshes.
+	* This can speed up processing and reduce memory usage on scenes with large Nanite-enabled meshes,
+	* but potentially with lower-quality results for operations querying the mesh scene.
+	*/
+	bool bIgnoreStaticMeshSourceData = false;
 };
 
 
diff --git a/Engine/Source/Developer/GeometryProcessingInterfaces/Public/GeometryProcessingInterfaces/ApproximateActors.h b/Engine/Source/Developer/GeometryProcessingInterfaces/Public/GeometryProcessingInterfaces/ApproximateActors.h
--- a/Engine/Source/Developer/GeometryProcessingInterfaces/Public/GeometryProcessingInterfaces/ApproximateActors.h
+++ b/Engine/Source/Developer/GeometryProcessingInterfaces/Public/GeometryProcessingInterfaces/ApproximateActors.h
@@ -32,6 +32,12 @@ class IGeometryProcessing_ApproximateActors : public IModularFeature
 public:
 	virtual ~IGeometryProcessing_ApproximateActors() {}
 
+	enum class EMeshDataSourceLODPolicy
+	{
+		LOD0SourceMeshes,
+		LOD0RenderMeshes
+	};
+
 	enum class EApproximationPolicy
 	{
 		MeshAndGeneratedMaterial,
@@ -97,6 +103,13 @@ public:
 		// high-level control of the overall approximation process
 		EApproximationPolicy BasePolicy = EApproximationPolicy::MeshAndGeneratedMaterial;
 
+		// 
+		// Actor/Scene configuration settings
+		//
+
+		// control which LOD, and which type of LOD mesh, should be used as the source for the approximation process
+		EMeshDataSourceLODPolicy MeshDataLODPolicy = EMeshDataSourceLODPolicy::LOD0SourceMeshes;
+
 		//
 		// Mesh Preprocessing settings
 		//
diff --git a/Engine/Source/Runtime/Engine/Classes/Engine/MeshMerging.h b/Engine/Source/Runtime/Engine/Classes/Engine/MeshMerging.h
--- a/Engine/Source/Runtime/Engine/Classes/Engine/MeshMerging.h
+++ b/Engine/Source/Runtime/Engine/Classes/Engine/MeshMerging.h
@@ -899,17 +899,28 @@ struct FMeshApproximationSettings
 	float NearPlaneDist = 1.0f;
 
 
+	//
+	// Performance Settings
+	//
+
+
+	/** If true, LOD0 Render Meshes (or Nanite Fallback meshes) are used instead of Source Mesh data. This can significantly reduce computation time and memory usage, but potentially at the cost of lower quality output. */
+	UPROPERTY(BlueprintReadWrite, EditAnywhere, Category = PerformanceSettings)
+	bool bUseRenderLODMeshes = false;
+
+
+
 	//
 	// Debug Output Settings
 	//
 
 
 	/** If true, print out debugging messages */
-	UPROPERTY(BlueprintReadWrite, EditAnywhere, AdvancedDisplay, Category = DebugSettings)
+	UPROPERTY(BlueprintReadWrite, EditAnywhere, Category = DebugSettings)
 	bool bPrintDebugMessages = false;
 
 	/** If true, write the full mesh triangle set (ie flattened, non-instanced) used for mesh generation. Warning: this asset may be extremely large!! */
-	UPROPERTY(BlueprintReadWrite, EditAnywhere, AdvancedDisplay, Category = DebugSettings)
+	UPROPERTY(BlueprintReadWrite, EditAnywhere, Category = DebugSettings)
 	bool bEmitFullDebugMesh = false;
 
 
diff --git a/Engine/Source/Runtime/GeometryCore/Public/MeshAdapter.h b/Engine/Source/Runtime/GeometryCore/Public/MeshAdapter.h
--- a/Engine/Source/Runtime/GeometryCore/Public/MeshAdapter.h
+++ b/Engine/Source/Runtime/GeometryCore/Public/MeshAdapter.h
@@ -229,5 +229,32 @@ struct TIndexVectorMeshArrayAdapter
 typedef TIndexMeshArrayAdapter<uint32, double> FIndexMeshArrayAdapterd;
 
 
+/**
+ * TMeshWrapperAdapterd<T> can be used to present an arbitrary Mesh / Adapter type as a FTriangleMeshAdapterd.
+ * This is useful in cases where it would be difficult or undesirable to write code templated on
+ * the standard "Mesh Type" signature. If the code is written for FTriangleMeshAdapterd then this
+ * shim can be used to present any compatible mesh type as a FTriangleMeshAdapterd
+ */ 
+template <class WrappedMeshType>
+struct TMeshWrapperAdapterd : public UE::Geometry::FTriangleMeshAdapterd
+{
+	WrappedMeshType* WrappedAdapter;
+
+	TMeshWrapperAdapterd(WrappedMeshType* WrappedAdapterIn) : WrappedAdapter(WrappedAdapterIn)
+	{
+		IsTriangle = [this](int index) { return WrappedAdapter->IsTriangle(index); };
+		IsVertex = [this](int index) { return WrappedAdapter->IsVertex(index); };
+		MaxTriangleID = [this]() { return WrappedAdapter->MaxTriangleID(); };
+		MaxVertexID = [this]() { return WrappedAdapter->MaxVertexID(); };
+		TriangleCount = [this]() { return WrappedAdapter->TriangleCount(); };
+		VertexCount = [this]() { return WrappedAdapter->VertexCount(); };
+		GetChangeStamp = [this]() { return WrappedAdapter->GetChangeStamp(); };
+		GetTriangle = [this](int32 TriangleID) { return WrappedAdapter->GetTriangle(TriangleID); };
+		GetVertex = [this](int32 VertexID) { return WrappedAdapter->GetVertex(VertexID); };
+	}
+};
+
+
+
 } // end namespace UE::Geometry
 } // end namespace UE
\ No newline at end of file
diff --git a/Engine/Source/Runtime/MeshConversion/Public/StaticMeshLODResourcesAdapter.h b/Engine/Source/Runtime/MeshConversion/Public/StaticMeshLODResourcesAdapter.h
new file mode 100644
--- /dev/null
+++ b/Engine/Source/Runtime/MeshConversion/Public/StaticMeshLODResourcesAdapter.h
@@ -0,0 +1,199 @@
+// Copyright Epic Games, Inc. All Rights Reserved.
+
+#pragma once
+
+#include "CoreMinimal.h"
+#include "StaticMeshAttributes.h"
+#include "StaticMeshResources.h"
+
+
+/**
+ * Basic struct to adapt a FStaticMeshLODResources for use by GeometryProcessing classes that template the mesh type and expect a standard set of basic accessors
+ * For example, this adapter will let you use a FStaticMeshLODResources with GeometryProcessing's TMeshAABBTree3
+ */
+struct FStaticMeshLODResourcesMeshAdapter
+{
+	using FIndex3i = UE::Geometry::FIndex3i;
+protected:
+	const FStaticMeshLODResources* Mesh;
+
+	FVector3d BuildScale = FVector3d::One();
+	bool bScaleNormals = false;
+
+	TArray<const FStaticMeshSection*> ValidSections;
+	TArray<int32> TriangleOffsetArray;
+	int32 NumTriangles;
+
+protected:
+	FStaticMeshLODResourcesMeshAdapter()
+		: Mesh(nullptr)
+		, NumTriangles(0)
+	{
+	}
+
+public:
+	FStaticMeshLODResourcesMeshAdapter(const FStaticMeshLODResources* MeshIn)
+		: Mesh(MeshIn)
+		, NumTriangles(0)
+	{
+		TriangleOffsetArray.Reserve(MeshIn->Sections.Num() + 1);
+		ValidSections.Reserve(MeshIn->Sections.Num());
+
+		for (const FStaticMeshSection& Section : MeshIn->Sections)
+		{
+			TriangleOffsetArray.Add(NumTriangles);
+			NumTriangles += Section.NumTriangles;
+			ValidSections.Add(&Section);
+		}
+
+		TriangleOffsetArray.Add(NumTriangles);
+	}
+
+	void SetBuildScale(const FVector3d& BuildScaleIn, bool bScaleNormalsIn)
+	{
+		BuildScale = BuildScaleIn;
+		bScaleNormals = bScaleNormalsIn;
+	}
+
+	bool IsTriangle(int32 TID) const
+	{
+		return TID >= 0 && TID < MaxTriangleID();
+	}
+	bool IsVertex(int32 VID) const
+	{
+		return VID >= 0 && VID < MaxVertexID();
+	}
+	// ID and Count are the same for StaticMeshLODResources because it's compact
+	int32 MaxTriangleID() const
+	{
+		return TriangleCount();
+	}
+	int32 TriangleCount() const
+	{
+		return Mesh && Mesh->IndexBuffer.GetAllowCPUAccess() ? NumTriangles : 0;
+	}
+	int32 MaxVertexID() const
+	{
+		return VertexCount();
+	}
+	int32 VertexCount() const
+	{
+		return Mesh && Mesh->VertexBuffers.PositionVertexBuffer.GetAllowCPUAccess() ? Mesh->GetNumVertices() : 0;
+	}
+	uint64 GetChangeStamp() const
+	{
+		// FStaticMeshLODResource doesn't provide any mechanism to know if it's been modified so just return 1
+		// and leave it to the caller to not build an aabb and then change the underlying mesh
+		return 1;
+	}
+
+	inline const FStaticMeshSection& TriangleToSection(int32& InOutIDValue) const
+	{
+		int32 SectionIdx = Algo::UpperBound(TriangleOffsetArray, InOutIDValue) - 1;
+		InOutIDValue -= TriangleOffsetArray[SectionIdx];
+		return *ValidSections[SectionIdx];
+	}
+
+	inline FIndex3i GetTriangle(int32 IDValue) const
+	{
+		const FStaticMeshSection& TriSection = TriangleToSection(IDValue);
+		const FIndexArrayView IndexArrayView = Mesh->IndexBuffer.GetArrayView();
+		return FIndex3i(IndexArrayView[TriSection.FirstIndex + IDValue * 3 + 0],
+			IndexArrayView[TriSection.FirstIndex + IDValue * 3 + 1],
+			IndexArrayView[TriSection.FirstIndex + IDValue * 3 + 2]);
+	}
+
+	inline FVector3d GetVertex(int32 IDValue) const
+	{
+		const FVector3f& Position = Mesh->VertexBuffers.PositionVertexBuffer.VertexPosition(IDValue);
+		return FVector3d(BuildScale.X * (double)Position.X, BuildScale.Y * (double)Position.Y, BuildScale.Z * (double)Position.Z);
+	}
+
+	inline void GetTriVertices(int32 IDValue, FVector3d& V0, FVector3d& V1, FVector3d& V2) const
+	{
+		FIndex3i VtxIndices = GetTriangle(IDValue);
+
+		V0 = GetVertex(VtxIndices.A);
+		V1 = GetVertex(VtxIndices.B);
+		V2 = GetVertex(VtxIndices.C);
+	}
+
+	template<typename VectorType>
+	inline void GetTriVertices(int32 IDValue, VectorType& V0, VectorType& V1, VectorType& V2) const
+	{
+		FIndex3i VtxIndices = GetTriangle(IDValue);
+
+		V0 = GetVertex(VtxIndices.A);
+		V1 = GetVertex(VtxIndices.B);
+		V2 = GetVertex(VtxIndices.C);
+	}
+
+	inline bool HasNormals() const
+	{
+		return Mesh && Mesh->VertexBuffers.StaticMeshVertexBuffer.GetAllowCPUAccess() && Mesh->VertexBuffers.StaticMeshVertexBuffer.GetTangentData();
+	}
+	inline bool IsNormal(int32 NID) const
+	{
+		return NID >= 0 && NID < NormalCount();
+	}
+	inline int32 MaxNormalID() const
+	{
+		return NormalCount();
+	}
+	inline int32 NormalCount() const
+	{
+		return HasNormals() ? VertexCount() : 0;
+	}
+
+	inline FVector3f GetNormal(int32 IDValue) const
+	{
+		const FVector4f& Normal = Mesh->VertexBuffers.StaticMeshVertexBuffer.VertexTangentZ(IDValue);
+		return (!bScaleNormals) ? FVector3f(Normal.X, Normal.Y, Normal.Z) :
+			UE::Geometry::Normalized(FVector3f(Normal.X / BuildScale.X, Normal.Y / BuildScale.Y, Normal.Z / BuildScale.Z));
+	}
+
+	/** Get Normals for a given Triangle */
+	template<typename VectorType>
+	inline void GetTriNormals(int32 TriId, VectorType& N0, VectorType& N1, VectorType& N2)
+	{
+		FIndex3i VtxIndices = GetTriangle(TriId);
+
+		N0 = GetNormal(VtxIndices.A);
+		N1 = GetNormal(VtxIndices.B);
+		N2 = GetNormal(VtxIndices.C);
+	}
+
+	inline bool HasUVs(const int32 UVLayer = 0) const
+	{
+		return Mesh && Mesh->VertexBuffers.StaticMeshVertexBuffer.GetAllowCPUAccess() && Mesh->VertexBuffers.StaticMeshVertexBuffer.GetTexCoordData() && (int32)Mesh->VertexBuffers.StaticMeshVertexBuffer.GetNumTexCoords() > UVLayer && UVLayer >= 0;
+	}
+	inline bool IsUV(const int32 UVId) const
+	{
+		return HasUVs() && UVId >= 0 && UVId < UVCount();
+	}
+	inline int32 MaxUVID() const
+	{
+		return UVCount();
+	}
+	inline int32 UVCount() const
+	{
+		return HasUVs() ? Mesh->GetNumVertices() : 0;
+	}
+
+	/** Get UV by VertexInstanceID for a given UVLayer */
+	inline FVector2f GetUV(const int32 IDValue, const int32 UVLayer) const
+	{
+		return Mesh->VertexBuffers.StaticMeshVertexBuffer.GetVertexUV(IDValue, UVLayer);
+	}
+
+	/** Get UVs for a given UVLayer and Triangle */
+	template<typename VectorType>
+	inline void GetTriUVs(const int32 TriId, const int32 UVLayer, VectorType& UV0, VectorType& UV1, VectorType& UV2)
+	{
+		FIndex3i VtxIndices = GetTriangle(TriId);
+
+		UV0 = GetUV(VtxIndices.A, UVLayer);
+		UV1 = GetUV(VtxIndices.B, UVLayer);
+		UV2 = GetUV(VtxIndices.C, UVLayer);
+	}
+};
