ModelingTools: fix static analysis issue
#rnx

diff --git a/Engine/Plugins/Runtime/MeshModelingToolset/Source/ModelingComponents/Private/Scene/MeshSceneAdapter.cpp b/Engine/Plugins/Runtime/MeshModelingToolset/Source/ModelingComponents/Private/Scene/MeshSceneAdapter.cpp
--- a/Engine/Plugins/Runtime/MeshModelingToolset/Source/ModelingComponents/Private/Scene/MeshSceneAdapter.cpp
+++ b/Engine/Plugins/Runtime/MeshModelingToolset/Source/ModelingComponents/Private/Scene/MeshSceneAdapter.cpp
@@ -1334,7 +1334,7 @@ static TUniquePtr<IMeshSpatialWrapper> SpatialWrapperFactory( const FMeshTypeCon
 				return SMWrapper;
 			}
 		}
-		else if (BuildOptions.bEnableUVQueries == false && BuildOptions.bEnableUVQueries == false)
+		else if (BuildOptions.bEnableUVQueries == false && BuildOptions.bEnableNormalsQueries == false)
 		{
 			TUniquePtr<FCompressedStaticMeshSpatialWrapper> SMWrapper = MakeUnique<FCompressedStaticMeshSpatialWrapper>();
 			SMWrapper->SourceContainer = MeshContainer;
