ModelingTools: fix LWC issues in scene snapping and AddPrimitiveTool. Convert MeshDebugDraw grid-drawing functions to use double instead of float, and update call sites.
#rnx

diff --git a/Engine/Plugins/Experimental/ChaosEditor/Source/FractureEditor/Private/FractureToolPlaneCut.cpp b/Engine/Plugins/Experimental/ChaosEditor/Source/FractureEditor/Private/FractureToolPlaneCut.cpp
--- a/Engine/Plugins/Experimental/ChaosEditor/Source/FractureEditor/Private/FractureToolPlaneCut.cpp
+++ b/Engine/Plugins/Experimental/ChaosEditor/Source/FractureEditor/Private/FractureToolPlaneCut.cpp
@@ -90,7 +90,7 @@ void UFractureToolPlaneCut::Render(const FSceneView* View, FViewport* Viewport,
 			const float GridLineSpacing = PlaneSize / (NumGridLines - 1);
 			const FColor GridColor(64, 64, 64, 128);
 			const float GridThickness = 1.0f;
-			const UE::Geometry::FFrame3f DrawFrame(Center, Transform.GetRotation());
+			const UE::Geometry::FFrame3d DrawFrame(Center, Transform.GetRotation());
 			MeshDebugDraw::DrawSimpleGrid(DrawFrame, NumGridLines, GridLineSpacing, GridThickness, GridColor, true, PDI, FTransform::Identity);
 		};
 
diff --git a/Engine/Plugins/Experimental/MeshModelingToolsetExp/Source/MeshModelingToolsExp/Private/DynamicMeshSculptTool.cpp b/Engine/Plugins/Experimental/MeshModelingToolsetExp/Source/MeshModelingToolsExp/Private/DynamicMeshSculptTool.cpp
--- a/Engine/Plugins/Experimental/MeshModelingToolsetExp/Source/MeshModelingToolsExp/Private/DynamicMeshSculptTool.cpp
+++ b/Engine/Plugins/Experimental/MeshModelingToolsetExp/Source/MeshModelingToolsExp/Private/DynamicMeshSculptTool.cpp
@@ -1471,7 +1471,7 @@ void UDynamicMeshSculptTool::Render(IToolsContextRenderAPI* RenderAPI)
 		FColor GridColor(128, 128, 128, 32);
 		float GridThickness = 0.5f*RenderCameraState.GetPDIScalingFactor();
 		int NumGridLines = 10;
-		FFrame3f DrawFrame(GizmoProperties->Position, GizmoProperties->Rotation);
+		FFrame3d DrawFrame(GizmoProperties->Position, GizmoProperties->Rotation);
 		MeshDebugDraw::DrawSimpleFixedScreenAreaGrid(RenderCameraState, DrawFrame, NumGridLines, 45.0, GridThickness, GridColor, false, PDI, FTransform::Identity);
 	}
 }
diff --git a/Engine/Plugins/Experimental/MeshModelingToolsetExp/Source/MeshModelingToolsExp/Private/Sculpting/MeshSculptToolBase.cpp b/Engine/Plugins/Experimental/MeshModelingToolsetExp/Source/MeshModelingToolsExp/Private/Sculpting/MeshSculptToolBase.cpp
--- a/Engine/Plugins/Experimental/MeshModelingToolsetExp/Source/MeshModelingToolsExp/Private/Sculpting/MeshSculptToolBase.cpp
+++ b/Engine/Plugins/Experimental/MeshModelingToolsetExp/Source/MeshModelingToolsExp/Private/Sculpting/MeshSculptToolBase.cpp
@@ -284,7 +284,7 @@ void UMeshSculptToolBase::Render(IToolsContextRenderAPI* RenderAPI)
 		FColor GridColor(128, 128, 128, 32);
 		float GridThickness = 0.5f*RenderCameraState.GetPDIScalingFactor();
 		int NumGridLines = 10;
-		FFrame3f DrawFrame(GizmoProperties->Position, GizmoProperties->Rotation);
+		FFrame3d DrawFrame(GizmoProperties->Position, GizmoProperties->Rotation);
 		MeshDebugDraw::DrawSimpleFixedScreenAreaGrid(RenderCameraState, DrawFrame, NumGridLines, 45.0, GridThickness, GridColor, false, PDI, FTransform::Identity);
 	}
 }
diff --git a/Engine/Plugins/Experimental/UVEditor/Source/UVEditor/Private/UVEditor2DViewportClient.cpp b/Engine/Plugins/Experimental/UVEditor/Source/UVEditor/Private/UVEditor2DViewportClient.cpp
--- a/Engine/Plugins/Experimental/UVEditor/Source/UVEditor/Private/UVEditor2DViewportClient.cpp
+++ b/Engine/Plugins/Experimental/UVEditor/Source/UVEditor/Private/UVEditor2DViewportClient.cpp
@@ -260,7 +260,7 @@ void FUVEditor2DViewportClient::DrawGrid(const FSceneView* View, FPrimitiveDrawI
 	if (ZoomFactor < 100000 && ZoomFactor > 1)
 	{
 		// Setup and call grid calling function
-		UE::Geometry::FFrame3f LocalFrame((FVector3f)ViewLoc);
+		UE::Geometry::FFrame3d LocalFrame(ViewLoc);
 		FTransform Transform;
 		TArray<FColor> Colors;
 		Colors.Push(FUVEditorUXSettings::GridMajorColor);
diff --git a/Engine/Plugins/Runtime/MeshModelingToolset/Source/MeshModelingTools/Private/AddPrimitiveTool.cpp b/Engine/Plugins/Runtime/MeshModelingToolset/Source/MeshModelingTools/Private/AddPrimitiveTool.cpp
--- a/Engine/Plugins/Runtime/MeshModelingToolset/Source/MeshModelingTools/Private/AddPrimitiveTool.cpp
+++ b/Engine/Plugins/Runtime/MeshModelingToolset/Source/MeshModelingTools/Private/AddPrimitiveTool.cpp
@@ -217,9 +217,9 @@ void UAddPrimitiveTool::UpdatePreviewPosition(const FInputDeviceRay& DeviceClick
 	FPlane DrawPlane(FVector::ZeroVector, FVector(0, 0, 1));
 	if (ShapeSettings->TargetSurface == EMakeMeshPlacementType::GroundPlane)
 	{
-		FVector3f DrawPlanePos = (FVector3f)FMath::RayPlaneIntersection(ClickPosWorldRay.Origin, ClickPosWorldRay.Direction, DrawPlane);
+		FVector3d DrawPlanePos = (FVector3d)FMath::RayPlaneIntersection(ClickPosWorldRay.Origin, ClickPosWorldRay.Direction, DrawPlane);
 		bHit = true;
-		ShapeFrame = FFrame3f(DrawPlanePos);
+		ShapeFrame = FFrame3d(DrawPlanePos);
 	}
 	else
 	{
@@ -228,20 +228,20 @@ void UAddPrimitiveTool::UpdatePreviewPosition(const FInputDeviceRay& DeviceClick
 		bHit = ToolSceneQueriesUtil::FindNearestVisibleObjectHit(this, Result, ClickPosWorldRay);
 		if (bHit)
 		{
-			FVector3f Normal = (FVector3f)Result.ImpactNormal;
+			FVector3d Normal = (FVector3d)Result.ImpactNormal;
 			if (!ShapeSettings->bAlignToNormal)
 			{
-				Normal = FVector3f::UnitZ();
+				Normal = FVector3d::UnitZ();
 			}
-			ShapeFrame = FFrame3f((FVector3f)Result.ImpactPoint, Normal);
+			ShapeFrame = FFrame3d((FVector3d)Result.ImpactPoint, Normal);
 			ShapeFrame.ConstrainedAlignPerpAxes();
 		}
 		else
 		{
 			// fall back to ground plane if we don't have a scene hit
-			FVector3f DrawPlanePos = (FVector3f)FMath::RayPlaneIntersection(ClickPosWorldRay.Origin, ClickPosWorldRay.Direction, DrawPlane);
+			FVector3d DrawPlanePos = (FVector3d)FMath::RayPlaneIntersection(ClickPosWorldRay.Origin, ClickPosWorldRay.Direction, DrawPlane);
 			bHit = true;
-			ShapeFrame = FFrame3f(DrawPlanePos);
+			ShapeFrame = FFrame3d(DrawPlanePos);
 		}
 	}
 
@@ -258,14 +258,14 @@ void UAddPrimitiveTool::UpdatePreviewPosition(const FInputDeviceRay& DeviceClick
 			TArray<FSceneSnapQueryResult> Results;
 			if (SnapManager->ExecuteSceneSnapQuery(Request, Results))
 			{
-				ShapeFrame.Origin = (FVector3f)Results[0].Position;
+				ShapeFrame.Origin = (FVector3d)Results[0].Position;
 			}
 		}
 	}
 
 	if (ShapeSettings->Rotation != 0)
 	{
-		ShapeFrame.Rotate(FQuaternionf(ShapeFrame.Z(), ShapeSettings->Rotation, true));
+		ShapeFrame.Rotate(FQuaterniond(ShapeFrame.Z(), ShapeSettings->Rotation, true));
 	}
 
 	if (bHit)
diff --git a/Engine/Plugins/Runtime/MeshModelingToolset/Source/MeshModelingTools/Private/DrawPolygonTool.cpp b/Engine/Plugins/Runtime/MeshModelingToolset/Source/MeshModelingTools/Private/DrawPolygonTool.cpp
--- a/Engine/Plugins/Runtime/MeshModelingToolset/Source/MeshModelingTools/Private/DrawPolygonTool.cpp
+++ b/Engine/Plugins/Runtime/MeshModelingToolset/Source/MeshModelingTools/Private/DrawPolygonTool.cpp
@@ -318,7 +318,7 @@ void UDrawPolygonTool::Render(IToolsContextRenderAPI* RenderAPI)
 	//
 	if (bInInteractiveExtrude == false)
 	{
-		FFrame3f DrawFrame((FVector3f)DrawPlaneOrigin, (FQuaternionf)DrawPlaneOrientation);
+		FFrame3d DrawFrame(DrawPlaneOrigin, DrawPlaneOrientation);
 		MeshDebugDraw::DrawSimpleFixedScreenAreaGrid(RenderCameraState, DrawFrame, NumGridLines, 45.0, GridThickness, GridColor, false, PDI, FTransform::Identity);
 	}
 
@@ -1327,7 +1327,7 @@ bool UDrawPolygonTool::GeneratePolygonMesh(const TArray<FVector3d>& Polygon, con
 		FVector3d ExtrudeNormal = FVector3d::UnitZ();
 		Extruder.ExtrudedPositionFunc = [&ExtrudeDistance, &ExtrudeNormal](const FVector3d& Position, const FVector3f& Normal, int VertexID)
 		{
-			return Position + ExtrudeDistance * ExtrudeNormal;
+			return Position + ExtrudeDistance * (FVector3d)ExtrudeNormal;
 		};
 
 		Extruder.Apply();
diff --git a/Engine/Plugins/Runtime/MeshModelingToolset/Source/MeshModelingTools/Public/AddPrimitiveTool.h b/Engine/Plugins/Runtime/MeshModelingToolset/Source/MeshModelingTools/Public/AddPrimitiveTool.h
--- a/Engine/Plugins/Runtime/MeshModelingToolset/Source/MeshModelingTools/Public/AddPrimitiveTool.h
+++ b/Engine/Plugins/Runtime/MeshModelingToolset/Source/MeshModelingTools/Public/AddPrimitiveTool.h
@@ -529,7 +529,7 @@ protected:
 	UWorld* TargetWorld;
 
 	void UpdatePreviewPosition(const FInputDeviceRay& ClickPos);
-	UE::Geometry::FFrame3f ShapeFrame;
+	UE::Geometry::FFrame3d ShapeFrame;
 
 	void UpdatePreviewMesh() const;
 };
diff --git a/Engine/Plugins/Runtime/MeshModelingToolset/Source/MeshModelingTools/Public/DrawPolygonTool.h b/Engine/Plugins/Runtime/MeshModelingToolset/Source/MeshModelingTools/Public/DrawPolygonTool.h
--- a/Engine/Plugins/Runtime/MeshModelingToolset/Source/MeshModelingTools/Public/DrawPolygonTool.h
+++ b/Engine/Plugins/Runtime/MeshModelingToolset/Source/MeshModelingTools/Public/DrawPolygonTool.h
@@ -319,9 +319,6 @@ protected:
 	UPROPERTY()
 	TObjectPtr<UPlaneDistanceFromHitMechanic> HeightMechanic;
 
-
-	UE::Geometry::FFrame3f HitPosFrameWorld;
-
 	/** Generate extruded meshes.  Returns true on success. */
 	bool GeneratePolygonMesh(const TArray<FVector3d>& Polygon, const TArray<TArray<FVector3d>>& PolygonHoles, FDynamicMesh3* ResultMeshOut, UE::Geometry::FFrame3d& WorldFrameOut, bool bIncludePreviewVtx, double ExtrudeDistance, bool bExtrudeSymmetric);
 
diff --git a/Engine/Plugins/Runtime/MeshModelingToolset/Source/ModelingComponents/Private/Drawing/MeshDebugDrawing.cpp b/Engine/Plugins/Runtime/MeshModelingToolset/Source/ModelingComponents/Private/Drawing/MeshDebugDrawing.cpp
--- a/Engine/Plugins/Runtime/MeshModelingToolset/Source/ModelingComponents/Private/Drawing/MeshDebugDrawing.cpp
+++ b/Engine/Plugins/Runtime/MeshModelingToolset/Source/ModelingComponents/Private/Drawing/MeshDebugDrawing.cpp
@@ -21,9 +21,9 @@ void MeshDebugDraw::DrawNormals(
 	{
 		FVector3f Normal = Overlay->GetElement(ElementID);
 		int ParentVID = Overlay->GetParentVertex(ElementID);
-		FVector3f ParentPos = (FVector3f)Mesh->GetVertex(ParentVID);
+		FVector3d ParentPos = Mesh->GetVertex(ParentVID);
 
-		FVector A = (FVector)ParentPos, B = (FVector)(ParentPos + Length * Normal);
+		FVector A = (FVector)ParentPos, B = (FVector)(ParentPos + (double)Length * (FVector3d)Normal);
 		PDI->DrawLine(Transform.TransformPosition(A), Transform.TransformPosition(B),
 			Color, 0, Thickness, 0, bScreenSpace);
 	}
@@ -74,27 +74,27 @@ void MeshDebugDraw::DrawTriCentroids(
 
 
 void MeshDebugDraw::DrawSimpleGrid(	
-	const FFrame3f& LocalFrame, int GridLines, float GridLineSpacing,
+	const FFrame3d& LocalFrame, int GridLines, double GridLineSpacing,
 	float LineWidth, FColor Color, bool bDepthTested,
 	FPrimitiveDrawInterface* PDI, const FTransform& Transform)
 {
 	ESceneDepthPriorityGroup DepthPriority = (bDepthTested) ? SDPG_World : SDPG_Foreground;
 
-	FFrame3f WorldFrame = LocalFrame;
+	FFrame3d WorldFrame = LocalFrame;
 	WorldFrame.Transform(Transform);
 
-	float Width = (float)(GridLines-1) * GridLineSpacing;
-	float Extent = Width * 0.5;
+	double Width = (double)(GridLines-1) * GridLineSpacing;
+	double Extent = Width * 0.5;
 
-	FVector3f Origin = WorldFrame.Origin;
-	FVector3f X = WorldFrame.X();
-	FVector3f Y = WorldFrame.Y();
-	FVector3f A, B;
+	FVector3d Origin = WorldFrame.Origin;
+	FVector3d X = WorldFrame.X();
+	FVector3d Y = WorldFrame.Y();
+	FVector3d A, B;
 
 	int LineSteps = GridLines / 2;
 	for (int i = 0; i < LineSteps; i++)
 	{
-		float dx = (float)i * GridLineSpacing;
+		double dx = (double)i * GridLineSpacing;
 		A = Origin - Extent * Y - dx * X;
 		B = Origin + Extent * Y - dx * X;
 		PDI->DrawLine((FVector)A, (FVector)B, Color, DepthPriority, LineWidth, 0, true);
@@ -114,22 +114,22 @@ void MeshDebugDraw::DrawSimpleGrid(
 
 void MeshDebugDraw::DrawSimpleFixedScreenAreaGrid(
 	const FViewCameraState& CameraState,
-	const FFrame3f& LocalFrame, int32 NumGridLines, float VisualAngleSpan,
+	const FFrame3d& LocalFrame, int32 NumGridLines, double VisualAngleSpan,
 	float LineWidth, FColor Color, bool bDepthTested,
 	FPrimitiveDrawInterface* PDI, const FTransform& Transform)
 {
 	FVector WorldOrigin = Transform.TransformPosition((FVector)LocalFrame.Origin);
-	float GridWidth = ToolSceneQueriesUtil::CalculateDimensionFromVisualAngleD(CameraState, (FVector3d)WorldOrigin, VisualAngleSpan);
-	float GridLineSpacing = GridWidth / (float)NumGridLines;
+	double GridWidth = ToolSceneQueriesUtil::CalculateDimensionFromVisualAngleD(CameraState, (FVector3d)WorldOrigin, VisualAngleSpan);
+	double GridLineSpacing = GridWidth / (float)NumGridLines;
 	DrawSimpleGrid(LocalFrame, NumGridLines, GridLineSpacing, LineWidth, Color, bDepthTested, PDI, Transform);
 }
 
 
 void MeshDebugDraw::DrawHierarchicalGrid(
-	float BaseScale, float GridZoomFactor, int32 MaxLevelDensity,
+	double BaseScale, double GridZoomFactor, int32 MaxLevelDensity,
 	const FVector& WorldMaxBounds, const FVector& WorldMinBounds,
 	int32 Levels, int32 Subdivisions, TArray<FColor>& Colors,
-	const FFrame3f& LocalFrame, float LineWidth, bool bDepthTested,
+	const FFrame3d& LocalFrame, float LineWidth, bool bDepthTested,
 	FPrimitiveDrawInterface* PDI, const FTransform& Transform)
 {
 
@@ -137,27 +137,27 @@ void MeshDebugDraw::DrawHierarchicalGrid(
     // This allows us to track and discretely shift grid resolutions at certain zoom levels.
 	// This code assumes that we want the grid to remain stable with one logrithmic "unit"
 	// around the base zoom factor of 1.0
-	float LogZoom = FMath::LogX(Subdivisions, GridZoomFactor );
-	float LogZoomDirection = FMath::Sign(LogZoom);
+	double LogZoom = FMath::LogX(Subdivisions, GridZoomFactor );
+	double LogZoomDirection = FMath::Sign(LogZoom);
 	LogZoom = FMath::Abs(LogZoom);
 	LogZoom = FMathf::Floor(LogZoom);
 	LogZoom = LogZoomDirection * LogZoom;
 
 	// Adjust grid scales based on current zoom levels
-	TArray<float> GridScales;
+	TArray<double> GridScales;
 	GridScales.SetNum(Levels);
 	for (int32 Level = 0; Level < Levels; ++Level)
 	{
-		GridScales[Level] = BaseScale * FMathf::Pow(Subdivisions, LogZoom - Level);
+		GridScales[Level] = BaseScale * FMathd::Pow(Subdivisions, LogZoom - Level);
 		ensure(!FMath::IsNearlyZero(GridScales[Level]));
 	}
 	
 	//Determine the center of the drawing area for the grid, snapping to grid positions
-	FVector3f GridOrigin(
+	FVector3d GridOrigin(
 		FMath::GridSnap(LocalFrame.Origin.X, GridScales[0]),
 		FMath::GridSnap(LocalFrame.Origin.Y, GridScales[0]),
 		LocalFrame.Origin.Z);
-	UE::Geometry::FFrame3f GridFrame(GridOrigin);
+	UE::Geometry::FFrame3d GridFrame(GridOrigin);
 
 	// Draw each level of the grid
 	for (int32 Level = 0; Level < Levels; ++Level)
@@ -166,7 +166,7 @@ void MeshDebugDraw::DrawHierarchicalGrid(
 		float AdjustedLineWidth = LineWidth * FMathf::Pow(2, -Level);
 
 		// Compute the number of needed grid lines based on the coarsest grid's scale, this way we never "run out" of lines as we pan about
-		int32 GridLines = FMathf::Ceil(FMathf::Max(((WorldMaxBounds.X - WorldMinBounds.X + GridScales[0]*2) / GridScales[Level]),
+		int32 GridLines = FMathd::Ceil(FMathd::Max(((WorldMaxBounds.X - WorldMinBounds.X + GridScales[0]*2) / GridScales[Level]),
 			                                        ((WorldMaxBounds.Y - WorldMinBounds.Y + GridScales[0]*2) / GridScales[Level])));
 
 		// If we ever have too many lines to draw, just bail. This preserves performance for large, deep grids.
diff --git a/Engine/Plugins/Runtime/MeshModelingToolset/Source/ModelingComponents/Private/Mechanics/ConstructionPlaneMechanic.cpp b/Engine/Plugins/Runtime/MeshModelingToolset/Source/ModelingComponents/Private/Mechanics/ConstructionPlaneMechanic.cpp
--- a/Engine/Plugins/Runtime/MeshModelingToolset/Source/ModelingComponents/Private/Mechanics/ConstructionPlaneMechanic.cpp
+++ b/Engine/Plugins/Runtime/MeshModelingToolset/Source/ModelingComponents/Private/Mechanics/ConstructionPlaneMechanic.cpp
@@ -128,7 +128,7 @@ void UConstructionPlaneMechanic::Render(IToolsContextRenderAPI* RenderAPI)
 		FColor GridColor(128, 128, 128, 32);
 		float GridThickness = 0.75f*PDIScale;
 
-		FFrame3f DrawFrame(Plane);
+		FFrame3d DrawFrame(Plane);
 		MeshDebugDraw::DrawSimpleFixedScreenAreaGrid(CameraState, DrawFrame, NumGridLines, 45.0, GridThickness, GridColor, false, PDI, FTransform::Identity);
 	}
 }
diff --git a/Engine/Plugins/Runtime/MeshModelingToolset/Source/ModelingComponents/Private/Snapping/ModelingSceneSnappingManager.cpp b/Engine/Plugins/Runtime/MeshModelingToolset/Source/ModelingComponents/Private/Snapping/ModelingSceneSnappingManager.cpp
--- a/Engine/Plugins/Runtime/MeshModelingToolset/Source/ModelingComponents/Private/Snapping/ModelingSceneSnappingManager.cpp
+++ b/Engine/Plugins/Runtime/MeshModelingToolset/Source/ModelingComponents/Private/Snapping/ModelingSceneSnappingManager.cpp
@@ -22,43 +22,41 @@ using namespace UE::Geometry;
 
 #define LOCTEXT_NAMESPACE "USceneSnappingManager"
 
-
-static float SnapToIncrement(float fValue, float fIncrement, float offset = 0)
+static double SnapToIncrement(double fValue, double fIncrement, double offset = 0)
 {
 	if (!FMath::IsFinite(fValue))
 	{
 		return 0;
 	}
 	fValue -= offset;
-	float sign = FMath::Sign(fValue);
+	double sign = FMath::Sign(fValue);
 	fValue = FMath::Abs(fValue);
-	int nInc = (int)(fValue / fIncrement);
-	float fRem = (float)fmod(fValue, fIncrement);
-	if (fRem > fIncrement / 2)
+	int64 nInc = (int64)(fValue / fIncrement);
+	double fRem = (double)fmod(fValue, fIncrement);
+	if (fRem > fIncrement / 2.0)
 	{
 		++nInc;
 	}
-	return sign * (float)nInc * fIncrement + offset;
+	return sign * (double)nInc * fIncrement + offset;
 }
 
-
 //@ todo this are mirrored from GeometryProcessing, which is still experimental...replace w/ direct calls once GP component is standardized
-static float OpeningAngleDeg(FVector A, FVector B, const FVector& P)
+static double OpeningAngleDeg(FVector3d A, FVector3d B, const FVector3d& P)
 {
 	A -= P;
 	A.Normalize();
 	B -= P;
 	B.Normalize();
-	float Dot = FMath::Clamp(FVector::DotProduct(A,B), -1.0f, 1.0f);
-	return acos(Dot) * (180.0f / 3.141592653589f);
+	double Dot = FMath::Clamp(FVector3d::DotProduct(A,B), -1.0, 1.0);
+	return FMathd::ACos(Dot) * (180.0 / 3.141592653589);
 }
 
-static FVector NearestSegmentPt(FVector A, FVector B, const FVector& P)
+static FVector3d NearestSegmentPt(FVector3d A, FVector3d B, const FVector3d& P)
 {
-	FVector Direction = (B - A);
-	float Length = Direction.Size();
+	FVector3d Direction = (B - A);
+	double Length = Direction.Size();
 	Direction /= Length;
-	float t = FVector::DotProduct( (P - A), Direction);
+	double t = FVector3d::DotProduct( (P - A), Direction);
 	if (t >= Length)
 	{
 		return B;
@@ -332,7 +330,7 @@ static bool FindNearestVisibleObjectHit_Internal(
 		return false;
 	}
 
-	float NearestVisible = TNumericLimits<float>::Max();
+	double NearestVisible = TNumericLimits<double>::Max();
 	for (const FHitResult& CurResult : OutHits)
 	{
 		// if we have hit Component in the SpatialCache, prefer to use that
@@ -348,7 +346,7 @@ static bool FindNearestVisibleObjectHit_Internal(
 		}
 	}
 
-	return NearestVisible < TNumericLimits<float>::Max();
+	return NearestVisible < TNumericLimits<double>::Max();
 }
 
 
@@ -454,7 +452,6 @@ bool UModelingSceneSnappingManager::ExecuteSceneHitQuery(const FSceneHitQueryReq
 
 
 
-
 bool UModelingSceneSnappingManager::ExecuteSceneSnapQueryPosition(const FSceneSnapQueryRequest& Request, TArray<FSceneSnapQueryResult>& Results) const
 {
 	if (!QueriesAPI)
@@ -475,9 +472,9 @@ bool UModelingSceneSnappingManager::ExecuteSceneSnapQueryPosition(const FSceneSn
 
 			FVector GridSize = Request.GridSize.Get(SnappingConfig.PositionGridDimensions);
 
-			SnapResult.Position.X = SnapToIncrement(Request.Position.X, GridSize.X);
-			SnapResult.Position.Y = SnapToIncrement(Request.Position.Y, GridSize.Y);
-			SnapResult.Position.Z = SnapToIncrement(Request.Position.Z, GridSize.Z);
+			SnapResult.Position.X = (FVector::FReal)SnapToIncrement(Request.Position.X, GridSize.X);
+			SnapResult.Position.Y = (FVector::FReal)SnapToIncrement(Request.Position.Y, GridSize.Y);
+			SnapResult.Position.Z = (FVector::FReal)SnapToIncrement(Request.Position.Z, GridSize.Z);
 
 			Results.Add(SnapResult);
 			FoundResultCount++;
@@ -500,7 +497,7 @@ bool UModelingSceneSnappingManager::ExecuteSceneSnapQueryPosition(const FSceneSn
 		{
 			for (int j = 0; j < 3; ++j)
 			{
-				double VisualAngle = OpeningAngleDeg(RequestIn.Position, SnapResult.TriVertices[j], (FVector)WorldRay.Origin);
+				double VisualAngle = OpeningAngleDeg((FVector3d)RequestIn.Position, (FVector3d)SnapResult.TriVertices[j], WorldRay.Origin);
 				if (VisualAngle < SmallestSnapAngle)
 				{
 					SmallestSnapAngle = VisualAngle;
@@ -520,8 +517,8 @@ bool UModelingSceneSnappingManager::ExecuteSceneSnapQueryPosition(const FSceneSn
 		{
 			for (int j = 0; j < 3; ++j)
 			{
-				FVector EdgeNearestPt = NearestSegmentPt(SnapResult.TriVertices[j], SnapResult.TriVertices[(j+1)%3], RequestIn.Position);
-				double VisualAngle = OpeningAngleDeg(RequestIn.Position, EdgeNearestPt, (FVector)WorldRay.Origin);
+				FVector3d EdgeNearestPt = NearestSegmentPt((FVector3d)SnapResult.TriVertices[j], (FVector3d)SnapResult.TriVertices[(j+1)%3], (FVector3d)RequestIn.Position);
+				double VisualAngle = OpeningAngleDeg((FVector3d)RequestIn.Position, EdgeNearestPt, WorldRay.Origin);
 				if (VisualAngle < SmallestSnapAngle )
 				{
 					SmallestSnapAngle = VisualAngle;
@@ -573,8 +570,8 @@ bool UModelingSceneSnappingManager::ExecuteSceneSnapQueryPosition(const FSceneSn
 	if (bHitSceneCache)
 	{
 		// is this check necessary? isn't this always going to be zero for a scene hit??
-		float VisualAngle = OpeningAngleDeg(Request.Position, (FVector)HitPoint.WorldPoint, (FVector)RayStart);
-		if (VisualAngle < Request.VisualAngleThresholdDegrees)
+		double VisualAngle = OpeningAngleDeg((FVector3d)Request.Position, HitPoint.WorldPoint, RayStart);
+		if (VisualAngle < (double)Request.VisualAngleThresholdDegrees)
 		{
 			FVector3d A, B, C;
 			SpatialCache->GetGeometry(HitIdentifier, HitPoint.GeometryType, HitPoint.GeometryIndex, true, A, B, C);
@@ -608,8 +605,8 @@ bool UModelingSceneSnappingManager::ExecuteSceneSnapQueryPosition(const FSceneSn
 	if (bHitWorld)
 	{
 		// is this check necessary? isn't this always going to be zero for a scene hit??
-		float VisualAngle = OpeningAngleDeg(Request.Position, HitResult.ImpactPoint, RayStart);
-		if (VisualAngle < Request.VisualAngleThresholdDegrees)
+		double VisualAngle = OpeningAngleDeg((FVector3d)Request.Position, (FVector3d)HitResult.ImpactPoint, RayStart);
+		if (VisualAngle < (double)Request.VisualAngleThresholdDegrees)
 		{
 			FSceneSnapQueryResult SnapResult;
 			if ( GetComponentHitTriangle_Internal<FVector>(HitResult, SnapResult.TriVertices) )
@@ -638,7 +635,6 @@ bool UModelingSceneSnappingManager::ExecuteSceneSnapQueryPosition(const FSceneSn
 }
 
 
-
 bool UE::Geometry::RegisterSceneSnappingManager(UInteractiveToolsContext* ToolsContext)
 {
 	if (ensure(ToolsContext))
diff --git a/Engine/Plugins/Runtime/MeshModelingToolset/Source/ModelingComponents/Public/Drawing/MeshDebugDrawing.h b/Engine/Plugins/Runtime/MeshModelingToolset/Source/ModelingComponents/Public/Drawing/MeshDebugDrawing.h
--- a/Engine/Plugins/Runtime/MeshModelingToolset/Source/ModelingComponents/Public/Drawing/MeshDebugDrawing.h
+++ b/Engine/Plugins/Runtime/MeshModelingToolset/Source/ModelingComponents/Public/Drawing/MeshDebugDrawing.h
@@ -80,7 +80,7 @@ namespace MeshDebugDraw
 	 * @param Transform transform applied to LocalFrame. Pass as Identity() if you have world frame.
 	 */
 	void MODELINGCOMPONENTS_API DrawSimpleGrid(
-		const FFrame3f& LocalFrame, int GridLines, float GridLineSpacing,
+		const FFrame3d& LocalFrame, int GridLines, double GridLineSpacing,
 		float LineWidth, FColor Color, bool bDepthTested,
 		FPrimitiveDrawInterface* PDI, const FTransform& Transform);
 
@@ -99,7 +99,7 @@ namespace MeshDebugDraw
 	 */
 	void MODELINGCOMPONENTS_API DrawSimpleFixedScreenAreaGrid(
 		const FViewCameraState& CameraState, 
-		const FFrame3f& LocalFrame, int32 NumGridLines, float VisualAngleSpan,
+		const FFrame3d& LocalFrame, int32 NumGridLines, double VisualAngleSpan,
 		float LineWidth, FColor Color, bool bDepthTested,
 		FPrimitiveDrawInterface* PDI, const FTransform& Transform);
 
@@ -119,9 +119,9 @@ namespace MeshDebugDraw
 	 * @param Transform transform applied to LocalFrame. Pass as Identity() if you have world frame.		
 	*/
 	void MODELINGCOMPONENTS_API DrawHierarchicalGrid(
-		float BaseGridScale, float GridZoomFactor, int32 MaxLevelDensity,
+		double BaseGridScale, double GridZoomFactor, int32 MaxLevelDensity,
 		const FVector& WorldMaxBounds, const FVector& WorldMinBounds,
 		int32 Levels, int32 Subdivisions, TArray<FColor>& Colors,
-		const FFrame3f& LocalFrame, float LineWidth, bool bDepthTested,
+		const FFrame3d& LocalFrame, float LineWidth, bool bDepthTested,
 		FPrimitiveDrawInterface* PDI, const FTransform& Transform);
 }
\ No newline at end of file
