ModelingTools: fix issue in MeshSceneAdapter introduced by recent refactor. FActorChildMesh::MeshSpatial was being always initialized to nullptr, and only updated if the mesh was decomposed. Add code to catch this case in future, and a bit of header documentation explaining the connections between the various data structures.
#rnx

diff --git a/Engine/Plugins/Runtime/MeshModelingToolset/Source/ModelingComponents/Private/Scene/MeshSceneAdapter.cpp b/Engine/Plugins/Runtime/MeshModelingToolset/Source/ModelingComponents/Private/Scene/MeshSceneAdapter.cpp
--- a/Engine/Plugins/Runtime/MeshModelingToolset/Source/ModelingComponents/Private/Scene/MeshSceneAdapter.cpp
+++ b/Engine/Plugins/Runtime/MeshModelingToolset/Source/ModelingComponents/Private/Scene/MeshSceneAdapter.cpp
@@ -1513,7 +1513,7 @@ void FMeshSceneAdapter::InitializeSpatialWrappers(const TArray<FActorAdapter*>&
 			{
 				MeshInfo->NonUniformScaleCount++;
 			}
-			ChildMesh->MeshSpatial = MeshInfo->SpatialWrapper.Get();
+			ChildMesh->MeshSpatial = nullptr;	// these are now initialized at beginning of Build() function
 		}
 	}
 }
@@ -1527,6 +1527,11 @@ void FMeshSceneAdapter::Build(const FMeshSceneAdapterBuildOptions& BuildOptions)
 	for (TPair<void*, TSharedPtr<FSpatialWrapperInfo>> Pair : SpatialAdapters)
 	{
 		Pair.Value->SpatialWrapper = SpatialWrapperFactory(Pair.Value->SourceContainer, BuildOptions);
+		// populate the MeshSpatial members of the FActorChildMeshes
+		for (FActorChildMesh* ParentMesh : Pair.Value->ParentMeshes)
+		{
+			ParentMesh->MeshSpatial = Pair.Value->SpatialWrapper.Get();
+		}
 	}
 
 	if (BuildOptions.bThickenThinMeshes)
@@ -1978,11 +1983,12 @@ void FMeshSceneAdapter::Build_FullDecompose(const FMeshSceneAdapterBuildOptions&
 		}
 		else
 		{
-			// have to null out spatials for the child meshes so that they are ignored during computation
+			// disconnect parent meshes from this spatialwrapper as it is now invalid
 			for (FActorChildMesh* MeshInstance : MeshesToDecompose)
 			{
 				MeshInstance->MeshSpatial = nullptr;
 			}
+			WrapperInfo->ParentMeshes.Reset();
 		}
 
 		// Exit if we don't have any more work to do. This happens if we ended up skipping all the possible decompositions
@@ -2254,6 +2260,20 @@ void FMeshSceneAdapter::BuildSpatialEvaluationCache()
 {
 	TRACE_CPUPROFILER_EVENT_SCOPE(MeshScene_BuildSpatialEvaluationCache);
 
+	// validate that internals are correctly configured
+	for (TPair<void*, TSharedPtr<FSpatialWrapperInfo>> Pair : SpatialAdapters)
+	{
+		FSpatialWrapperInfo& WrapperInfo = *Pair.Value;
+		IMeshSpatialWrapper* SpatialWrapper = WrapperInfo.SpatialWrapper.Get();
+		for (FActorChildMesh* MeshInstance : WrapperInfo.ParentMeshes)
+		{
+			if (ensure(MeshInstance->MeshSpatial == SpatialWrapper) == false)
+			{
+				UE_LOG(LogGeometry, Warning, TEXT("FMeshSceneAdapter::BuildSpatialEvaluationCache: broken MeshSpatial link found!"));
+			}
+		}
+	}
+
 	// build list of unique meshes we need to evaluate for spatial queries
 	SortedSpatials.Reset();
 	for (const TUniquePtr<FActorAdapter>& Actor : SceneActors)
diff --git a/Engine/Plugins/Runtime/MeshModelingToolset/Source/ModelingComponents/Public/Scene/MeshSceneAdapter.h b/Engine/Plugins/Runtime/MeshModelingToolset/Source/ModelingComponents/Public/Scene/MeshSceneAdapter.h
--- a/Engine/Plugins/Runtime/MeshModelingToolset/Source/ModelingComponents/Public/Scene/MeshSceneAdapter.h
+++ b/Engine/Plugins/Runtime/MeshModelingToolset/Source/ModelingComponents/Public/Scene/MeshSceneAdapter.h
@@ -392,11 +392,42 @@ protected:
 	// top-level list of ActorAdapters, which represent each Actor and set of Components
 	TArray<TUniquePtr<FActorAdapter>> SceneActors;
 
+	/*
+	 * The data structures used for supporting spatial queries in FMeshSceneAdapter are somewhat complex.
+	 * Externally, the user adds Actors/Components, which results in FActorAdapter objects, each with list of
+	 * FActorChildMesh. Each FActorChildMesh is a /reference/ to a mesh, not a mesh itself, as well as a transform.
+	 * This allows unique mesh data to be shared across many usages, eg similar to StaticMeshActor / ISMComponent / Asset.
+	 * (However there is no required 1-1 correspondence between FActorChildMesh and an actual Component, and in the case
+	 *  of an InstancedStaticMeshComponent, there will be an FActorChildMesh for each Component)
+	 *
+	 * Each FActorChildMesh has a pointer to an IMeshSpatialWrapper, this is where unique mesh data, and any spatial data
+	 * structures, will be stored. Each unique IMeshSpatialWrapper instance is owned by a FSpatialWrapperInfo, which
+	 * are stored in the SpatialAdapters map below. Each FSpatialWrapperInfo also knows which parent FActorChildMesh 
+	 * objects own it.
+	 * 
+	 * The main action inside a FSpatialWrapperInfo (ie unique mesh) is inside the IMeshSpatialWrapper implementation,
+	 * this is where mesh data is unpacked and things like an AABBTree will be built if required. The top-level Build() 
+	 * will do this work.
+	 * 
+	 * *However* evaluating spatial queries across a large set of actors / scene meshes would require testing each
+	 * mesh instance one-by-one. To speed this up, BuildSpatialEvaluationCache() can be called to build an Octree
+	 * across mesh instances. FSpatialCacheInfo is a separate representation of each unique (actor, chlidmesh, spatial, boundingbox)
+	 * tuple, a list of these is built and then inserted into an Octree.
+	 */
+
+
+	// FSpatialWrapperInfo is a "unique mesh" that knows which parent objects have references to it
 	struct FSpatialWrapperInfo
 	{
+		// identifier for mesh, including source pointer, mesh type, etc
 		FMeshTypeContainer SourceContainer;
+		// list of scene mesh instances that reference this unique mehs
 		TArray<FActorChildMesh*> ParentMeshes;
+		// number of non-uniform scales applied to this mesh by parent instances
 		int32 NonUniformScaleCount = 0;
+		// implementation of IMeshSpatialWrapper for this mesh, that provides spatial and other mesh data queries
+		// *NOTE* that each FActorChildMesh in ParentMeshes has a pointer to this object, so if SpatialWrapper is
+		// replaced, the FActorChildMesh objects must also be updated
 		TUniquePtr<IMeshSpatialWrapper> SpatialWrapper;
 	};
 
@@ -412,6 +443,7 @@ protected:
 	void Build_FullDecompose(const FMeshSceneAdapterBuildOptions& BuildOptions);
 
 
+	// FSpatialCacheInfo represents a unique scene mesh with spatial data structure
 	struct FSpatialCacheInfo
 	{
 		FActorAdapter* Actor;
@@ -419,7 +451,9 @@ protected:
 		IMeshSpatialWrapper* Spatial;
 		FAxisAlignedBox3d Bounds;
 	};
+	// list of all unique scene meshes that have spatial data structure available
 	TArray<FSpatialCacheInfo> SortedSpatials;
+	// Octree of elements in SortedSpatials list
 	TSharedPtr<FSparseDynamicOctree3> Octree;
 	FAxisAlignedBox3d CachedWorldBounds;
 	bool bHaveSpatialEvaluationCache = false;
