Fixed compilation warning as error.

ParameterizeMeshTool.h was rollback and converted to edit: history ispreserved.

diff --git a/Engine/Plugins/Runtime/MeshModelingToolset/Source/MeshModelingToolsEditorOnly/Public/ParameterizeMeshTool.h b/Engine/Plugins/Runtime/MeshModelingToolset/Source/MeshModelingToolsEditorOnly/Public/ParameterizeMeshTool.h
new file mode 100644
--- /dev/null
+++ b/Engine/Plugins/Runtime/MeshModelingToolset/Source/MeshModelingToolsEditorOnly/Public/ParameterizeMeshTool.h
@@ -0,0 +1,107 @@
+// Copyright Epic Games, Inc. All Rights Reserved.
+
+#pragma once
+
+#include "CoreMinimal.h"
+#include "BaseTools/SingleSelectionMeshEditingTool.h"
+#include "DynamicMesh/DynamicMesh3.h"
+#include "MeshOpPreviewHelpers.h"
+#include "Properties/MeshMaterialProperties.h"
+#include "Properties/MeshUVChannelProperties.h"
+#include "PropertySets/PolygroupLayersProperties.h"
+#include "Drawing/UVLayoutPreview.h"
+
+#include "ParameterizeMeshTool.generated.h"
+
+// Forward declarations
+class UDynamicMeshComponent;
+class UMaterialInterface;
+class UMaterialInstanceDynamic;
+class UParameterizeMeshToolProperties;
+class UParameterizeMeshToolUVAtlasProperties;
+class UParameterizeMeshToolXAtlasProperties;
+class UParameterizeMeshToolPatchBuilderProperties;
+
+
+UCLASS()
+class MESHMODELINGTOOLSEDITORONLY_API UParameterizeMeshToolBuilder : public USingleSelectionMeshEditingToolBuilder
+{
+	GENERATED_BODY()
+public:
+	virtual USingleSelectionMeshEditingTool* CreateNewTool(const FToolBuilderState& SceneState) const override;
+};
+
+
+UCLASS()
+class MESHMODELINGTOOLSEDITORONLY_API UParameterizeMeshToolPatchBuilderGroupLayerProperties : public UPolygroupLayersProperties
+{
+	GENERATED_BODY()
+public:
+
+	UPROPERTY(EditAnywhere, Category = "PolyGroup Layer", meta = (DisplayName = "Constrain to PolyGroups"))
+	bool bConstrainToPolygroups = false;
+};
+
+
+/**
+ * UParameterizeMeshTool automatically decomposes the input mesh into charts, solves for UVs,
+ * and then packs the resulting charts
+ */
+UCLASS()
+class MESHMODELINGTOOLSEDITORONLY_API UParameterizeMeshTool : public USingleSelectionMeshEditingTool, public UE::Geometry::IDynamicMeshOperatorFactory
+{
+	GENERATED_BODY()
+
+public:
+	virtual void Setup() override;
+	virtual void OnShutdown(EToolShutdownType ShutdownType) override;
+
+	virtual void Render(IToolsContextRenderAPI* RenderAPI) override;
+	virtual void OnTick(float DeltaTime) override;
+
+	virtual bool HasCancel() const override { return true; }
+	virtual bool HasAccept() const override { return true; }
+	virtual bool CanAccept() const override;
+
+	virtual void OnPropertyModified(UObject* PropertySet, FProperty* Property) override;
+
+	// IDynamicMeshOperatorFactory API
+	virtual TUniquePtr<UE::Geometry::FDynamicMeshOperator> MakeNewOperator() override;
+
+protected:
+	UPROPERTY()
+	TObjectPtr<UMeshUVChannelProperties> UVChannelProperties = nullptr;
+
+	UPROPERTY()
+	TObjectPtr<UParameterizeMeshToolProperties> Settings = nullptr;
+
+	UPROPERTY()
+	TObjectPtr<UParameterizeMeshToolUVAtlasProperties> UVAtlasProperties = nullptr;
+
+	UPROPERTY()
+	TObjectPtr<UParameterizeMeshToolXAtlasProperties> XAtlasProperties = nullptr;
+
+	UPROPERTY()
+	TObjectPtr<UParameterizeMeshToolPatchBuilderProperties> PatchBuilderProperties = nullptr;
+
+	UPROPERTY()
+	TObjectPtr<UParameterizeMeshToolPatchBuilderGroupLayerProperties> PolygroupLayerProperties = nullptr;
+
+	UPROPERTY()
+	TObjectPtr<UExistingMeshMaterialProperties> MaterialSettings = nullptr;
+
+	UPROPERTY()
+	bool bCreateUVLayoutViewOnSetup = true;
+
+	UPROPERTY()
+	TObjectPtr<UUVLayoutPreview> UVLayoutView = nullptr;
+
+	UPROPERTY()
+	TObjectPtr<UMeshOpPreviewWithBackgroundCompute> Preview = nullptr;
+
+	TSharedPtr<UE::Geometry::FDynamicMesh3, ESPMode::ThreadSafe> InputMesh;
+
+	void OnMethodTypeChanged();
+
+	void OnPreviewMeshUpdated();
+};
diff --git a/Engine/Source/Runtime/Core/Private/Misc/DateTime.cpp b/Engine/Source/Runtime/Core/Private/Misc/DateTime.cpp
--- a/Engine/Source/Runtime/Core/Private/Misc/DateTime.cpp
+++ b/Engine/Source/Runtime/Core/Private/Misc/DateTime.cpp
@@ -93,7 +93,7 @@ void FDateTime::GetDate(int32& OutYear, int32& OutMonth, int32& OutDay) const
 
 	int32 i, j, k, l, n;
 
-	l = FMath::FloorToInt(GetDate().GetJulianDay() + 0.5) + 68569;
+	l = (int32)(FMath::FloorToInt(GetDate().GetJulianDay() + 0.5) + 68569);
 	n = 4 * l / 146097;
 	l = l - (146097 * n + 3) / 4;
 	i = 4000 * (l + 1) / 1461001;
