GeometryCore: clean up issues that block non-LWC builds

diff --git a/Engine/Plugins/Experimental/MeshModelingToolsetExp/Source/MeshModelingToolsEditorOnlyExp/Private/PolygonOnMeshTool.cpp b/Engine/Plugins/Experimental/MeshModelingToolsetExp/Source/MeshModelingToolsEditorOnlyExp/Private/PolygonOnMeshTool.cpp
--- a/Engine/Plugins/Experimental/MeshModelingToolsetExp/Source/MeshModelingToolsEditorOnlyExp/Private/PolygonOnMeshTool.cpp
+++ b/Engine/Plugins/Experimental/MeshModelingToolsetExp/Source/MeshModelingToolsEditorOnlyExp/Private/PolygonOnMeshTool.cpp
@@ -75,7 +75,7 @@ void UPolygonOnMeshTool::Setup()
 	AddInputBehavior(HoverBehavior);
 
 	IPrimitiveComponentBackedTarget* TargetComponent = Cast<IPrimitiveComponentBackedTarget>(Target);
-	WorldTransform = TargetComponent->GetWorldTransform();
+	WorldTransform = (FTransform3d)TargetComponent->GetWorldTransform();
 
 	// hide input StaticMeshComponent
 	TargetComponent->SetOwnerVisibility(false);
diff --git a/Engine/Plugins/Experimental/MeshModelingToolsetExp/Source/MeshModelingToolsExp/Private/HoleFillTool.cpp b/Engine/Plugins/Experimental/MeshModelingToolsetExp/Source/MeshModelingToolsExp/Private/HoleFillTool.cpp
--- a/Engine/Plugins/Experimental/MeshModelingToolsetExp/Source/MeshModelingToolsExp/Private/HoleFillTool.cpp
+++ b/Engine/Plugins/Experimental/MeshModelingToolsetExp/Source/MeshModelingToolsExp/Private/HoleFillTool.cpp
@@ -171,7 +171,7 @@ void UHoleFillTool::Setup()
 	SelectionMechanic->Properties->bSelectFaces = false;
 	SelectionMechanic->Properties->bSelectVertices = false;
 	SelectionMechanic->Initialize(OriginalMesh.Get(),
-		TargetComponent->GetWorldTransform(),
+		(FTransform3d)TargetComponent->GetWorldTransform(),
 		GetTargetWorld(),
 		Topology.Get(),
 		[this]() { return &MeshSpatial; }
diff --git a/Engine/Plugins/Experimental/MeshModelingToolsetExp/Source/MeshModelingToolsExp/Private/MeshBoundaryToolBase.cpp b/Engine/Plugins/Experimental/MeshModelingToolsetExp/Source/MeshModelingToolsExp/Private/MeshBoundaryToolBase.cpp
--- a/Engine/Plugins/Experimental/MeshModelingToolsetExp/Source/MeshModelingToolsExp/Private/MeshBoundaryToolBase.cpp
+++ b/Engine/Plugins/Experimental/MeshModelingToolsetExp/Source/MeshModelingToolsExp/Private/MeshBoundaryToolBase.cpp
@@ -45,7 +45,7 @@ void UMeshBoundaryToolBase::Setup()
 	SelectionMechanic->Properties->bSelectFaces = false;
 	SelectionMechanic->Properties->bSelectVertices = false;
 	SelectionMechanic->Initialize(OriginalMesh.Get(),
-		Cast<IPrimitiveComponentBackedTarget>(Target)->GetWorldTransform(),
+		(FTransform3d)Cast<IPrimitiveComponentBackedTarget>(Target)->GetWorldTransform(),
 		GetTargetWorld(),
 		Topology.Get(),
 		[this]() { return &MeshSpatial; }
diff --git a/Engine/Plugins/Experimental/MeshModelingToolsetExp/Source/MeshModelingToolsExp/Private/Properties/MeshAnalysisProperties.cpp b/Engine/Plugins/Experimental/MeshModelingToolsetExp/Source/MeshModelingToolsExp/Private/Properties/MeshAnalysisProperties.cpp
--- a/Engine/Plugins/Experimental/MeshModelingToolsetExp/Source/MeshModelingToolsExp/Private/Properties/MeshAnalysisProperties.cpp
+++ b/Engine/Plugins/Experimental/MeshModelingToolsetExp/Source/MeshModelingToolsExp/Private/Properties/MeshAnalysisProperties.cpp
@@ -54,7 +54,7 @@ FString NumberToStringWithFractionalPrecision(const RealType Number,
 
 void UMeshAnalysisProperties::Update(const FDynamicMesh3& MeshIn, const FTransform& Transform)
 {
-	FTriangleMeshAdapterd TransformedMesh = UE::Geometry::MakeTransformedDynamicMeshAdapter(&MeshIn, Transform);
+	FTriangleMeshAdapterd TransformedMesh = UE::Geometry::MakeTransformedDynamicMeshAdapter(&MeshIn, (FTransform3d)Transform);
 	FVector2d VolArea = TMeshQueries<FTriangleMeshAdapterd>::GetVolumeArea(TransformedMesh);
 	
 	TArray<TTuple<double, int32>> MagnitudeToFractionalDigits;
diff --git a/Engine/Plugins/Runtime/MeshModelingToolset/Source/MeshModelingTools/Private/CombineMeshesTool.cpp b/Engine/Plugins/Runtime/MeshModelingToolset/Source/MeshModelingTools/Private/CombineMeshesTool.cpp
--- a/Engine/Plugins/Runtime/MeshModelingToolset/Source/MeshModelingTools/Private/CombineMeshesTool.cpp
+++ b/Engine/Plugins/Runtime/MeshModelingToolset/Source/MeshModelingTools/Private/CombineMeshesTool.cpp
@@ -173,8 +173,8 @@ void UCombineMeshesTool::CreateNewAsset()
 	AccumulateDMesh.Attributes()->EnableTangents();
 	AccumulateDMesh.Attributes()->EnableMaterialID();
 	AccumulateDMesh.Attributes()->EnablePrimaryColors();
-	FTransform AccumToWorld(Box.GetCenter());
-	FTransform ToAccum(-Box.GetCenter());
+	FTransform3d AccumToWorld(Box.GetCenter());
+	FTransform3d ToAccum(-Box.GetCenter());
 
 	FSimpleShapeSet3d SimpleCollision;
 	UE::Geometry::FComponentCollisionSettings CollisionSettings;
@@ -258,7 +258,7 @@ void UCombineMeshesTool::CreateNewAsset()
 		{
 			// TODO: will need to refactor this when we support duplicating multiple
 			check(Targets.Num() == 1);
-			AccumToWorld = (FTransform)UE::ToolTarget::GetLocalToWorldTransform(Targets[0]);
+			AccumToWorld = (FTransform3d)UE::ToolTarget::GetLocalToWorldTransform(Targets[0]);
 		}
 
 		// max len explicitly enforced here, would ideally notify user
@@ -270,7 +270,7 @@ void UCombineMeshesTool::CreateNewAsset()
 
 		FCreateMeshObjectParams NewMeshObjectParams;
 		NewMeshObjectParams.TargetWorld = GetTargetWorld();
-		NewMeshObjectParams.Transform = AccumToWorld;
+		NewMeshObjectParams.Transform = (FTransform)AccumToWorld;
 		NewMeshObjectParams.BaseName = UseBaseName;
 		NewMeshObjectParams.Materials = AllMaterials;
 		NewMeshObjectParams.SetMesh(&AccumulateDMesh);
diff --git a/Engine/Plugins/Runtime/MeshModelingToolset/Source/MeshModelingTools/Private/EditMeshPolygonsTool.cpp b/Engine/Plugins/Runtime/MeshModelingToolset/Source/MeshModelingTools/Private/EditMeshPolygonsTool.cpp
--- a/Engine/Plugins/Runtime/MeshModelingToolset/Source/MeshModelingTools/Private/EditMeshPolygonsTool.cpp
+++ b/Engine/Plugins/Runtime/MeshModelingToolset/Source/MeshModelingTools/Private/EditMeshPolygonsTool.cpp
@@ -358,7 +358,7 @@ void UEditMeshPolygonsTool::Setup()
 		SelectionMechanic->PolyEdgesRenderer.LineThickness = 1.0;
 	}
 	SelectionMechanic->Initialize(CurrentMesh.Get(),
-		Preview->PreviewMesh->GetTransform(),
+		(FTransform3d)Preview->PreviewMesh->GetTransform(),
 		TargetWorld,
 		Topology.Get(),
 		[this]() { return &GetSpatial(); }
diff --git a/Engine/Plugins/Runtime/MeshModelingToolset/Source/ModelingComponents/Private/Selection/PolygonSelectionMechanic.cpp b/Engine/Plugins/Runtime/MeshModelingToolset/Source/ModelingComponents/Private/Selection/PolygonSelectionMechanic.cpp
--- a/Engine/Plugins/Runtime/MeshModelingToolset/Source/ModelingComponents/Private/Selection/PolygonSelectionMechanic.cpp
+++ b/Engine/Plugins/Runtime/MeshModelingToolset/Source/ModelingComponents/Private/Selection/PolygonSelectionMechanic.cpp
@@ -168,7 +168,7 @@ void UPolygonSelectionMechanic::Initialize(
 		DrawnTriangleSetComponent->RegisterComponent();
 	}
 
-	PreviewGeometryActor->SetActorTransform(TargetTransformIn);
+	PreviewGeometryActor->SetActorTransform((FTransform)TargetTransformIn);
 
 	DrawnTriangleSetComponent->Clear();
 	CurrentlyHighlightedGroups.Empty();
@@ -181,7 +181,7 @@ void UPolygonSelectionMechanic::Initialize(
 {
 
 	Initialize(MeshComponentIn->GetMesh(),
-		MeshComponentIn->GetComponentTransform(),
+		(FTransform3d)MeshComponentIn->GetComponentTransform(),
 		MeshComponentIn->GetWorld(),
 		TopologyIn,
 		GetSpatialSourceFuncIn);
diff --git a/Engine/Source/Runtime/GeometryCore/Public/Quaternion.h b/Engine/Source/Runtime/GeometryCore/Public/Quaternion.h
--- a/Engine/Source/Runtime/GeometryCore/Public/Quaternion.h
+++ b/Engine/Source/Runtime/GeometryCore/Public/Quaternion.h
@@ -98,16 +98,32 @@ PRAGMA_ENABLE_DEPRECATION_WARNINGS
 	//   SetFromRotationMatrix(FMatrix3f)
 
 
-	explicit inline operator FQuat() const
+	explicit inline operator FQuat4f() const
 	{
-		FQuat Quat; 
+		FQuat4f Quat; 
 		Quat.X = (float)X;
 		Quat.Y = (float)Y;
 		Quat.Z = (float)Z;
 		Quat.W = (float)W;
 		return Quat;
 	}
-	explicit inline TQuaternion(const FQuat& Quat)
+	explicit inline operator FQuat4d() const
+	{
+		FQuat4d Quat; 
+		Quat.X = (double)X;
+		Quat.Y = (double)Y;
+		Quat.Z = (double)Z;
+		Quat.W = (double)W;
+		return Quat;
+	}
+	explicit inline TQuaternion(const FQuat4f& Quat)
+	{
+		X = (RealType)Quat.X;
+		Y = (RealType)Quat.Y;
+		Z = (RealType)Quat.Z;
+		W = (RealType)Quat.W;
+	}
+	explicit inline TQuaternion(const FQuat4d& Quat)
 	{
 		X = (RealType)Quat.X;
 		Y = (RealType)Quat.Y;
diff --git a/Engine/Source/Runtime/GeometryCore/Public/TransformTypes.h b/Engine/Source/Runtime/GeometryCore/Public/TransformTypes.h
--- a/Engine/Source/Runtime/GeometryCore/Public/TransformTypes.h
+++ b/Engine/Source/Runtime/GeometryCore/Public/TransformTypes.h
@@ -98,7 +98,7 @@ public:
 	*/
 	operator FTransform3d() const
 	{
-		return FTransform((FQuat4d)Rotation, (FVector3d)Translation, (FVector3d)Scale3D);
+		return FTransform3d((FQuat4d)Rotation, (FVector3d)Translation, (FVector3d)Scale3D);
 	}
 
 	/**
