MeshModelingTools: Fix crash when hitting delete and marquee selecting at the same time in PolyEdit.

#rnx

diff --git a/Engine/Plugins/Runtime/MeshModelingToolset/Source/MeshModelingTools/Private/EditMeshPolygonsTool.cpp b/Engine/Plugins/Runtime/MeshModelingToolset/Source/MeshModelingTools/Private/EditMeshPolygonsTool.cpp
--- a/Engine/Plugins/Runtime/MeshModelingToolset/Source/MeshModelingTools/Private/EditMeshPolygonsTool.cpp
+++ b/Engine/Plugins/Runtime/MeshModelingToolset/Source/MeshModelingTools/Private/EditMeshPolygonsTool.cpp
@@ -659,6 +659,15 @@ void UEditMeshPolygonsTool::RegisterActions(FInteractiveToolActionSet& ActionSet
 
 void UEditMeshPolygonsTool::RequestAction(EEditMeshPolygonsToolActions ActionType)
 {
+	if (SelectionMechanic && SelectionMechanic->IsCurrentlyMarqueeDragging())
+	{
+		PendingAction = EEditMeshPolygonsToolActions::NoAction;
+		GetToolManager()->DisplayMessage(
+			LOCTEXT("CannotActDuringMarquee", "Cannot perform action while marquee selecting"), 
+			EToolMessageLevel::UserWarning);
+		return;
+	}
+
 	if (PendingAction != EEditMeshPolygonsToolActions::NoAction)
 	{
 		return;
diff --git a/Engine/Plugins/Runtime/MeshModelingToolset/Source/ModelingComponents/Private/Selection/PolygonSelectionMechanic.cpp b/Engine/Plugins/Runtime/MeshModelingToolset/Source/ModelingComponents/Private/Selection/PolygonSelectionMechanic.cpp
--- a/Engine/Plugins/Runtime/MeshModelingToolset/Source/ModelingComponents/Private/Selection/PolygonSelectionMechanic.cpp
+++ b/Engine/Plugins/Runtime/MeshModelingToolset/Source/ModelingComponents/Private/Selection/PolygonSelectionMechanic.cpp
@@ -716,6 +716,8 @@ void UPolygonSelectionMechanic::OnUpdateModifierState(int ModifierID, bool bIsOn
 
 void UPolygonSelectionMechanic::OnDragRectangleStarted()
 {
+	bCurrentlyMarqueeDragging = true;
+
 	ParentTool->GetToolManager()->BeginUndoTransaction(LOCTEXT("SelectionChange", "Selection"));
 	BeginChange();
 
@@ -769,6 +771,8 @@ void UPolygonSelectionMechanic::OnDragRectangleChanged(const FCameraRectangle& C
 
 void UPolygonSelectionMechanic::OnDragRectangleFinished(const FCameraRectangle& Rectangle, bool bCancelled)
 {
+	bCurrentlyMarqueeDragging = false;
+
 	TriIsOccludedCache.Reset();
 
 	if (PersistentSelection != PreDragPersistentSelection)
@@ -792,7 +796,10 @@ void UPolygonSelectionMechanic::UpdateMarqueeEnabled()
 
 void UPolygonSelectionMechanic::BeginChange()
 {
-	check(ActiveChange.IsValid() == false);
+	// If you hit this ensure, either you didn't match a BeginChange with an EndChange(), or you
+	// interleaved two actions that both needed a BeginChange(). For instance, are you sure you
+	// are not performing actions while a marquee rectangle is active?
+	ensure(ActiveChange.IsValid() == false);
 	ActiveChange = MakeUnique<FPolygonSelectionMechanicSelectionChange>();
 	ActiveChange->Before = PersistentSelection;
 	ActiveChange->Timestamp = SelectionTimestamp;
@@ -800,10 +807,9 @@ void UPolygonSelectionMechanic::BeginChange()
 
 TUniquePtr<FToolCommandChange> UPolygonSelectionMechanic::EndChange()
 {
-	check(ActiveChange.IsValid() == true);
-	ActiveChange->After = PersistentSelection;
-	if (SelectionTimestamp != ActiveChange->Timestamp)
+	if (ensure(ActiveChange.IsValid()) && SelectionTimestamp != ActiveChange->Timestamp)
 	{
+		ActiveChange->After = PersistentSelection;
 		return MoveTemp(ActiveChange);
 	}
 	ActiveChange = TUniquePtr<FPolygonSelectionMechanicSelectionChange>();
@@ -812,10 +818,9 @@ TUniquePtr<FToolCommandChange> UPolygonSelectionMechanic::EndChange()
 
 bool UPolygonSelectionMechanic::EndChangeAndEmitIfModified()
 {
-	check(ActiveChange.IsValid() == true);
-	ActiveChange->After = PersistentSelection;
-	if (SelectionTimestamp != ActiveChange->Timestamp)
+	if (ensure(ActiveChange.IsValid()) && SelectionTimestamp != ActiveChange->Timestamp)
 	{
+		ActiveChange->After = PersistentSelection;
 		GetParentTool()->GetToolManager()->EmitObjectChange(this, MoveTemp(ActiveChange),
 			LOCTEXT("SelectionChangeMessage", "Selection Change"));
 		return true;
diff --git a/Engine/Plugins/Runtime/MeshModelingToolset/Source/ModelingComponents/Public/Selection/PolygonSelectionMechanic.h b/Engine/Plugins/Runtime/MeshModelingToolset/Source/ModelingComponents/Public/Selection/PolygonSelectionMechanic.h
--- a/Engine/Plugins/Runtime/MeshModelingToolset/Source/ModelingComponents/Public/Selection/PolygonSelectionMechanic.h
+++ b/Engine/Plugins/Runtime/MeshModelingToolset/Source/ModelingComponents/Public/Selection/PolygonSelectionMechanic.h
@@ -351,6 +351,12 @@ public:
 	 */
 	FSimpleMulticastDelegate OnFaceSelectionPreviewChanged;
 
+	/**
+	 * When true, the selection mechanic is currently tracking a marquee rectangle drag, and acting
+	 * on the selection may be unwise until it is over (and an OnSelectionChanged event is fired).
+	 */
+	bool IsCurrentlyMarqueeDragging() { return bCurrentlyMarqueeDragging; }
+
 	// TODO: Is it worth issuing separate callbacks in normal selection changes and in FChange ones, to
 	// allow the user to bundle in some FChanges into the normal callback?
 
@@ -410,6 +416,7 @@ protected:
 	FGroupTopologySelection LastUpdateRectangleSelection;
 	FGroupTopologySelector::FSelectionSettings PreDragTopoSelectorSettings;
 	TMap<int32, bool> TriIsOccludedCache;
+	bool bCurrentlyMarqueeDragging = false;
 
 	FVector3d LastClickedHitPosition;
 	FVector3d LastClickedHitNormal;
