MeshModelingTools: Add a little helper struct to convert lists/sets of eids into stable identifiers and back.

diff --git a/Engine/Plugins/Runtime/MeshModelingToolset/Source/ModelingComponents/Public/Selection/StoredMeshSelectionUtil.h b/Engine/Plugins/Runtime/MeshModelingToolset/Source/ModelingComponents/Public/Selection/StoredMeshSelectionUtil.h
--- a/Engine/Plugins/Runtime/MeshModelingToolset/Source/ModelingComponents/Public/Selection/StoredMeshSelectionUtil.h
+++ b/Engine/Plugins/Runtime/MeshModelingToolset/Source/ModelingComponents/Public/Selection/StoredMeshSelectionUtil.h
@@ -53,6 +53,62 @@ namespace Geometry
 	MODELINGCOMPONENTS_API void ClearActiveToolSelection(UInteractiveToolManager* ToolManager);
 
 
+	/*
+	 * Helper class that can store a list of edges as pairs of triangle IDs and {0,1,2} indices into 
+	 * the triangle edge triplet, because regular FDynamicMesh edge IDs may not stay the same across 
+	 * some mesh operations that preserve the vertex/triangle ID topology. For instance, deleting and 
+	 * reinserting triangles during an undo/redo transaction may preserve all the relevant vids and tids,
+	 * but change the eids of edges even though these edges still exist in the mesh topology.
+	 *
+	 * Note this way of identifying edges will still not be stable in the case of triangles being
+	 * reinserted with rotated indices (e.g. a triangle "a, b, c" reinserted as "b, c, d"). Storing
+	 * Vid pairs would be necessary to be robust to that case, but finding an edge from vertices is
+	 * a slower operation compared to the constant-time tri/index lookup used here.
+	 */
+	class MODELINGCOMPONENTS_API FMeshEdgesFromTriangleSubIndices
+	{
+	public:
+		template <typename EidContainerType>
+		void InitializeFromEdgeIDs(const FDynamicMesh3& Mesh, const  EidContainerType& Eids)
+		{
+			EdgeTriIndexPairs.Reset(Eids.Num());
+			for (int32 Eid : Eids)
+			{
+				if (ensure(Mesh.IsEdge(Eid)))
+				{
+					int32 Tid = Mesh.GetEdgeT(Eid).A;
+					FIndex3i TriEids = Mesh.GetTriEdges(Tid);
+					EdgeTriIndexPairs.Add(TPair<int32, int8>(Tid, IndexUtil::FindTriIndex(Eid, TriEids)));
+				}
+			}
+		}
+
+		template <typename EidContainerType>
+		void GetEdgeIDs(const FDynamicMesh3& Mesh, EidContainerType& EidsOut)
+		{
+			// Can't pass in a desired capacity to Reset() because TSet does not support it
+			EidsOut.Reset(); 
+			for (const TPair<int32, int8>& VidPair : EdgeTriIndexPairs)
+			{
+				if (ensure(Mesh.IsTriangle(VidPair.Key)))
+				{
+					int32 Eid = Mesh.GetTriEdges(VidPair.Key)[VidPair.Value];
+					checkSlow(Eid != IndexConstants::InvalidID);
+					EidsOut.Add(Eid);
+				}
+			}
+		}
+
+		bool IsEmpty() const { return EdgeTriIndexPairs.IsEmpty(); }
+		void Reset() { EdgeTriIndexPairs.Reset(); }
+		void Empty() { EdgeTriIndexPairs.Empty(); }
+
+	private:
+		// The pair is (tid, {0,1,2}), where the second element is an index into the edge triplet
+		// for that triangle.
+		TArray<TPair<int32, int8>> EdgeTriIndexPairs;
+	};
+
 	//
 	// Utility functions for group topology manipulation.
 	// These should probably move to another location
