MeshModelingTools: Make the Add Primitive tools accept/cancel so that the initial click can place the mesh and then the settings can be modified before accepting the tool.

diff --git a/Engine/Plugins/Runtime/MeshModelingToolset/Source/MeshModelingTools/Private/AddPrimitiveTool.cpp b/Engine/Plugins/Runtime/MeshModelingToolset/Source/MeshModelingTools/Private/AddPrimitiveTool.cpp
--- a/Engine/Plugins/Runtime/MeshModelingToolset/Source/MeshModelingTools/Private/AddPrimitiveTool.cpp
+++ b/Engine/Plugins/Runtime/MeshModelingToolset/Source/MeshModelingTools/Private/AddPrimitiveTool.cpp
@@ -1,11 +1,15 @@
 // Copyright Epic Games, Inc. All Rights Reserved.
 
 #include "AddPrimitiveTool.h"
+
+#include "BaseGizmos/TransformGizmoUtil.h"
+#include "BaseGizmos/TransformProxy.h"
 #include "ToolBuilderUtil.h"
 #include "InteractiveToolManager.h"
 #include "SceneQueries/SceneSnappingManager.h"
 #include "BaseBehaviors/MouseHoverBehavior.h"
 #include "Selection/ToolSelectionUtil.h"
+#include "Mechanics/DragAlignmentMechanic.h"
 #include "ModelingObjectsCreationAPI.h"
 #include "ToolSceneQueriesUtil.h"
 #include "ToolSetupUtil.h"
@@ -78,29 +82,6 @@ UInteractiveTool* UAddPrimitiveToolBuilder::BuildTool(const FToolBuilderState& S
 	return NewTool;
 }
 
-bool
-UProceduralShapeToolProperties::IsEquivalent( const UProceduralShapeToolProperties* Other ) const
-{
-#if WITH_EDITOR
-	UClass* Class = GetClass();
-	if ( Other->GetClass() != Class )
-	{
-		return false;
-	}
-	for (const FProperty* Prop : TFieldRange<FProperty>(Class))
-	{
-		if (Prop->HasMetaData(TEXT("ProceduralShapeSetting")) &&
-			(!Prop->Identical_InContainer(this, Other)))
-		{
-			return false;
-		}
-	}
-	return true;
-#else
-	return false;
-#endif
-}
-
 void UAddPrimitiveTool::SetWorld(UWorld* World)
 {
 	this->TargetWorld = World;
@@ -113,6 +94,11 @@ UAddPrimitiveTool::UAddPrimitiveTool(const FObjectInitializer&)
 	ShapeSettings->ClearFlags(RF_Transactional);
 }
 
+bool UAddPrimitiveTool::CanAccept() const
+{
+	return CurrentState == EState::AdjustingSettings;
+}
+
 void UAddPrimitiveTool::Setup()
 {
 	USingleClickTool::Setup();
@@ -142,28 +128,79 @@ void UAddPrimitiveTool::Setup()
 	PreviewMesh->SetMaterial(MaterialProperties->Material.Get());
 	PreviewMesh->EnableWireframe(MaterialProperties->bShowWireframe);
 
+	UTransformProxy* TransformProxy = NewObject<UTransformProxy>(this);
+	TransformProxy->OnTransformChanged.AddWeakLambda(this, [this](UTransformProxy*, FTransform NewTransform) 
+		{
+			PreviewMesh->SetTransform(NewTransform);
+		});
+
+	// TODO: It might be nice to use a repositionable gizmo, but the drag alignment mechanic can't currently 
+	// hit the preview mesh, which makes middle click repositioning feel broken and not very useful.
+	Gizmo = UE::TransformGizmoUtil::CreateCustomTransformGizmo(GetToolManager(),
+		ETransformGizmoSubElements::StandardTranslateRotate, this);
+	Gizmo->SetActiveTarget(TransformProxy, GetToolManager());
+
+	DragAlignmentMechanic = NewObject<UDragAlignmentMechanic>(this);
+	DragAlignmentMechanic->Setup(this);
+	DragAlignmentMechanic->AddToGizmo(Gizmo);
+
 	UpdatePreviewMesh();
 
-	GetToolManager()->DisplayMessage(
-		LOCTEXT("OnStartAddPrimitiveTool", "This Tool creates new shapes. Configure the shape via its settings, position it by moving the mouse in the scene, and drop it as a new object or instance by left-clicking."),
-		EToolMessageLevel::UserNotification);
+	SetState(EState::PlacingPrimitive);
+}
+
+void UAddPrimitiveTool::SetState(EState NewState)
+{
+	CurrentState = NewState;
+
+	bool bGizmoActive = (CurrentState == EState::AdjustingSettings);
+	Gizmo->SetVisibility(bGizmoActive && ShapeSettings->bShowGizmo);
+	ShapeSettings->bShowGizmoOptions = bGizmoActive;
+	NotifyOfPropertyChangeByTool(ShapeSettings);
+
+	if (CurrentState == EState::PlacingPrimitive)
+	{
+		GetToolManager()->DisplayMessage(
+			LOCTEXT("OnStartAddPrimitiveTool", "This Tool creates new shapes. Click in the scene to choose initial placement of mesh."),
+			EToolMessageLevel::UserNotification);
+	}
+	else
+	{
+		// Initialize gizmo to current preview location
+		Gizmo->ReinitializeGizmoTransform(PreviewMesh->GetTransform());
+
+		GetToolManager()->DisplayMessage(
+			LOCTEXT("OnStartAddPrimitiveTool", "Alter shape settings in the detail panel or modify placement with gizmo, then accept the tool."),
+			EToolMessageLevel::UserNotification);
+	}
 }
 
 
 void UAddPrimitiveTool::Shutdown(EToolShutdownType ShutdownType)
 {
+	if (ShutdownType == EToolShutdownType::Accept)
+	{
+		GenerateAsset();
+	}
+
+	DragAlignmentMechanic->Shutdown();
+	DragAlignmentMechanic = nullptr;
+	GetToolManager()->GetPairedGizmoManager()->DestroyAllGizmosByOwner(this);
+	Gizmo = nullptr;
+
 	PreviewMesh->SetVisible(false);
 	PreviewMesh->Disconnect();
 	PreviewMesh = nullptr;
 
 	OutputTypeProperties->SaveProperties(this);
 	ShapeSettings->SaveProperties(this);
-	MaterialProperties->SaveProperties(this);	
+	MaterialProperties->SaveProperties(this);
 }
 
 
 void UAddPrimitiveTool::Render(IToolsContextRenderAPI* RenderAPI)
 {
+	DragAlignmentMechanic->Render(RenderAPI);
 }
 
 
@@ -179,6 +216,11 @@ void UAddPrimitiveTool::OnPropertyModified(UObject* PropertySet, FProperty* Prop
 		PreviewMesh->SetMaterial(MaterialProperties->Material.Get());
 		UpdatePreviewMesh();
 	}
+
+	if (Gizmo)
+	{
+		Gizmo->SetVisibility(CurrentState == EState::AdjustingSettings && ShapeSettings->bShowGizmo);
+	}
 }
 
 
@@ -186,7 +228,9 @@ void UAddPrimitiveTool::OnPropertyModified(UObject* PropertySet, FProperty* Prop
 
 FInputRayHit UAddPrimitiveTool::BeginHoverSequenceHitTest(const FInputDeviceRay& PressPos)
 {
-	return FInputRayHit(0.0f);		// always hit in hover 
+	FInputRayHit Result(0);
+	Result.bHit = (CurrentState == EState::PlacingPrimitive);
+	return Result;
 }
 
 void UAddPrimitiveTool::OnBeginHover(const FInputDeviceRay& DevicePos)
@@ -321,35 +365,33 @@ void UAddPrimitiveTool::UpdatePreviewMesh() const
 	checkSlow(CalculateTangentsSuccessful);
 }
 
-void UAddPrimitiveTool::OnClicked(const FInputDeviceRay& ClickPos)
+FInputRayHit UAddPrimitiveTool::IsHitByClick(const FInputDeviceRay& ClickPos)
 {
-	UMaterialInterface* Material = PreviewMesh->GetMaterial();
-
-	if (ShapeSettings->bInstanceIfPossible && LastGenerated != nullptr && IsEquivalentLastGeneratedAsset())
-	{
-		GetToolManager()->BeginUndoTransaction(LOCTEXT("AddPrimitiveToolTransactionName", "Add Shape Mesh"));
-		FActorSpawnParameters SpawnParameters;
-		SpawnParameters.Template = LastGenerated->Actor;
-		FRotator Rotation(0.0f, 0.0f, 0.0f);
-		AStaticMeshActor* CloneActor = TargetWorld->SpawnActor<AStaticMeshActor>(FVector::ZeroVector, Rotation, SpawnParameters);
-		// some properties must be manually set on the component because they will not persist reliably through the spawn template (especially if the actor creation was undone)
-		CloneActor->GetStaticMeshComponent()->SetWorldTransform(PreviewMesh->GetTransform());
-		CloneActor->GetStaticMeshComponent()->SetStaticMesh(LastGenerated->StaticMesh);
-		CloneActor->GetStaticMeshComponent()->SetMaterial(0, Material);
-#if WITH_EDITOR
-		CloneActor->SetActorLabel(LastGenerated->Label);
-#endif
-		// select newly-created object
-		ToolSelectionUtil::SetNewActorSelection(GetToolManager(), CloneActor);
-		GetToolManager()->EndUndoTransaction();
+	FInputRayHit Result(0);
+	Result.bHit = (CurrentState == EState::PlacingPrimitive);
+	return Result;
+}
 
+void UAddPrimitiveTool::OnClicked(const FInputDeviceRay& ClickPos)
+{
+	if (!ensure(CurrentState == EState::PlacingPrimitive))
+	{
 		return;
 	}
-	LastGenerated = nullptr;
+
+	UpdatePreviewPosition(ClickPos);
+	SetState(EState::AdjustingSettings);
+	GetToolManager()->EmitObjectChange(this, MakeUnique<FStateChange>(PreviewMesh->GetTransform()),
+		LOCTEXT("PlaceMeshTransaction", "Place Mesh"));
+}
+
+void UAddPrimitiveTool::GenerateAsset()
+{
+	UMaterialInterface* Material = PreviewMesh->GetMaterial();
 
 	const FDynamicMesh3* CurMesh = PreviewMesh->GetPreviewDynamicMesh();
 
-	GetToolManager()->BeginUndoTransaction(LOCTEXT("AddPrimitiveToolTransactionName", "Add Shape Mesh"));
+	GetToolManager()->BeginUndoTransaction(LOCTEXT("AddPrimitiveToolTransactionName", "Add Primitive Tool"));
 
 	FCreateMeshObjectParams NewMeshObjectParams;
 	NewMeshObjectParams.TargetWorld = TargetWorld;
@@ -363,18 +405,6 @@ void UAddPrimitiveTool::OnClicked(const FInputDeviceRay& ClickPos)
 	{
 		if (Result.NewActor != nullptr)
 		{
-			if (Cast<AStaticMeshActor>(Result.NewActor) != nullptr)
-			{
-				LastGenerated = NewObject<ULastActorInfo>(this);
-				LastGenerated->ShapeSettings = DuplicateObject(ShapeSettings, nullptr);
-				LastGenerated->MaterialProperties = DuplicateObject(MaterialProperties, nullptr);
-				LastGenerated->Actor = Result.NewActor;
-				LastGenerated->StaticMesh = CastChecked<AStaticMeshActor>(LastGenerated->Actor)->GetStaticMeshComponent()->GetStaticMesh();
-#if WITH_EDITOR
-				LastGenerated->Label = LastGenerated->Actor->GetActorLabel();
-#endif
-			}
-
 			// select newly-created object
 			ToolSelectionUtil::SetNewActorSelection(GetToolManager(), Result.NewActor);
 		}
@@ -383,6 +413,22 @@ void UAddPrimitiveTool::OnClicked(const FInputDeviceRay& ClickPos)
 	GetToolManager()->EndUndoTransaction();
 }
 
+void UAddPrimitiveTool::FStateChange::Apply(UObject* Object)
+{
+	UAddPrimitiveTool* Tool = Cast<UAddPrimitiveTool>(Object);
+
+	// Set preview transform before changing state so that the adjustment gizmo is initialized properly
+	Tool->PreviewMesh->SetTransform(MeshTransform);
+
+	Tool->SetState(EState::AdjustingSettings);
+}
+
+void UAddPrimitiveTool::FStateChange::Revert(UObject* Object)
+{
+	UAddPrimitiveTool* Tool = Cast<UAddPrimitiveTool>(Object);
+	Tool->SetState(EState::PlacingPrimitive);
+}
+
 
 UAddBoxPrimitiveTool::UAddBoxPrimitiveTool(const FObjectInitializer& ObjectInitializer)
 	: Super(ObjectInitializer.SetDefaultSubobjectClass<UProceduralBoxToolProperties>(TEXT("ShapeSettings")))
diff --git a/Engine/Plugins/Runtime/MeshModelingToolset/Source/MeshModelingTools/Public/AddPrimitiveTool.h b/Engine/Plugins/Runtime/MeshModelingToolset/Source/MeshModelingTools/Public/AddPrimitiveTool.h
--- a/Engine/Plugins/Runtime/MeshModelingToolset/Source/MeshModelingTools/Public/AddPrimitiveTool.h
+++ b/Engine/Plugins/Runtime/MeshModelingToolset/Source/MeshModelingTools/Public/AddPrimitiveTool.h
@@ -13,6 +13,8 @@
 #include "AddPrimitiveTool.generated.h"
 
 PREDECLARE_USE_GEOMETRY_CLASS(FDynamicMesh3);
+class UCombinedTransformGizmo;
+class UDragAlignmentMechanic;
 
 /**
  * Builder
@@ -77,10 +79,6 @@ class MESHMODELINGTOOLS_API UProceduralShapeToolProperties : public UInteractive
 	GENERATED_BODY()
 
 public:
-	/** Create instances of the last created asset rather than creating a whole new asset, provided the Shape settings have not changed.
-	  * If false, all created actors will have separate underlying mesh assets. */
-	UPROPERTY(EditAnywhere, Category = Asset)
-	bool bInstanceIfPossible = false;
 
 	/** How Polygroups are assigned to shape primitives. */
 	UPROPERTY(EditAnywhere, Category = Shape, meta = (ProceduralShapeSetting))
@@ -106,7 +104,13 @@ public:
 	UPROPERTY(EditAnywhere, Category = Positioning, meta = (EditCondition = "TargetSurface == EMakeMeshPlacementType::OnScene"))
 	bool bAlignToNormal = true;
 
-	bool IsEquivalent( const UProceduralShapeToolProperties* ) const;
+	/** Show a gizmo to allow the mesh to be repositioned after the initial placement click. */
+	UPROPERTY(EditAnywhere, Category = Positioning, meta = (EditCondition = "bShowGizmoOptions", EditConditionHides, HideEditConditionToggle))
+	bool bShowGizmo = true;
+
+	//~ Not user visible- used to hide the bShowGizmo option when not yet placed mesh.
+	UPROPERTY(meta = (TransientToolProperty))
+	bool bShowGizmoOptions = false;
 };
 
 UCLASS()
@@ -427,32 +431,6 @@ public:
 };
 
 
-UCLASS(Transient)
-class MESHMODELINGTOOLS_API ULastActorInfo : public UObject
-{
-	GENERATED_BODY()
-
-public:
-	FString Label = "";
-
-	UPROPERTY()
-	TObjectPtr<AActor> Actor = nullptr;
-
-	UPROPERTY()
-	TObjectPtr<UStaticMesh> StaticMesh = nullptr;
-
-	UPROPERTY()
-	TObjectPtr<UProceduralShapeToolProperties> ShapeSettings;
-
-	UPROPERTY()
-	TObjectPtr<UNewMeshMaterialProperties> MaterialProperties;
-
-	bool IsInvalid() const
-	{
-		return Actor == nullptr || StaticMesh == nullptr || ShapeSettings == nullptr || MaterialProperties == nullptr;
-	}
-};
-
 /**
  * Base tool to create primitives
  */
@@ -471,14 +449,15 @@ public:
 
 	virtual void Render(IToolsContextRenderAPI* RenderAPI) override;
 
-	virtual bool HasCancel() const override { return false; }
-	virtual bool HasAccept() const override { return false; }
-	virtual bool CanAccept() const override { return false; }
+	virtual bool HasCancel() const override { return true; }
+	virtual bool HasAccept() const override { return true; }
+	virtual bool CanAccept() const override;
 
 	virtual void OnPropertyModified(UObject* PropertySet, FProperty* Property) override;
 
+	// USingleClickTool
 	virtual void OnClicked(const FInputDeviceRay& ClickPos) override;
-
+	virtual FInputRayHit IsHitByClick(const FInputDeviceRay& ClickPos) override;
 
 	// IHoverBehaviorTarget interface
 	virtual FInputRayHit BeginHoverSequenceHitTest(const FInputDeviceRay& PressPos) override;
@@ -488,9 +467,20 @@ public:
 
 
 protected:
+	enum class EState
+	{
+		PlacingPrimitive,
+		AdjustingSettings
+	};
+
+	EState CurrentState = EState::PlacingPrimitive;
+	void SetState(EState NewState);
+
 	virtual void GenerateMesh(FDynamicMesh3* OutMesh) const {}
 	virtual UProceduralShapeToolProperties* CreateShapeSettings(){return nullptr;}
 
+	virtual void GenerateAsset();
+
 	/** Property set for type of output object (StaticMesh, Volume, etc) */
 	UPROPERTY()
 	TObjectPtr<UCreateMeshObjectTypeProperties> OutputTypeProperties;
@@ -501,27 +491,14 @@ protected:
 	UPROPERTY()
 	TObjectPtr<UNewMeshMaterialProperties> MaterialProperties;
 
-
-	/**
-	 * Checks if the passed-in settings would create the same asset as the current settings
-	 */
-	bool IsEquivalentLastGeneratedAsset() const
-	{
-		if (LastGenerated == nullptr || LastGenerated->IsInvalid())
-		{
-			return false;
-		}
-		return (LastGenerated->MaterialProperties->UVScale == MaterialProperties->UVScale) &&
-			(LastGenerated->MaterialProperties->bWorldSpaceUVScale == MaterialProperties->bWorldSpaceUVScale) &&
-			ShapeSettings->IsEquivalent(LastGenerated->ShapeSettings);
-	}
-
-
 	UPROPERTY()
 	TObjectPtr<UPreviewMesh> PreviewMesh;
 
 	UPROPERTY()
-	TObjectPtr<ULastActorInfo> LastGenerated;
+	TObjectPtr<UCombinedTransformGizmo> Gizmo = nullptr;
+
+	UPROPERTY()
+	TObjectPtr<UDragAlignmentMechanic> DragAlignmentMechanic = nullptr;
 
 	UPROPERTY()
 	FString AssetName = TEXT("GeneratedAsset");
@@ -532,6 +509,26 @@ protected:
 	UE::Geometry::FFrame3d ShapeFrame;
 
 	void UpdatePreviewMesh() const;
+
+	// Used to make the initial placement of the mesh undoable
+	class FStateChange : public FToolCommandChange
+	{
+	public:
+		FStateChange(const FTransform& MeshTransformIn)
+			: MeshTransform(MeshTransformIn)
+		{
+		}
+
+		virtual void Apply(UObject* Object) override;
+		virtual void Revert(UObject* Object) override;
+		virtual FString ToString() const override
+		{
+			return TEXT("UAddPrimitiveTool::FStateChange");
+		}
+
+	protected:
+		FTransform MeshTransform;
+	};
 };
 
 
