Re-add a few symbol exports that were lost in UE5 branching.
#rnx

diff --git a/Engine/Plugins/Runtime/MeshModelingToolset/Source/ModelingComponents/Public/AssetUtils/MeshDescriptionUtil.h b/Engine/Plugins/Runtime/MeshModelingToolset/Source/ModelingComponents/Public/AssetUtils/MeshDescriptionUtil.h
--- a/Engine/Plugins/Runtime/MeshModelingToolset/Source/ModelingComponents/Public/AssetUtils/MeshDescriptionUtil.h
+++ b/Engine/Plugins/Runtime/MeshModelingToolset/Source/ModelingComponents/Public/AssetUtils/MeshDescriptionUtil.h
@@ -46,7 +46,7 @@ namespace UE
 		//
 
 		/** How to update a boolean build setting */
-		enum class EBuildSettingBoolChange
+		enum class MODELINGCOMPONENTS_API EBuildSettingBoolChange
 		{
 			Disable,
 			Enable,
@@ -54,7 +54,7 @@ namespace UE
 		};
 
 		/** Set of changes to apply to StaticMesh Build Settings */
-		struct FStaticMeshBuildSettingChange
+		struct MODELINGCOMPONENTS_API FStaticMeshBuildSettingChange
 		{
 			EBuildSettingBoolChange AutoGeneratedNormals = EBuildSettingBoolChange::NoChange;
 			EBuildSettingBoolChange AutoGeneratedTangents = EBuildSettingBoolChange::NoChange;
diff --git a/Engine/Plugins/Runtime/MeshModelingToolset/Source/ModelingComponents/Public/Physics/ComponentCollisionUtil.h b/Engine/Plugins/Runtime/MeshModelingToolset/Source/ModelingComponents/Public/Physics/ComponentCollisionUtil.h
--- a/Engine/Plugins/Runtime/MeshModelingToolset/Source/ModelingComponents/Public/Physics/ComponentCollisionUtil.h
+++ b/Engine/Plugins/Runtime/MeshModelingToolset/Source/ModelingComponents/Public/Physics/ComponentCollisionUtil.h
@@ -20,7 +20,7 @@ struct FSimpleShapeSet3d;
 /**
  * Component/BodySetup collision settings (eg StaticMeshComponent) we might need to pass through the functions below
  */
-struct FComponentCollisionSettings
+struct MODELINGCOMPONENTS_API FComponentCollisionSettings
 {
 	int32 CollisionTypeFlag = 0;		// this is ECollisionTraceFlag
 	bool bIsGeneratedCollision = true;
