ModelingComponents: pass visibility filter to SpatialCache query in ModelingSceneSnappingManager snap query, so that hidden objects are not included in the set of valid targets
#rnx

diff --git a/Engine/Plugins/Runtime/MeshModelingToolset/Source/ModelingComponents/Private/Snapping/ModelingSceneSnappingManager.cpp b/Engine/Plugins/Runtime/MeshModelingToolset/Source/ModelingComponents/Private/Snapping/ModelingSceneSnappingManager.cpp
--- a/Engine/Plugins/Runtime/MeshModelingToolset/Source/ModelingComponents/Private/Snapping/ModelingSceneSnappingManager.cpp
+++ b/Engine/Plugins/Runtime/MeshModelingToolset/Source/ModelingComponents/Private/Snapping/ModelingSceneSnappingManager.cpp
@@ -542,7 +542,7 @@ bool UModelingSceneSnappingManager::ExecuteSceneSnapQueryPosition(const FSceneSn
 	// cast ray into scene cache
 	FSceneGeometryPoint HitPoint;
 	FSceneGeometryID HitIdentifier;
-	bool bHitSceneCache = SpatialCache->FindNearestHit(WorldRay, HitPoint, HitIdentifier);
+	bool bHitSceneCache = SpatialCache->FindNearestHit(WorldRay, HitPoint, HitIdentifier, &TmpFilter);
 
 	// cast ray into world
 	FHitResult HitResult;
