MeshModelingTools: Allow user to set specific extrude distance in PolyEd.

diff --git a/Engine/Plugins/Runtime/MeshModelingToolset/Source/MeshModelingTools/Private/EditMeshPolygonsTool.cpp b/Engine/Plugins/Runtime/MeshModelingToolset/Source/MeshModelingTools/Private/EditMeshPolygonsTool.cpp
--- a/Engine/Plugins/Runtime/MeshModelingToolset/Source/MeshModelingTools/Private/EditMeshPolygonsTool.cpp
+++ b/Engine/Plugins/Runtime/MeshModelingToolset/Source/MeshModelingTools/Private/EditMeshPolygonsTool.cpp
@@ -1075,7 +1075,7 @@ void UEditMeshPolygonsTool::StartActivity(TObjectPtr<UInteractiveToolActivity> A
 		SelectionMechanic->SetIsEnabled(false);
 		SetToolPropertySourceEnabled(SelectionMechanic->Properties, false);
 		CurrentActivity = Activity;
-		if (CurrentActivity == BevelEdgeActivity)
+		if (CurrentActivity->HasAccept())
 		{
 			SetToolPropertySourceEnabled(AcceptCancelAction, true);
 		}
diff --git a/Engine/Plugins/Runtime/MeshModelingToolset/Source/MeshModelingTools/Private/ToolActivities/PolyEditBevelEdgeActivity.h b/Engine/Plugins/Runtime/MeshModelingToolset/Source/MeshModelingTools/Private/ToolActivities/PolyEditBevelEdgeActivity.h
--- a/Engine/Plugins/Runtime/MeshModelingToolset/Source/MeshModelingTools/Private/ToolActivities/PolyEditBevelEdgeActivity.h
+++ b/Engine/Plugins/Runtime/MeshModelingToolset/Source/MeshModelingTools/Private/ToolActivities/PolyEditBevelEdgeActivity.h
@@ -48,6 +48,7 @@ public:
 	virtual bool CanStart() const override;
 	virtual EToolActivityStartResult Start() override;
 	virtual bool IsRunning() const override { return bIsRunning; }
+	virtual bool HasAccept() const override { return true; }
 	virtual bool CanAccept() const override;
 	virtual EToolActivityEndResult End(EToolShutdownType) override;
 	virtual void Render(IToolsContextRenderAPI* RenderAPI) override;
diff --git a/Engine/Plugins/Runtime/MeshModelingToolset/Source/MeshModelingTools/Private/ToolActivities/PolyEditExtrudeActivity.cpp b/Engine/Plugins/Runtime/MeshModelingToolset/Source/MeshModelingTools/Private/ToolActivities/PolyEditExtrudeActivity.cpp
--- a/Engine/Plugins/Runtime/MeshModelingToolset/Source/MeshModelingTools/Private/ToolActivities/PolyEditExtrudeActivity.cpp
+++ b/Engine/Plugins/Runtime/MeshModelingToolset/Source/MeshModelingTools/Private/ToolActivities/PolyEditExtrudeActivity.cpp
@@ -37,6 +37,44 @@ namespace PolyEditExtrudeActivityLocals
 	{
 		return static_cast<FExtrudeOp::EDirectionMode>(static_cast<int>(Value));
 	}
+
+	EPolyEditExtrudeDistanceMode GetDistanceMode(UPolyEditExtrudeActivity& Activity, UPolyEditExtrudeActivity::EPropertySetToUse PropertySetToUse)
+	{
+		switch (PropertySetToUse)
+		{
+		case UPolyEditExtrudeActivity::EPropertySetToUse::Extrude:
+			return Activity.ExtrudeProperties->DistanceMode;
+			break;
+		case UPolyEditExtrudeActivity::EPropertySetToUse::Offset:
+			return Activity.OffsetProperties->DistanceMode;
+			break;
+		case UPolyEditExtrudeActivity::EPropertySetToUse::PushPull:
+			return Activity.PushPullProperties->DistanceMode;
+			break;
+		}
+
+		ensure(false);
+		return Activity.ExtrudeProperties->DistanceMode;
+	}
+
+	double GetFixedDistance(UPolyEditExtrudeActivity& Activity, UPolyEditExtrudeActivity::EPropertySetToUse PropertySetToUse)
+	{
+		switch (PropertySetToUse)
+		{
+		case UPolyEditExtrudeActivity::EPropertySetToUse::Extrude:
+			return Activity.ExtrudeProperties->Distance;
+			break;
+		case UPolyEditExtrudeActivity::EPropertySetToUse::Offset:
+			return Activity.OffsetProperties->Distance;
+			break;
+		case UPolyEditExtrudeActivity::EPropertySetToUse::PushPull:
+			return Activity.PushPullProperties->Distance;
+			break;
+		}
+
+		ensure(false);
+		return Activity.ExtrudeProperties->Distance;
+	}
 }
 
 TUniquePtr<FDynamicMeshOperator> UPolyEditExtrudeActivity::MakeNewOperator()
@@ -70,8 +108,9 @@ TUniquePtr<FDynamicMeshOperator> UPolyEditExtrudeActivity::MakeNewOperator()
 	}
 
 	ActivityContext->CurrentTopology->GetSelectedTriangles(ActiveSelection, Op->TriangleSelection);
-
-	Op->ExtrudeDistance = ExtrudeHeightMechanic->CurrentHeight;
+	
+	Op->ExtrudeDistance = GetDistanceMode(*this, PropertySetToUse) == EPolyEditExtrudeDistanceMode::Fixed ?
+		GetFixedDistance(*this, PropertySetToUse) : ExtrudeHeightMechanic->CurrentHeight;
 	Op->UVScaleFactor = UVScaleFactor;
 
 	FTransform3d WorldTransform(ActivityContext->Preview->PreviewMesh->GetTransform());
@@ -114,6 +153,21 @@ void UPolyEditExtrudeActivity::Setup(UInteractiveTool* ParentToolIn)
 			ActivityContext->Preview->InvalidateResult();
 		}
 	});
+	ExtrudeProperties->WatchProperty(ExtrudeProperties->DistanceMode,
+	[this](EPolyEditExtrudeDistanceMode) {
+		if (bIsRunning)
+		{
+			ReinitializeExtrudeHeightMechanic();
+			ActivityContext->Preview->InvalidateResult();
+		}
+	});
+	ExtrudeProperties->WatchProperty(ExtrudeProperties->Distance,
+	[this](double) {
+		if (bIsRunning)
+		{
+			ActivityContext->Preview->InvalidateResult();
+		}
+	});
 
 	OffsetProperties = NewObject<UPolyEditOffsetProperties>();
 	OffsetProperties->RestoreProperties(ParentTool.Get());
@@ -135,6 +189,21 @@ void UPolyEditExtrudeActivity::Setup(UInteractiveTool* ParentToolIn)
 			ActivityContext->Preview->InvalidateResult();
 		}
 	});
+	OffsetProperties->WatchProperty(OffsetProperties->DistanceMode,
+	[this](EPolyEditExtrudeDistanceMode) {
+		if (bIsRunning)
+		{
+			ReinitializeExtrudeHeightMechanic();
+			ActivityContext->Preview->InvalidateResult();
+		}
+	});
+	OffsetProperties->WatchProperty(OffsetProperties->Distance,
+	[this](double) {
+		if (bIsRunning)
+		{
+			ActivityContext->Preview->InvalidateResult();
+		}
+	});
 
 	PushPullProperties = NewObject<UPolyEditPushPullProperties>();
 	PushPullProperties->RestoreProperties(ParentTool.Get());
@@ -156,6 +225,21 @@ void UPolyEditExtrudeActivity::Setup(UInteractiveTool* ParentToolIn)
 			ActivityContext->Preview->InvalidateResult();
 		}
 	});
+	PushPullProperties->WatchProperty(PushPullProperties->DistanceMode,
+	[this](EPolyEditExtrudeDistanceMode) {
+		if (bIsRunning)
+		{
+			ReinitializeExtrudeHeightMechanic();
+			ActivityContext->Preview->InvalidateResult();
+		}
+	});
+	PushPullProperties->WatchProperty(PushPullProperties->Distance,
+	[this](double) {
+		if (bIsRunning)
+		{
+			ActivityContext->Preview->InvalidateResult();
+		}
+	});
 
 	// Register ourselves to receive clicks and hover
 	USingleClickInputBehavior* ClickBehavior = NewObject<USingleClickInputBehavior>();
@@ -198,6 +282,8 @@ bool UPolyEditExtrudeActivity::CanStart() const
 
 EToolActivityStartResult UPolyEditExtrudeActivity::Start()
 {
+	using namespace PolyEditExtrudeActivityLocals;
+
 	if (!CanStart())
 	{
 		ParentTool->GetToolManager()->DisplayMessage(
@@ -239,7 +325,7 @@ EToolActivityStartResult UPolyEditExtrudeActivity::Start()
 	{
 		return ToolSceneQueriesUtil::FindWorldGridSnapPoint(ParentTool.Get(), WorldPos, SnapPos);
 	};
-	ExtrudeHeightMechanic->CurrentHeight = 1.0f; // Arbitrary intialization
+	ExtrudeHeightMechanic->CurrentHeight = GetFixedDistance(*this, PropertySetToUse);
 	
 	BeginExtrude();
 
@@ -278,6 +364,7 @@ void UPolyEditExtrudeActivity::BeginExtrude()
 	ActiveSelectionFrameWorld.Transform(WorldTransform);
 
 	ReinitializeExtrudeHeightMechanic();
+	ActivityContext->Preview->InvalidateResult();
 
 	float BoundsMaxDim = ActivityContext->CurrentMesh->GetBounds().MaxDim();
 	if (BoundsMaxDim > 0)
@@ -289,7 +376,14 @@ void UPolyEditExtrudeActivity::BeginExtrude()
 // The height mechanics has to get reinitialized whenever extrude direction changes
 void UPolyEditExtrudeActivity::ReinitializeExtrudeHeightMechanic()
 {
-	ExtrusionFrameWorld = ActiveSelectionFrameWorld;
+	using namespace PolyEditExtrudeActivityLocals;
+
+	if (GetDistanceMode(*this, PropertySetToUse) != EPolyEditExtrudeDistanceMode::ClickInViewport)
+	{
+		return;
+	}
+
+	FFrame3d ExtrusionFrameWorld = ActiveSelectionFrameWorld;
 	ExtrusionFrameWorld.AlignAxis(2, GetExtrudeDirection());
 
 	FDynamicMesh3 ExtrudeHitTargetMesh;	
@@ -497,7 +591,11 @@ FVector3d UPolyEditExtrudeActivity::GetExtrudeDirection() const
 
 void UPolyEditExtrudeActivity::Render(IToolsContextRenderAPI* RenderAPI)
 {
-	if (ensure(bIsRunning) && ExtrudeHeightMechanic != nullptr)
+	using namespace PolyEditExtrudeActivityLocals;
+
+	if (ensure(bIsRunning) 
+		&& ExtrudeHeightMechanic != nullptr
+		&& GetDistanceMode(*this, PropertySetToUse) == EPolyEditExtrudeDistanceMode::ClickInViewport)
 	{
 		ExtrudeHeightMechanic->Render(RenderAPI);
 	}
@@ -524,8 +622,11 @@ void UPolyEditExtrudeActivity::Tick(float DeltaTime)
 
 FInputRayHit UPolyEditExtrudeActivity::IsHitByClick(const FInputDeviceRay& ClickPos)
 {
+	using namespace PolyEditExtrudeActivityLocals;
+
 	FInputRayHit OutHit;
-	OutHit.bHit = bIsRunning;
+	OutHit.bHit = bIsRunning 
+		&& GetDistanceMode(*this, PropertySetToUse) == EPolyEditExtrudeDistanceMode::ClickInViewport;
 	return OutHit;
 }
 
@@ -539,14 +640,16 @@ void UPolyEditExtrudeActivity::OnClicked(const FInputDeviceRay& ClickPos)
 
 FInputRayHit UPolyEditExtrudeActivity::BeginHoverSequenceHitTest(const FInputDeviceRay& PressPos)
 {
+	using namespace PolyEditExtrudeActivityLocals;
+
 	FInputRayHit OutHit;
-	OutHit.bHit = bIsRunning;
+	OutHit.bHit = bIsRunning && GetDistanceMode(*this, PropertySetToUse) == EPolyEditExtrudeDistanceMode::ClickInViewport;
 	return OutHit;
 }
 
 bool UPolyEditExtrudeActivity::OnUpdateHover(const FInputDeviceRay& DevicePos)
 {
-	if (!bRequestedApply)
+	if (bIsRunning && !bRequestedApply)
 	{
 		ExtrudeHeightMechanic->UpdateCurrentDistance(DevicePos.WorldRay);
 		ActivityContext->Preview->InvalidateResult();
diff --git a/Engine/Plugins/Runtime/MeshModelingToolset/Source/MeshModelingTools/Public/ToolActivities/PolyEditExtrudeActivity.h b/Engine/Plugins/Runtime/MeshModelingToolset/Source/MeshModelingTools/Public/ToolActivities/PolyEditExtrudeActivity.h
--- a/Engine/Plugins/Runtime/MeshModelingToolset/Source/MeshModelingTools/Public/ToolActivities/PolyEditExtrudeActivity.h
+++ b/Engine/Plugins/Runtime/MeshModelingToolset/Source/MeshModelingTools/Public/ToolActivities/PolyEditExtrudeActivity.h
@@ -30,6 +30,20 @@ enum class EPolyEditExtrudeDirection
 	LocalZ
 };
 
+UENUM()
+enum class EPolyEditExtrudeDistanceMode
+{
+	/** Set distance by clicking in the viewport. */
+	ClickInViewport,
+
+	/** Set distance with an explicit numerical value, then explictly accept. */
+	Fixed,
+
+	//~ TODO: Add someday
+	// Gizmo,
+};
+
+
 // There is a lot of overlap in the options for Extrude, Offset, and Push/Pull, and they map to
 // the same op behind the scenes. However, we want to keep them as separate buttons to keep some
 // amount of shallowness in the UI, to make it more likely that new users will find the setting
@@ -97,6 +111,14 @@ class MESHMODELINGTOOLS_API UPolyEditExtrudeProperties : public UInteractiveTool
 	GENERATED_BODY()
 
 public:
+	/** How the extrude distance is set. */
+	UPROPERTY(EditAnywhere, Category = Extrude)
+	EPolyEditExtrudeDistanceMode DistanceMode = EPolyEditExtrudeDistanceMode::ClickInViewport;
+
+	/** Distance to extrude. */
+	UPROPERTY(EditAnywhere, Category = Extrude,
+		meta = (EditConditionHides, EditCondition = "DistanceMode == EPolyEditExtrudeDistanceMode::Fixed"))
+	double Distance = 100;
 
 	/** Direction in which to extrude. */
 	UPROPERTY(EditAnywhere, Category = Extrude,
@@ -105,7 +127,7 @@ public:
 
 	/** What axis to measure the extrusion distance along. */
 	UPROPERTY(EditAnywhere, Category = Extrude, AdvancedDisplay,
-		meta = (EditConditionHides, EditCondition = "DirectionMode != EPolyEditExtrudeModeOptions::SingleDirection"))
+		meta = (EditConditionHides, EditCondition = "DirectionMode != EPolyEditExtrudeModeOptions::SingleDirection && DistanceMode == EPolyEditExtrudeDistanceMode::ClickInViewport"))
 	EPolyEditExtrudeDirection MeasureDirection = EPolyEditExtrudeDirection::SelectionNormal;
 
 	/** Controls whether extruding an entire open-border patch should create a solid or an open shell */
@@ -117,7 +139,7 @@ public:
 	EPolyEditExtrudeModeOptions DirectionMode = EPolyEditExtrudeModeOptions::SingleDirection;
 
 	/** Controls the maximum distance vertices can move from the target distance in order to stay parallel with their source triangles. */
-	UPROPERTY(EditAnywhere, Category = Offset,
+	UPROPERTY(EditAnywhere, Category = Extrude,
 		meta = (ClampMin = "1", EditConditionHides, EditCondition = "DirectionMode == EPolyEditExtrudeModeOptions::SelectedTriangleNormalsEven"))
 	double MaxDistanceScaleFactor = 4.0;
 
@@ -137,6 +159,15 @@ class MESHMODELINGTOOLS_API UPolyEditOffsetProperties : public UInteractiveToolP
 	GENERATED_BODY()
 
 public:
+	/** How the offset distance is set. */
+	UPROPERTY(EditAnywhere, Category = Extrude)
+	EPolyEditExtrudeDistanceMode DistanceMode = EPolyEditExtrudeDistanceMode::ClickInViewport;
+
+	/** Offset distance. */
+	UPROPERTY(EditAnywhere, Category = Extrude,
+		meta = (EditConditionHides, EditCondition = "DistanceMode == EPolyEditExtrudeDistanceMode::Fixed"))
+	double Distance = 100;
+
 	/** Which way to move vertices during the offset */
 	UPROPERTY(EditAnywhere, Category = Offset)
 	EPolyEditOffsetModeOptions DirectionMode = EPolyEditOffsetModeOptions::VertexNormals;
@@ -150,8 +181,9 @@ public:
 	UPROPERTY(EditAnywhere, Category = Offset)
 	bool bShellsToSolids = true;
 
-	/** What axis to measure the extrusion distance along. When the direction mode is Single Direction, also controls the direction. */
-	UPROPERTY(EditAnywhere, Category = Offset, AdvancedDisplay)
+	/** What axis to measure the extrusion distance along. */
+	UPROPERTY(EditAnywhere, Category = Offset, AdvancedDisplay, meta = (EditConditionHides,
+		EditCondition = "DistanceMode == EPolyEditExtrudeDistanceMode::ClickInViewport"))
 	EPolyEditExtrudeDirection MeasureDirection = EPolyEditExtrudeDirection::SelectionNormal;
 
 	/**
@@ -170,6 +202,15 @@ class MESHMODELINGTOOLS_API UPolyEditPushPullProperties : public UInteractiveToo
 	GENERATED_BODY()
 
 public:
+	/** How the offset distance is set. */
+	UPROPERTY(EditAnywhere, Category = Extrude)
+	EPolyEditExtrudeDistanceMode DistanceMode = EPolyEditExtrudeDistanceMode::ClickInViewport;
+
+	/** Offset distance. */
+	UPROPERTY(EditAnywhere, Category = Extrude,
+		meta = (EditConditionHides, EditCondition = "DistanceMode == EPolyEditExtrudeDistanceMode::Fixed"))
+	double Distance = 100;
+
 	/** Which way to move vertices during the offset */
 	UPROPERTY(EditAnywhere, Category = ExtrusionOptions)
 	EPolyEditPushPullModeOptions DirectionMode = EPolyEditPushPullModeOptions::SelectedTriangleNormals;
@@ -183,8 +224,9 @@ public:
 	UPROPERTY(EditAnywhere, Category = ExtrusionOptions)
 	bool bShellsToSolids = true;
 
-	/** What axis to measure the extrusion distance along. When the direction mode is Single Direction, also controls the direction. */
-	UPROPERTY(EditAnywhere, Category = ExtrusionOptions, AdvancedDisplay)
+	/** What axis to measure the extrusion distance along. */
+	UPROPERTY(EditAnywhere, Category = ExtrusionOptions, AdvancedDisplay, meta = (EditConditionHides, 
+		EditCondition = "DistanceMode == EPolyEditExtrudeDistanceMode::ClickInViewport"))
 	EPolyEditExtrudeDirection MeasureDirection = EPolyEditExtrudeDirection::SelectionNormal;
 
 	/**
@@ -229,6 +271,7 @@ public:
 	virtual bool CanStart() const override;
 	virtual EToolActivityStartResult Start() override;
 	virtual bool IsRunning() const override { return bIsRunning; }
+	virtual bool HasAccept() const { return true; };
 	virtual bool CanAccept() const override;
 	virtual EToolActivityEndResult End(EToolShutdownType) override;
 	virtual void Render(IToolsContextRenderAPI* RenderAPI) override;
@@ -277,7 +320,6 @@ protected:
 
 	UE::Geometry::FGroupTopologySelection ActiveSelection;
 	UE::Geometry::FFrame3d ActiveSelectionFrameWorld;
-	UE::Geometry::FFrame3d ExtrusionFrameWorld;
 	float UVScaleFactor = 1.0f;
 
 	bool bRequestedApply = false;
diff --git a/Engine/Plugins/Runtime/MeshModelingToolset/Source/ModelingComponents/Public/InteractiveToolActivity.h b/Engine/Plugins/Runtime/MeshModelingToolset/Source/ModelingComponents/Public/InteractiveToolActivity.h
--- a/Engine/Plugins/Runtime/MeshModelingToolset/Source/ModelingComponents/Public/InteractiveToolActivity.h
+++ b/Engine/Plugins/Runtime/MeshModelingToolset/Source/ModelingComponents/Public/InteractiveToolActivity.h
@@ -109,6 +109,12 @@ public:
 		return ensure(false); 
 	}
 
+	/**
+	 * Similar to a tool, this determines whether the activity prefers an accept/cancel option to be
+	 * displayed (if true) vs a "complete".
+	 */
+	virtual bool HasAccept() const { return false; };
+
 	/**
 	 * If true, calling End with EToolShutdownType::Accept will result in a valid completion of
 	 * the activity. It's up to the tool whether to respect this or not, since End() should still
