MeshModelingTools: Fix incorrect ray direction transformation in CubeGrid and PolyEd focus.

diff --git a/Engine/Plugins/Experimental/MeshModelingToolsetExp/Source/MeshModelingToolsExp/Private/CubeGridTool.cpp b/Engine/Plugins/Experimental/MeshModelingToolsetExp/Source/MeshModelingToolsExp/Private/CubeGridTool.cpp
--- a/Engine/Plugins/Experimental/MeshModelingToolsetExp/Source/MeshModelingToolsExp/Private/CubeGridTool.cpp
+++ b/Engine/Plugins/Experimental/MeshModelingToolsetExp/Source/MeshModelingToolsExp/Private/CubeGridTool.cpp
@@ -1403,7 +1403,7 @@ bool UCubeGridTool::GetHitGridFace(const FRay& WorldRay, FCubeGrid::FCubeFace& F
 	if (MeshSpatial)
 	{
 		FRay3d LocalRay(CurrentMeshTransform.InverseTransformPosition((FVector3d)WorldRay.Origin),
-			CurrentMeshTransform.InverseTransformVectorNoScale((FVector3d)WorldRay.Direction));
+			CurrentMeshTransform.InverseTransformVector((FVector3d)WorldRay.Direction));
 
 		int32 Tid;
 		double LocalHitT = TNumericLimits<double>::Max();
diff --git a/Engine/Plugins/Runtime/MeshModelingToolset/Source/MeshModelingTools/Private/EditMeshPolygonsTool.cpp b/Engine/Plugins/Runtime/MeshModelingToolset/Source/MeshModelingTools/Private/EditMeshPolygonsTool.cpp
--- a/Engine/Plugins/Runtime/MeshModelingToolset/Source/MeshModelingTools/Private/EditMeshPolygonsTool.cpp
+++ b/Engine/Plugins/Runtime/MeshModelingToolset/Source/MeshModelingTools/Private/EditMeshPolygonsTool.cpp
@@ -730,7 +730,7 @@ FBox UEditMeshPolygonsTool::GetWorldSpaceFocusBox()
 bool UEditMeshPolygonsTool::GetWorldSpaceFocusPoint(const FRay& WorldRay, FVector& PointOut)
 {
 	FRay3d LocalRay(WorldTransform.InverseTransformPosition((FVector3d)WorldRay.Origin),
-		WorldTransform.InverseTransformNormal((FVector3d)WorldRay.Direction));
+		WorldTransform.InverseTransformVector((FVector3d)WorldRay.Direction));
 	UE::Geometry::Normalize(LocalRay.Direction);
 
 	int32 HitTID = GetSpatial().FindNearestHitTriangle(LocalRay);
