update ITF and ModelingComponents headers to account for recent move of FHitResult (out of EngineTypes.h and into its own header)

diff --git a/Engine/Plugins/Runtime/MeshModelingToolset/Source/ModelingComponents/Public/Mechanics/DragAlignmentMechanic.h b/Engine/Plugins/Runtime/MeshModelingToolset/Source/ModelingComponents/Public/Mechanics/DragAlignmentMechanic.h
--- a/Engine/Plugins/Runtime/MeshModelingToolset/Source/ModelingComponents/Public/Mechanics/DragAlignmentMechanic.h
+++ b/Engine/Plugins/Runtime/MeshModelingToolset/Source/ModelingComponents/Public/Mechanics/DragAlignmentMechanic.h
@@ -7,7 +7,7 @@
 #include "BaseBehaviors/BehaviorTargetInterfaces.h"
 #include "DynamicMesh/DynamicMeshAABBTree3.h"
 #include "ToolContextInterfaces.h" // FViewCameraState
-#include "Engine/EngineTypes.h" //FHitResult
+#include "Engine/HitResult.h"
 #include "InteractionMechanic.h"
 #include "TransformTypes.h"
 #include "VectorTypes.h"
diff --git a/Engine/Source/Runtime/InteractiveToolsFramework/Private/BaseGizmos/HitTargets.cpp b/Engine/Source/Runtime/InteractiveToolsFramework/Private/BaseGizmos/HitTargets.cpp
--- a/Engine/Source/Runtime/InteractiveToolsFramework/Private/BaseGizmos/HitTargets.cpp
+++ b/Engine/Source/Runtime/InteractiveToolsFramework/Private/BaseGizmos/HitTargets.cpp
@@ -2,7 +2,7 @@
 
 #include "BaseGizmos/HitTargets.h"
 #include "Components/PrimitiveComponent.h"
-#include "Engine/EngineTypes.h"    // FHitResult
+#include "Engine/HitResult.h"
 
 FInputRayHit UGizmoLambdaHitTarget::IsHit(const FInputDeviceRay& ClickPos) const
 {
diff --git a/Engine/Source/Runtime/InteractiveToolsFramework/Public/BaseTools/BaseBrushTool.h b/Engine/Source/Runtime/InteractiveToolsFramework/Public/BaseTools/BaseBrushTool.h
--- a/Engine/Source/Runtime/InteractiveToolsFramework/Public/BaseTools/BaseBrushTool.h
+++ b/Engine/Source/Runtime/InteractiveToolsFramework/Public/BaseTools/BaseBrushTool.h
@@ -5,7 +5,7 @@
 #include "CoreMinimal.h"
 #include "BaseTools/MeshSurfacePointTool.h"
 #include "InteractiveTool.h"
-#include "Engine/EngineTypes.h"
+#include "Engine/HitResult.h"
 #include "BaseBrushTool.generated.h"
 
 class UBrushStampIndicator;
diff --git a/Engine/Source/Runtime/InteractiveToolsFramework/Public/TargetInterfaces/PrimitiveComponentBackedTarget.h b/Engine/Source/Runtime/InteractiveToolsFramework/Public/TargetInterfaces/PrimitiveComponentBackedTarget.h
--- a/Engine/Source/Runtime/InteractiveToolsFramework/Public/TargetInterfaces/PrimitiveComponentBackedTarget.h
+++ b/Engine/Source/Runtime/InteractiveToolsFramework/Public/TargetInterfaces/PrimitiveComponentBackedTarget.h
@@ -3,7 +3,7 @@
 #pragma once
 
 #include "CoreMinimal.h"
-#include "Engine/EngineTypes.h"    // FHitResult
+#include "Engine/HitResult.h"
 #include "UObject/Interface.h"
 
 #include "PrimitiveComponentBackedTarget.generated.h"
diff --git a/Engine/Source/Runtime/InteractiveToolsFramework/Public/ToolTargets/PrimitiveComponentToolTarget.h b/Engine/Source/Runtime/InteractiveToolsFramework/Public/ToolTargets/PrimitiveComponentToolTarget.h
--- a/Engine/Source/Runtime/InteractiveToolsFramework/Public/ToolTargets/PrimitiveComponentToolTarget.h
+++ b/Engine/Source/Runtime/InteractiveToolsFramework/Public/ToolTargets/PrimitiveComponentToolTarget.h
@@ -5,6 +5,7 @@
 #include "CoreMinimal.h"
 
 #include "TargetInterfaces/PrimitiveComponentBackedTarget.h"
+#include "Engine/HitResult.h"
 #include "ToolTargets/ToolTarget.h"
 
 #include "PrimitiveComponentToolTarget.generated.h"
