MeshModelingTools: Fix minor static analysis complaint in CubeGridTood.

#rnx

diff --git a/Engine/Plugins/Runtime/MeshModelingToolset/Source/ModelingOperators/Private/CompositionOps/CubeGridBooleanOp.cpp b/Engine/Plugins/Runtime/MeshModelingToolset/Source/ModelingOperators/Private/CompositionOps/CubeGridBooleanOp.cpp
--- a/Engine/Plugins/Runtime/MeshModelingToolset/Source/ModelingOperators/Private/CompositionOps/CubeGridBooleanOp.cpp
+++ b/Engine/Plugins/Runtime/MeshModelingToolset/Source/ModelingOperators/Private/CompositionOps/CubeGridBooleanOp.cpp
@@ -331,9 +331,10 @@ protected:
 
 		// We're looking at things from the opposite direction, so our diagonal needs to
 		// connect the opposite face indices.
+		checkSlow(DiagFaceIdx1 == 0 || DiagFaceIdx1 == 1);
 		DiagFaceIdx1 = 1 - DiagFaceIdx1;
 		DiagFaceIdx2 = DiagFaceIdx1 + 2;
-		CCWFromDiag1 = (DiagFaceIdx1 + 1) % 4;
+		CCWFromDiag1 = DiagFaceIdx1 + 1;
 		CCWFromDiag2 = (DiagFaceIdx2 + 1) % 4;
 
 		bDiagonalWelded = TopFaceIdxWelded[DiagFaceIdx1] && TopFaceIdxWelded[DiagFaceIdx2];
