MeshModelingTools: Change how HoleFillOp assigns UV's to single trianges to avoid an unnecessary connected component search for the whole mesh. Also make it work in the presence of unset UV's though this currently can't happen inside the HoleFill Tool because it passes everything through a MeshDescription conversion.

diff --git a/Engine/Plugins/Runtime/MeshModelingToolset/Source/ModelingOperators/Private/CleaningOps/HoleFillOp.cpp b/Engine/Plugins/Runtime/MeshModelingToolset/Source/ModelingOperators/Private/CleaningOps/HoleFillOp.cpp
--- a/Engine/Plugins/Runtime/MeshModelingToolset/Source/ModelingOperators/Private/CleaningOps/HoleFillOp.cpp
+++ b/Engine/Plugins/Runtime/MeshModelingToolset/Source/ModelingOperators/Private/CleaningOps/HoleFillOp.cpp
@@ -92,28 +92,6 @@ namespace
 		return BaseTriangle.IndexOf(ParentVertex);
 	}
 
-	// Given two triangles adjacent to a newly added triangle in the base mesh, find the elements that should comprise 
-	// the newly added triangle in the UV mesh.
-	void SetUVTriangleFromExistingTriangles(FDynamicMeshUVOverlay& UVOverlay,
-											int TriangleA,
-											int TriangleB,
-											int NewBaseTriangleIndex,
-											FIndex3i& NewTriangleElements)
-	{
-		for (const FIndex3i& TriElements : { UVOverlay.GetTriangle(TriangleA), UVOverlay.GetTriangle(TriangleB) })
-		{
-			for (int TriangleElementIndex = 0; TriangleElementIndex < 3; TriangleElementIndex++)
-			{
-				const int ElementID = TriElements[TriangleElementIndex];
-				int IndexInBaseTriangle = FindParentVertexInBaseTriangle(UVOverlay, ElementID, NewBaseTriangleIndex);
-				if (IndexInBaseTriangle >= 0)
-				{
-					NewTriangleElements[IndexInBaseTriangle] = ElementID;
-				}
-			}
-		}
-	}
-
 	float Area(FDynamicMeshUVOverlay& UVOverlay, FIndex3i& Tri)
 	{
 		const FVector2f AB = UVOverlay.GetElement(Tri[1]) - UVOverlay.GetElement(Tri[0]);
@@ -215,6 +193,10 @@ bool FHoleFillOp::FillSingleTriangleHole(const FEdgeLoop& Loop, int32& NewGroupI
 	{
 		// Orientation matches an existing triangle, so reverse it
 		Swap(Vertices[0], Vertices[1]);
+
+		// keep ExistingEdges and ExistingTriangles aligned such that index i corresponds to the edge from vert i to (i+1)%3
+		Swap(ExistingEdges[1], ExistingEdges[2]);
+		Swap(ExistingTriangles[1], ExistingTriangles[2]);
 	}
 
 	int NewTriangleID = ResultMesh->AppendTriangle(Vertices, NewGroupID);
@@ -227,39 +209,69 @@ bool FHoleFillOp::FillSingleTriangleHole(const FEdgeLoop& Loop, int32& NewGroupI
 	{
 		FDynamicMeshUVOverlay* UVOverlay = ResultMesh->Attributes()->PrimaryUV();
 
-		// check we have a hole in the UV mesh as well
-		check(UVOverlay->IsSeamEdge(ExistingEdges[0]));
-		check(UVOverlay->IsSeamEdge(ExistingEdges[1]));
-		check(UVOverlay->IsSeamEdge(ExistingEdges[2]));
+		// check we have a hole in the UV mesh as well (note that boundary edges of unset UV triangles aren't seen 
+		// as seams by IsSeamEdge)
+		checkSlow(UVOverlay->IsSeamEdge(ExistingEdges[0]) || !UVOverlay->IsSetTriangle(ExistingTriangles[0]));
+		checkSlow(UVOverlay->IsSeamEdge(ExistingEdges[1]) || !UVOverlay->IsSetTriangle(ExistingTriangles[1]));
+		checkSlow(UVOverlay->IsSeamEdge(ExistingEdges[2]) || !UVOverlay->IsSetTriangle(ExistingTriangles[2]));
+
+		// We're going to try to assign the UVs such that the number of UV seam edges is minimized. If at least 
+		// two of our neighboring triangles have a shared UV element at one of our vertices, then using the elements
+		// of these two triangles will decrease the number of seam edges because we end up removing two seams and
+		// adding at most one (or none if this was an enclosed hole in a single island).
+	
+		// For each of our vertices, get the elements associated with that vertex by the two neighbor triangles. The 
+		// pair we store is (element from previous triangle, element from next triangle).
+		FIndex2i VertUVElements[3];
+		for (int NeighborIndex = 0; NeighborIndex < 3; ++NeighborIndex)
+		{
+			FIndex3i NeighborTriVids = ResultMesh->GetTriangle(ExistingTriangles[NeighborIndex]);
+			FIndex3i NeighborTriElements = UVOverlay->GetTriangle(ExistingTriangles[NeighborIndex]);
+
+			// Find our two vids in TriVids, get the corresponding elements in TriElements
+			int VertIndex1 = NeighborIndex;
+			int VertIndex2 = (VertIndex1 + 1) % 3;
+
+			VertUVElements[VertIndex1].B = NeighborTriElements[IndexUtil::FindTriIndex(Vertices[VertIndex1], NeighborTriVids)];
+			VertUVElements[VertIndex2].A = NeighborTriElements[IndexUtil::FindTriIndex(Vertices[VertIndex2], NeighborTriVids)];
+		}
 
 		FIndex3i NewTriangleElements;
-		int NewTriangleUVIsland;
 
-		if (TriangleUVIsland[ExistingTriangles[0]] == TriangleUVIsland[ExistingTriangles[1]])
+		/*
+		 * Imagine going counterclockwise from the bottom left in the triangle below. A and B correspond
+		 * to the element values in VertUVElements. When the two are equal the vertex has a shared element,
+		 * and the adjacent elements need to be on the adjoining sides.
+		 * 
+		 *     B.A
+		 *     / \
+		 *   A.---.B
+		 *    B   A
+		 */
+		if (VertUVElements[0].A == VertUVElements[0].B)
 		{
-			// Connect new triangle to existing triangles in a common UV island
-			SetUVTriangleFromExistingTriangles(*UVOverlay, ExistingTriangles[0], ExistingTriangles[1], NewTriangleID, NewTriangleElements);
-			NewTriangleUVIsland = TriangleUVIsland[ExistingTriangles[0]];
+			NewTriangleElements = FIndex3i(VertUVElements[0].A, VertUVElements[1].A, VertUVElements[2].B);
 		}
-		else if (TriangleUVIsland[ExistingTriangles[1]] == TriangleUVIsland[ExistingTriangles[2]])
+		else if (VertUVElements[1].A == VertUVElements[1].B)
 		{
-			SetUVTriangleFromExistingTriangles(*UVOverlay, ExistingTriangles[1], ExistingTriangles[2], NewTriangleID, NewTriangleElements);
-			NewTriangleUVIsland = TriangleUVIsland[ExistingTriangles[1]];
+			NewTriangleElements = FIndex3i(VertUVElements[0].B, VertUVElements[1].A, VertUVElements[2].A);
 		}
-		else if (TriangleUVIsland[ExistingTriangles[2]] == TriangleUVIsland[ExistingTriangles[0]])
+		else if (VertUVElements[2].A == VertUVElements[2].B)
 		{
-			SetUVTriangleFromExistingTriangles(*UVOverlay, ExistingTriangles[2], ExistingTriangles[0], NewTriangleID, NewTriangleElements);
-			NewTriangleUVIsland = TriangleUVIsland[ExistingTriangles[2]];
+			NewTriangleElements = FIndex3i(VertUVElements[0].A, VertUVElements[1].B, VertUVElements[2].A);
 		}
 		else
 		{
-			// All edges are in different UV islands.
-			// Pick an arbitrary edge and add a new third UV element
+			// Pick an arbitrary set UV edge and add a new third UV element
+			int32 NeighborTidToUse = UVOverlay->IsSetTriangle(ExistingTriangles[0]) ? ExistingTriangles[0] : ExistingTriangles[1];
 			SetUVTriangleFromAdjacentTriangle(*UVOverlay, ExistingTriangles[0], NewTriangleID, NewTriangleElements);
-			NewTriangleUVIsland = TriangleUVIsland[ExistingTriangles[0]];
 		}
 
-		EMeshResult Result = UVOverlay->SetTriangle(NewTriangleID, NewTriangleElements);
+		EMeshResult Result = EMeshResult::Ok;
+		if (NewTriangleElements.A != IndexConstants::InvalidID)
+		{
+			Result = UVOverlay->SetTriangle(NewTriangleID, NewTriangleElements);
+		}
 
 		if (!ensure(Result == EMeshResult::Ok))
 		{
@@ -267,9 +279,6 @@ bool FHoleFillOp::FillSingleTriangleHole(const FEdgeLoop& Loop, int32& NewGroupI
 			// Or create a totally new triangle with unique UVs and pass that in as a fallback?
 			return false;
 		}
-
-		TriangleUVIsland.SetNum(NewTriangleID + 1);
-		TriangleUVIsland[NewTriangleID] = NewTriangleUVIsland;
 	}
 
 
@@ -305,26 +314,6 @@ void FHoleFillOp::CalculateResult(FProgressCancel* Progress)
 		return;
 	}
 
-	// Initialize TriangleUVIsland, the UVIsland ID for each triangle in the mesh
-	// TODO: Initialize this only if really absolutely necessary
-	if (ResultMesh->HasAttributes() && ResultMesh->Attributes()->PrimaryUV() && FillOptions.bQuickFillSmallHoles)
-	{
-		FDynamicMeshUVOverlay* UVOverlay = ResultMesh->Attributes()->PrimaryUV();
-		FMeshConnectedComponents UVComponents(ResultMesh.Get());
-		UVComponents.FindConnectedTriangles([UVOverlay](int32 Triangle0, int32 Triangle1) {
-			return UVOverlay->AreTrianglesConnected(Triangle0, Triangle1);
-		});
-
-		TriangleUVIsland.Init( FDynamicMesh3::InvalidID, ResultMesh->MaxTriangleID());
-		for (int IslandID = 0; IslandID < UVComponents.Num(); ++IslandID)
-		{
-			for (int TriangleID : UVComponents.GetComponent(IslandID).Indices)
-			{
-				TriangleUVIsland[TriangleID] = IslandID;
-			}
-		}
-	}
-
 	TSet<int32> NewGroupIDs;
 
 	for (FEdgeLoop& Loop : Loops)
diff --git a/Engine/Plugins/Runtime/MeshModelingToolset/Source/ModelingOperators/Public/CleaningOps/HoleFillOp.h b/Engine/Plugins/Runtime/MeshModelingToolset/Source/ModelingOperators/Public/CleaningOps/HoleFillOp.h
--- a/Engine/Plugins/Runtime/MeshModelingToolset/Source/ModelingOperators/Public/CleaningOps/HoleFillOp.h
+++ b/Engine/Plugins/Runtime/MeshModelingToolset/Source/ModelingOperators/Public/CleaningOps/HoleFillOp.h
@@ -64,7 +64,6 @@ public:
 	// FDynamicMeshOperator implementation
 	void CalculateResult(FProgressCancel* Progress) override;
 
-	TArray<int> TriangleUVIsland;
 	bool FillSingleTriangleHole(const FEdgeLoop& Loop, int32& NewGroupID);
 };
 
