MeshModelingTools: Make the UV's created by CubeGridTool not be flipped when cutting things away.

diff --git a/Engine/Plugins/Runtime/MeshModelingToolset/Source/ModelingOperators/Private/CompositionOps/CubeGridBooleanOp.cpp b/Engine/Plugins/Runtime/MeshModelingToolset/Source/ModelingOperators/Private/CompositionOps/CubeGridBooleanOp.cpp
--- a/Engine/Plugins/Runtime/MeshModelingToolset/Source/ModelingOperators/Private/CompositionOps/CubeGridBooleanOp.cpp
+++ b/Engine/Plugins/Runtime/MeshModelingToolset/Source/ModelingOperators/Private/CompositionOps/CubeGridBooleanOp.cpp
@@ -31,6 +31,10 @@ public:
 	// See comments in CubeGridBooleanOp.h about how these work
 	TArray<int32, TFixedAllocator<6>> FaceUVOrientations;
 	double HeightUVOffset = 0;
+
+	// When the generated shape is subtracted from things, the resulting faces end up being seen from
+	// to opposite direction, which will make the UVs look flipped. This bool allows us to pre-flip them.
+	bool bFlipUVsForSubtraction = false;
 		
 	virtual FMeshShapeGenerator& Generate() override
 	{
@@ -129,6 +133,13 @@ protected:
 			 */
 
 			int RotatedCornerIndex = (FaceCornerIndex + 4 - FaceOrientation) % 4;
+
+			// If flipping for subtraction, the order of corners changes from 0, 1, 2, 3 (CCW) to 1, 0, 3, 2.
+			if (bFlipUVsForSubtraction)
+			{
+				RotatedCornerIndex = (RotatedCornerIndex % 2) ? RotatedCornerIndex - 1 : RotatedCornerIndex + 1;
+			}
+
 			FVector2i UVBase = IndexUtil::BoxFacesUV[RotatedCornerIndex];
 
 			FVector2f UV;
@@ -459,6 +470,7 @@ void FCubeGridBooleanOp::CalculateResult(FProgressCancel* Progress)
 	}
 	Generator.bCrosswiseDiagonal = bCrosswiseDiagonal;
 	Generator.HeightUVOffset = OpMeshHeightUVOffset;
+	Generator.bFlipUVsForSubtraction = bSubtract;
 	Generator.FaceUVOrientations = FaceUVOrientations;
 
 	OpMesh = MakeUnique<FDynamicMesh3>(&Generator.Generate());
