MeshModelingTools: Add defensive nullptr check in tool targets' CommitDynamicMesh to avoid crashing when the underlying asset is deleted.

diff --git a/Engine/Plugins/Runtime/MeshModelingToolset/Source/ModelingComponentsEditorOnly/Private/ToolTargets/SkeletalMeshComponentToolTarget.cpp b/Engine/Plugins/Runtime/MeshModelingToolset/Source/ModelingComponentsEditorOnly/Private/ToolTargets/SkeletalMeshComponentToolTarget.cpp
--- a/Engine/Plugins/Runtime/MeshModelingToolset/Source/ModelingComponentsEditorOnly/Private/ToolTargets/SkeletalMeshComponentToolTarget.cpp
+++ b/Engine/Plugins/Runtime/MeshModelingToolset/Source/ModelingComponentsEditorOnly/Private/ToolTargets/SkeletalMeshComponentToolTarget.cpp
@@ -173,8 +173,12 @@ void USkeletalMeshComponentToolTarget::CommitMeshDescription(const FCommitter& C
 
 void USkeletalMeshComponentToolTarget::CommitDynamicMesh(const FDynamicMesh3& Mesh, const FDynamicMeshCommitInfo& CommitInfo)
 {
-	FMeshDescription CurrentMeshDescription = *GetMeshDescription();
-	CommitDynamicMeshViaMeshDescription(MoveTemp(CurrentMeshDescription), *this, Mesh, CommitInfo);
+	const FMeshDescription* CurrentMeshDescription = GetMeshDescription();
+	if (ensureMsgf(CurrentMeshDescription, TEXT("Unable to commit mesh, perhaps the user deleted "
+		"the asset while the tool was active?")))
+	{
+		CommitDynamicMeshViaMeshDescription(FMeshDescription(*CurrentMeshDescription), *this, Mesh, CommitInfo);
+	}
 }
 
 
diff --git a/Engine/Plugins/Runtime/MeshModelingToolset/Source/ModelingComponentsEditorOnly/Private/ToolTargets/StaticMeshComponentToolTarget.cpp b/Engine/Plugins/Runtime/MeshModelingToolset/Source/ModelingComponentsEditorOnly/Private/ToolTargets/StaticMeshComponentToolTarget.cpp
--- a/Engine/Plugins/Runtime/MeshModelingToolset/Source/ModelingComponentsEditorOnly/Private/ToolTargets/StaticMeshComponentToolTarget.cpp
+++ b/Engine/Plugins/Runtime/MeshModelingToolset/Source/ModelingComponentsEditorOnly/Private/ToolTargets/StaticMeshComponentToolTarget.cpp
@@ -237,8 +237,12 @@ FDynamicMesh3 UStaticMeshComponentToolTarget::GetDynamicMesh()
 
 void UStaticMeshComponentToolTarget::CommitDynamicMesh(const FDynamicMesh3& Mesh, const FDynamicMeshCommitInfo& CommitInfo)
 {
-	FMeshDescription CurrentMeshDescription = *GetMeshDescription();
-	CommitDynamicMeshViaMeshDescription(MoveTemp(CurrentMeshDescription), *this, Mesh, CommitInfo);
+	const FMeshDescription* CurrentMeshDescription = GetMeshDescription();
+	if (ensureMsgf(CurrentMeshDescription, TEXT("Unable to commit mesh, perhaps the user deleted "
+		"the asset while the tool was active?")))
+	{
+		CommitDynamicMeshViaMeshDescription(FMeshDescription(*CurrentMeshDescription), *this, Mesh, CommitInfo);
+	}
 }
 
 UStaticMesh* UStaticMeshComponentToolTarget::GetStaticMesh() const
