MeshModelingTools: Prevent a crash from trying to dereference a failed cast to UStaticMeshComponent in ModelingSceneSnappingManager.

diff --git a/Engine/Plugins/Runtime/MeshModelingToolset/Source/ModelingComponents/Private/Snapping/ModelingSceneSnappingManager.cpp b/Engine/Plugins/Runtime/MeshModelingToolset/Source/ModelingComponents/Private/Snapping/ModelingSceneSnappingManager.cpp
--- a/Engine/Plugins/Runtime/MeshModelingToolset/Source/ModelingComponents/Private/Snapping/ModelingSceneSnappingManager.cpp
+++ b/Engine/Plugins/Runtime/MeshModelingToolset/Source/ModelingComponents/Private/Snapping/ModelingSceneSnappingManager.cpp
@@ -358,13 +358,13 @@ bool GetComponentHitTriangle_Internal(FHitResult HitResult, VectorType* TriVerti
 		return false;
 	}
 
-	if (UPrimitiveComponent* Component = HitResult.Component.Get())
+	if (UStaticMeshComponent* Component = Cast<UStaticMeshComponent>(HitResult.Component.Get()))
 	{
 		// physics collision data is created from StaticMesh RenderData
 		// so use HitResult.FaceIndex to extract triangle from the LOD0 mesh
 		// (note: this may be incorrect if there are multiple sections...in that case I think we have to
 		//  first find section whose accumulated index range would contain .FaceIndexX)
-		UStaticMesh* StaticMesh = Cast<UStaticMeshComponent>(Component)->GetStaticMesh();
+		UStaticMesh* StaticMesh = Component->GetStaticMesh();
 		FStaticMeshLODResources& LOD = StaticMesh->GetRenderData()->LODResources[0];
 		FIndexArrayView Indices = LOD.IndexBuffer.GetArrayView();
 		int32 TriIdx = 3 * HitResult.FaceIndex;
