MeshModelingTools: Minor rephrasing in CubeGrid UV flipping logic to appease the static analyzer.

diff --git a/Engine/Plugins/Runtime/MeshModelingToolset/Source/ModelingOperators/Private/CompositionOps/CubeGridBooleanOp.cpp b/Engine/Plugins/Runtime/MeshModelingToolset/Source/ModelingOperators/Private/CompositionOps/CubeGridBooleanOp.cpp
--- a/Engine/Plugins/Runtime/MeshModelingToolset/Source/ModelingOperators/Private/CompositionOps/CubeGridBooleanOp.cpp
+++ b/Engine/Plugins/Runtime/MeshModelingToolset/Source/ModelingOperators/Private/CompositionOps/CubeGridBooleanOp.cpp
@@ -137,7 +137,12 @@ protected:
 			// If flipping for subtraction, the order of corners changes from 0, 1, 2, 3 (CCW) to 1, 0, 3, 2.
 			if (bFlipUVsForSubtraction)
 			{
-				RotatedCornerIndex = (RotatedCornerIndex % 2) ? RotatedCornerIndex - 1 : RotatedCornerIndex + 1;
+				// Static analysis doesn't like this commented out version because it thinks that RotatedCornerIndex can
+				// end up outside [0,3], which it can't. But we rephrase it below.
+				//RotatedCornerIndex = (RotatedCornerIndex % 2) ? RotatedCornerIndex - 1 : RotatedCornerIndex + 1;
+
+				static const int Adjustments[4] = { 1, -1, 1, -1 };
+				RotatedCornerIndex += Adjustments[RotatedCornerIndex];
 			}
 
 			FVector2i UVBase = IndexUtil::BoxFacesUV[RotatedCornerIndex];
