ModelingMode: add cvar to control snapping to Volumes. Disable by default.

diff --git a/Engine/Plugins/Runtime/MeshModelingToolset/Source/ModelingComponents/Private/Scene/SceneGeometrySpatialCache.cpp b/Engine/Plugins/Runtime/MeshModelingToolset/Source/ModelingComponents/Private/Scene/SceneGeometrySpatialCache.cpp
--- a/Engine/Plugins/Runtime/MeshModelingToolset/Source/ModelingComponents/Private/Scene/SceneGeometrySpatialCache.cpp
+++ b/Engine/Plugins/Runtime/MeshModelingToolset/Source/ModelingComponents/Private/Scene/SceneGeometrySpatialCache.cpp
@@ -11,6 +11,13 @@
 using namespace UE::Geometry;
 
 
+// This CVar controls whether scene hit queries will hit volume surfaces.
+// When enabled, Volumes with a valid CollisionProfile will be included as snap targets
+TAutoConsoleVariable<bool> CVarEnableModelingVolumeSnapping(
+	TEXT("modeling.EnableVolumeSnapping"),
+	false,
+	TEXT("Enable snapping to volumes"));
+
 namespace UE
 {
 namespace Local
@@ -384,8 +391,10 @@ bool FSceneGeometrySpatialCache::EnableComponentTracking(UPrimitiveComponent* Co
 		return true;
 	}
 
+	bool bEnableVolumes = CVarEnableModelingVolumeSnapping.GetValueOnAnyThread();
+
 	TSharedPtr<ISceneGeometrySpatial> NewSpatial;
-	if (Cast<UBrushComponent>(Component) != nullptr)
+	if (bEnableVolumes && Cast<UBrushComponent>(Component) != nullptr)
 	{
 		IdentifierOut = FSceneGeometryID{ UniqueIDGenerator++ };
 
diff --git a/Engine/Plugins/Runtime/MeshModelingToolset/Source/ModelingComponents/Private/Snapping/ModelingSceneSnappingManager.cpp b/Engine/Plugins/Runtime/MeshModelingToolset/Source/ModelingComponents/Private/Snapping/ModelingSceneSnappingManager.cpp
--- a/Engine/Plugins/Runtime/MeshModelingToolset/Source/ModelingComponents/Private/Snapping/ModelingSceneSnappingManager.cpp
+++ b/Engine/Plugins/Runtime/MeshModelingToolset/Source/ModelingComponents/Private/Snapping/ModelingSceneSnappingManager.cpp
@@ -15,6 +15,7 @@
 #include "Engine/StaticMesh.h"
 #include "Components/PrimitiveComponent.h"
 #include "Components/StaticMeshComponent.h"
+#include "Components/BrushComponent.h"
 #include "Components/DynamicMeshComponent.h"
 #include "UObject/UObjectGlobals.h"
 
@@ -22,6 +23,12 @@ using namespace UE::Geometry;
 
 #define LOCTEXT_NAMESPACE "USceneSnappingManager"
 
+
+
+// defined in SceneGeometrySpatialCache.cpp
+extern TAutoConsoleVariable<bool> CVarEnableModelingVolumeSnapping;
+
+
 static double SnapToIncrement(double fValue, double fIncrement, double offset = 0)
 {
 	if (!FMath::IsFinite(fValue))
@@ -330,6 +337,8 @@ static bool FindNearestVisibleObjectHit_Internal(
 		return false;
 	}
 
+	bool bEnableVolumes = CVarEnableModelingVolumeSnapping.GetValueOnAnyThread();
+
 	double NearestVisible = TNumericLimits<double>::Max();
 	for (const FHitResult& CurResult : OutHits)
 	{
@@ -339,6 +348,12 @@ static bool FindNearestVisibleObjectHit_Internal(
 			continue;
 		}
 
+		// filtering out any volume hits here will disable volume snapping
+		if (bEnableVolumes == false && Cast<UBrushComponent>(CurResult.GetComponent()) != nullptr)
+		{
+			continue;
+		}
+
 		if (CurResult.Distance < NearestVisible && VisibilityFilter.IsVisible(CurResult.Component.Get()) )
 		{
 			HitResultOut = CurResult;
