Make split/mshmrg preserve mesh pivots so users can round-trip through them w/out losing the pivot location

diff --git a/Engine/Plugins/Experimental/MeshModelingToolsetExp/Source/MeshModelingToolsExp/Private/SplitMeshesTool.cpp b/Engine/Plugins/Experimental/MeshModelingToolsetExp/Source/MeshModelingToolsExp/Private/SplitMeshesTool.cpp
--- a/Engine/Plugins/Experimental/MeshModelingToolsetExp/Source/MeshModelingToolsExp/Private/SplitMeshesTool.cpp
+++ b/Engine/Plugins/Experimental/MeshModelingToolsetExp/Source/MeshModelingToolsExp/Private/SplitMeshesTool.cpp
@@ -211,13 +211,20 @@ void USplitMeshesTool::UpdateSplitMeshes()
 				}
 			}
 			
-			// reposition mesh
-			FAxisAlignedBox3d Bounds = Submesh.GetBounds();
-			MeshTransforms::Translate(Submesh, -Bounds.Center());
+			// TODO: Consider whether to expose bCenterPivots as an option to the user
+			constexpr bool bCenterPivots = false;
+			FVector3d Origin = FVector3d::ZeroVector;
+			if (bCenterPivots)
+			{
+				// reposition mesh
+				FAxisAlignedBox3d Bounds = Submesh.GetBounds();
+				Origin = Bounds.Center();
+				MeshTransforms::Translate(Submesh, -Origin);
+			}
 
 			SplitInfo.Meshes[k] = MoveTemp(Submesh);
 			SplitInfo.Materials[k] = MoveTemp(NewMaterials);
-			SplitInfo.Origins[k] = Bounds.Center();
+			SplitInfo.Origins[k] = Origin;
 		}
 	}
 
diff --git a/Engine/Plugins/Runtime/MeshModelingToolset/Source/MeshModelingTools/Private/CombineMeshesTool.cpp b/Engine/Plugins/Runtime/MeshModelingToolset/Source/MeshModelingTools/Private/CombineMeshesTool.cpp
--- a/Engine/Plugins/Runtime/MeshModelingToolset/Source/MeshModelingTools/Private/CombineMeshesTool.cpp
+++ b/Engine/Plugins/Runtime/MeshModelingToolset/Source/MeshModelingTools/Private/CombineMeshesTool.cpp
@@ -173,8 +173,24 @@ void UCombineMeshesTool::CreateNewAsset()
 	AccumulateDMesh.Attributes()->EnableTangents();
 	AccumulateDMesh.Attributes()->EnableMaterialID();
 	AccumulateDMesh.Attributes()->EnablePrimaryColors();
-	FTransform3d AccumToWorld(Box.GetCenter());
-	FTransform3d ToAccum(-Box.GetCenter());
+	constexpr bool bCenterPivot = false;
+	FVector3d Origin = FVector3d::ZeroVector;
+	if (bCenterPivot)
+	{
+		// Place the pivot at the bounding box center
+		Origin = Box.GetCenter();
+	}
+	else if (!Targets.IsEmpty())
+	{
+		// Use the average pivot
+		for (int32 ComponentIdx = 0; ComponentIdx < Targets.Num(); ComponentIdx++)
+		{
+			Origin += UE::ToolTarget::GetLocalToWorldTransform(Targets[ComponentIdx]).TransformPosition(FVector3d::ZeroVector);
+		}
+		Origin /= Targets.Num();
+	}
+	FTransform3d AccumToWorld(Origin);
+	FTransform3d ToAccum(-Origin);
 
 	FSimpleShapeSet3d SimpleCollision;
 	UE::Geometry::FComponentCollisionSettings CollisionSettings;
