ModelingComponents: Handle ChildActorComponent with null child Actor in WorldRenderCapture

diff --git a/Engine/Plugins/Runtime/MeshModelingToolset/Source/ModelingComponents/Private/Scene/WorldRenderCapture.cpp b/Engine/Plugins/Runtime/MeshModelingToolset/Source/ModelingComponents/Private/Scene/WorldRenderCapture.cpp
--- a/Engine/Plugins/Runtime/MeshModelingToolset/Source/ModelingComponents/Private/Scene/WorldRenderCapture.cpp
+++ b/Engine/Plugins/Runtime/MeshModelingToolset/Source/ModelingComponents/Private/Scene/WorldRenderCapture.cpp
@@ -149,9 +149,12 @@ void FWorldRenderCapture::SetVisibleActors(const TArray<AActor*>& Actors)
 			else if (Cast<UChildActorComponent>(Component) != nullptr)
 			{
 				AActor* ChildActor = Cast<UChildActorComponent>(Component)->GetChildActor();
-				for (UActorComponent* SubComponent : ChildActor->GetComponents())
+				if (ChildActor != nullptr)
 				{
-					ComponentQueue.Add(SubComponent);
+					for (UActorComponent* SubComponent : ChildActor->GetComponents())
+					{
+						ComponentQueue.Add(SubComponent);
+					}
 				}
 			}
 		}
