Deprecate TTransformSRT3::Inverse() and attempt to fix usages of it and also FTransform::Inverse() in the modeling code

diff --git a/Engine/Plugins/Experimental/HairModelingToolset/Source/HairModelingToolset/Private/GroomToMeshTool.cpp b/Engine/Plugins/Experimental/HairModelingToolset/Source/HairModelingToolset/Private/GroomToMeshTool.cpp
--- a/Engine/Plugins/Experimental/HairModelingToolset/Source/HairModelingToolset/Private/GroomToMeshTool.cpp
+++ b/Engine/Plugins/Experimental/HairModelingToolset/Source/HairModelingToolset/Private/GroomToMeshTool.cpp
@@ -639,16 +639,14 @@ static void TransformToSpace(FDynamicMesh3* Mesh, UPrimitiveComponent* SourceCom
 	{
 		check(TargetSpace.TargetActor);
 		FTransformSRT3d ActorTransform(TargetSpace.TargetActor->GetActorTransform());
-		FTransformSRT3d ToActorTransform = ActorTransform.Inverse();
-		MeshTransforms::ApplyTransform(*Mesh, ToActorTransform);
+		MeshTransforms::ApplyTransformInverse(*Mesh, ActorTransform);
 
 	}
 	else if (TargetSpace.TargetSpace == ETargetSpaceType::Component)
 	{
 		check(TargetSpace.TargetComponent);
 		FTransformSRT3d ComponentTransform(TargetSpace.TargetComponent->GetComponentTransform());
-		FTransformSRT3d ToComponentTransform = ComponentTransform.Inverse();
-		MeshTransforms::ApplyTransform(*Mesh, ToComponentTransform);
+		MeshTransforms::ApplyTransformInverse(*Mesh, ComponentTransform);
 	}
 	else if (TargetSpace.TargetSpace == ETargetSpaceType::Frame)
 	{
diff --git a/Engine/Plugins/Experimental/MeshModelingToolsetExp/Source/MeshModelingToolsExp/Private/BakeMeshAttributeMapsTool.cpp b/Engine/Plugins/Experimental/MeshModelingToolsetExp/Source/MeshModelingToolsExp/Private/BakeMeshAttributeMapsTool.cpp
--- a/Engine/Plugins/Experimental/MeshModelingToolsetExp/Source/MeshModelingToolsExp/Private/BakeMeshAttributeMapsTool.cpp
+++ b/Engine/Plugins/Experimental/MeshModelingToolsetExp/Source/MeshModelingToolsExp/Private/BakeMeshAttributeMapsTool.cpp
@@ -550,7 +550,7 @@ void UBakeMeshAttributeMapsTool::UpdateDetailMesh()
 		const FTransformSRT3d DetailToWorld = UE::ToolTarget::GetLocalToWorldTransform(DetailTarget);
 		MeshTransforms::ApplyTransform(*DetailMesh, DetailToWorld);
 		const FTransformSRT3d WorldToBase = UE::ToolTarget::GetLocalToWorldTransform(Targets[0]);
-		MeshTransforms::ApplyTransform(*DetailMesh, WorldToBase.Inverse());
+		MeshTransforms::ApplyTransformInverse(*DetailMesh, WorldToBase);
 	}
 
 	DetailSpatial = MakeShared<FDynamicMeshAABBTree3, ESPMode::ThreadSafe>();
diff --git a/Engine/Plugins/Experimental/MeshModelingToolsetExp/Source/MeshModelingToolsExp/Private/BakeMeshAttributeVertexTool.cpp b/Engine/Plugins/Experimental/MeshModelingToolsetExp/Source/MeshModelingToolsExp/Private/BakeMeshAttributeVertexTool.cpp
--- a/Engine/Plugins/Experimental/MeshModelingToolsetExp/Source/MeshModelingToolsExp/Private/BakeMeshAttributeVertexTool.cpp
+++ b/Engine/Plugins/Experimental/MeshModelingToolsetExp/Source/MeshModelingToolsExp/Private/BakeMeshAttributeVertexTool.cpp
@@ -471,7 +471,7 @@ void UBakeMeshAttributeVertexTool::UpdateDetailMesh()
 		const FTransformSRT3d DetailToWorld(DetailComponent->GetWorldTransform());
 		MeshTransforms::ApplyTransform(*DetailMesh, DetailToWorld);
 		const FTransformSRT3d WorldToBase(TargetComponent->GetWorldTransform());
-		MeshTransforms::ApplyTransform(*DetailMesh, WorldToBase.Inverse());
+		MeshTransforms::ApplyTransformInverse(*DetailMesh, WorldToBase);
 	}
 
 	DetailSpatial = MakeShared<FDynamicMeshAABBTree3, ESPMode::ThreadSafe>();
diff --git a/Engine/Plugins/Experimental/MeshModelingToolsetExp/Source/MeshModelingToolsExp/Private/BakeMultiMeshAttributeMapsTool.cpp b/Engine/Plugins/Experimental/MeshModelingToolsetExp/Source/MeshModelingToolsExp/Private/BakeMultiMeshAttributeMapsTool.cpp
--- a/Engine/Plugins/Experimental/MeshModelingToolsetExp/Source/MeshModelingToolsExp/Private/BakeMultiMeshAttributeMapsTool.cpp
+++ b/Engine/Plugins/Experimental/MeshModelingToolsetExp/Source/MeshModelingToolsExp/Private/BakeMultiMeshAttributeMapsTool.cpp
@@ -287,7 +287,7 @@ public:
 		const FVector3d NormalOut3d(NormalOut);
 		// As a stop-gap fix, transform our normal to the local space of the Base mesh here.
 		// TODO: Handle WorldToBase transformation in the evaluator.
-		NormalOut = WorldToBaseTransform.TransformNormal(FVector3f(ChildMesh->WorldTransform.TransformNormal(NormalOut3d)));
+		NormalOut = BaseToWorldTransform.InverseTransformNormal(FVector3f(ChildMesh->WorldTransform.TransformNormal(NormalOut3d)));
 		return bSuccess;
 	}
 
@@ -348,7 +348,6 @@ public:
 
 public:
 	FTransformSRT3d BaseToWorldTransform;
-	FTransformSRT3d WorldToBaseTransform;
 	TMap<UActorComponent*, const FActorAdapter*> ActorComponentMap;
 
 protected:
@@ -422,7 +421,6 @@ public:
 		};
 		DetailMeshScene->ProcessActorChildMeshes(ProcessChildMesh);
 		DetailSampler.BaseToWorldTransform = BaseToWorldTransform;
-		DetailSampler.WorldToBaseTransform = BaseToWorldTransform.Inverse();
 		Baker->SetDetailSampler(&DetailSampler);
 
 		for (const EBakeMapType MapType : ENUM_EBAKEMAPTYPE_ALL)
diff --git a/Engine/Plugins/Experimental/MeshModelingToolsetExp/Source/MeshModelingToolsExp/Private/EditPivotTool.cpp b/Engine/Plugins/Experimental/MeshModelingToolsetExp/Source/MeshModelingToolsExp/Private/EditPivotTool.cpp
--- a/Engine/Plugins/Experimental/MeshModelingToolsetExp/Source/MeshModelingToolsExp/Private/EditPivotTool.cpp
+++ b/Engine/Plugins/Experimental/MeshModelingToolsetExp/Source/MeshModelingToolsExp/Private/EditPivotTool.cpp
@@ -401,7 +401,7 @@ void UEditPivotTool::UpdateAssets(const FFrame3d& NewPivotWorldFrame)
 	GetToolManager()->BeginUndoTransaction(LOCTEXT("EditPivotToolTransactionName", "Edit Pivot"));
 
 	FTransform NewWorldTransform = NewPivotWorldFrame.ToFTransform();
-	FTransform NewWorldInverse = NewWorldTransform.Inverse();
+	FTransform NewWorldInverse = NewPivotWorldFrame.ToInverseFTransform();
 	TArray<FTransform> OriginalTransforms;
 	for (int32 ComponentIdx = 0; ComponentIdx < Targets.Num(); ComponentIdx++)
 	{
@@ -512,7 +512,8 @@ void UEditPivotTool::UpdateAssets(const FFrame3d& NewPivotWorldFrame)
 		}
 		else
 		{
-			// try to invert baked transform
+			// try to invert baked transform -- note that this may not be a correct inverse if there is rotation + non-uniform scale,
+			// but it's the best we can do given we can not bake a separate custom scale when this is not the first occurrence
 			FTransform Baked = OriginalTransforms[MapToFirstOccurrences[ComponentIdx]] * NewWorldInverse;
 			Component->SetWorldTransform(Baked.Inverse() * OriginalTransforms[ComponentIdx]);
 		}
diff --git a/Engine/Plugins/Experimental/MeshModelingToolsetExp/Source/MeshModelingToolsExp/Private/MirrorTool.cpp b/Engine/Plugins/Experimental/MeshModelingToolsetExp/Source/MeshModelingToolsExp/Private/MirrorTool.cpp
--- a/Engine/Plugins/Experimental/MeshModelingToolsetExp/Source/MeshModelingToolsExp/Private/MirrorTool.cpp
+++ b/Engine/Plugins/Experimental/MeshModelingToolsetExp/Source/MeshModelingToolsExp/Private/MirrorTool.cpp
@@ -61,13 +61,13 @@ TUniquePtr<FDynamicMeshOperator> UMirrorOperatorFactory::MakeNewOperator()
 		}
 	}
 	LocalToWorld.SetScale3D(LocalToWorldScale);
-	FTransformSRT3d WorldToLocal = FTransformSRT3d(LocalToWorld).Inverse();
 
 	// Now we can get the plane parameters in local space.
-	MirrorOp->LocalPlaneOrigin = WorldToLocal.TransformPosition(MirrorTool->MirrorPlaneOrigin);;
+	MirrorOp->LocalPlaneOrigin = LocalToWorld.InverseTransformPosition(MirrorTool->MirrorPlaneOrigin);;
 
 	FVector3d WorldNormal = MirrorTool->MirrorPlaneNormal;
-	MirrorOp->LocalPlaneNormal = WorldToLocal.TransformNormal(MirrorTool->MirrorPlaneNormal);
+	FTransformSRT3d LocalToWorldSRT(LocalToWorld); // Convert to Geometry::FTransformSRT3d for InverseTransformNormal function
+	MirrorOp->LocalPlaneNormal = LocalToWorldSRT.InverseTransformNormal(MirrorTool->MirrorPlaneNormal);
 
 	return MirrorOp;
 }
diff --git a/Engine/Plugins/Experimental/MeshModelingToolsetExp/Source/MeshModelingToolsExp/Private/Physics/SetCollisionGeometryTool.cpp b/Engine/Plugins/Experimental/MeshModelingToolsetExp/Source/MeshModelingToolsExp/Private/Physics/SetCollisionGeometryTool.cpp
--- a/Engine/Plugins/Experimental/MeshModelingToolsetExp/Source/MeshModelingToolsExp/Private/Physics/SetCollisionGeometryTool.cpp
+++ b/Engine/Plugins/Experimental/MeshModelingToolsetExp/Source/MeshModelingToolsExp/Private/Physics/SetCollisionGeometryTool.cpp
@@ -602,7 +602,6 @@ void USetCollisionGeometryTool::PrecomputeInputMeshes()
 
 	UToolTarget* CollisionTarget = Targets[Targets.Num() - 1];
 	FTransformSRT3d TargetTransform(UE::ToolTarget::GetLocalToWorldTransform(CollisionTarget));
-	FTransformSRT3d TargetTransformInv = TargetTransform.Inverse();
 
 	InputMeshes.Reset();
 	InputMeshes.SetNum(SourceObjectIndices.Num());
@@ -613,7 +612,7 @@ void USetCollisionGeometryTool::PrecomputeInputMeshes()
 		{
 			FTransformSRT3d ToWorld(UE::ToolTarget::GetLocalToWorldTransform(Targets[k]));
 			MeshTransforms::ApplyTransform(SourceMesh, ToWorld);
-			MeshTransforms::ApplyTransform(SourceMesh, TargetTransformInv);
+			MeshTransforms::ApplyTransformInverse(SourceMesh, TargetTransform);
 		}
 		SourceMesh.DiscardAttributes();
 		InputMeshes[k] = MakeShared<FDynamicMesh3, ESPMode::ThreadSafe>(MoveTemp(SourceMesh));
diff --git a/Engine/Plugins/Experimental/MeshModelingToolsetExp/Source/MeshModelingToolsExp/Private/PlaneCutTool.cpp b/Engine/Plugins/Experimental/MeshModelingToolsetExp/Source/MeshModelingToolsExp/Private/PlaneCutTool.cpp
--- a/Engine/Plugins/Experimental/MeshModelingToolsetExp/Source/MeshModelingToolsExp/Private/PlaneCutTool.cpp
+++ b/Engine/Plugins/Experimental/MeshModelingToolsetExp/Source/MeshModelingToolsExp/Private/PlaneCutTool.cpp
@@ -285,12 +285,11 @@ TUniquePtr<FDynamicMeshOperator> UPlaneCutOperatorFactory::MakeNewOperator()
 		}
 	}
 	LocalToWorld.SetScale3D(LocalToWorldScale);
-	FTransform WorldToLocal = LocalToWorld.Inverse();
 
-	FVector LocalOrigin = WorldToLocal.TransformPosition((FVector)CutTool->CutPlaneWorld.Origin);
+	FVector LocalOrigin = LocalToWorld.InverseTransformPosition((FVector)CutTool->CutPlaneWorld.Origin);
 	FVector3d WorldNormal = CutTool->CutPlaneWorld.GetAxis(2);
-	FTransformSRT3d W2LForNormal(WorldToLocal);
-	FVector LocalNormal = (FVector)W2LForNormal.TransformNormal(WorldNormal);
+	FTransformSRT3d L2WForNormal(LocalToWorld);
+	FVector LocalNormal = (FVector)L2WForNormal.InverseTransformNormal(WorldNormal);
 	FVector BackTransformed = LocalToWorld.TransformVector(LocalNormal);
 	float NormalScaleFactor = FVector::DotProduct(BackTransformed, (FVector)WorldNormal);
 	if (NormalScaleFactor >= FLT_MIN)
diff --git a/Engine/Plugins/Experimental/MeshModelingToolsetExp/Source/MeshModelingToolsExp/Private/Sculpting/MeshSculptToolBase.cpp b/Engine/Plugins/Experimental/MeshModelingToolsetExp/Source/MeshModelingToolsExp/Private/Sculpting/MeshSculptToolBase.cpp
--- a/Engine/Plugins/Experimental/MeshModelingToolsetExp/Source/MeshModelingToolsExp/Private/Sculpting/MeshSculptToolBase.cpp
+++ b/Engine/Plugins/Experimental/MeshModelingToolsetExp/Source/MeshModelingToolsExp/Private/Sculpting/MeshSculptToolBase.cpp
@@ -634,7 +634,8 @@ retry_frame_update:
 	ActiveStrokePathArcLen += Distance(LastBrushFrameWorld.Origin, PrevBrushFrameWorld.Origin);
 
 	LastBrushFrameLocal = LastBrushFrameWorld;
-	LastBrushFrameLocal.Transform(CurTargetTransform.Inverse());
+	LastBrushFrameLocal.Transform(CurTargetTransform.InverseUnsafe()); // Note: Unsafe inverse used because we cannot handle scales on a frame regardless.
+	// TODO: in the case of a non-uniform scale, consider whether we should do additional work to align the Z axis?
 }
 
 void UMeshSculptToolBase::AlignBrushToView()
diff --git a/Engine/Plugins/Experimental/Water/Source/Runtime/Private/WaterBooleanUtils.cpp b/Engine/Plugins/Experimental/Water/Source/Runtime/Private/WaterBooleanUtils.cpp
--- a/Engine/Plugins/Experimental/Water/Source/Runtime/Private/WaterBooleanUtils.cpp
+++ b/Engine/Plugins/Experimental/Water/Source/Runtime/Private/WaterBooleanUtils.cpp
@@ -243,7 +243,7 @@ void FWaterBooleanUtils::MakeNormalizationTransform(
 	FromNormalizedOut = FTransformSRT3d::Identity();
 	FromNormalizedOut.SetTranslation(Bounds.Center());
 	FromNormalizedOut.SetScale(WorldSize * FVector3d::One());
-	ToNormalizedOut = FromNormalizedOut.Inverse();
+	ToNormalizedOut = FromNormalizedOut.InverseUnsafe(); // Ok to invert because there is no rotation
 }
 
 void FWaterBooleanUtils::SetToFaceNormals(FDynamicMesh3& Mesh)
diff --git a/Engine/Plugins/Runtime/MeshModelingToolset/Source/MeshModelingTools/Private/CombineMeshesTool.cpp b/Engine/Plugins/Runtime/MeshModelingToolset/Source/MeshModelingTools/Private/CombineMeshesTool.cpp
--- a/Engine/Plugins/Runtime/MeshModelingToolset/Source/MeshModelingTools/Private/CombineMeshesTool.cpp
+++ b/Engine/Plugins/Runtime/MeshModelingToolset/Source/MeshModelingTools/Private/CombineMeshesTool.cpp
@@ -356,7 +356,6 @@ void UCombineMeshesTool::UpdateExistingAsset()
 	SkipActor = UE::ToolTarget::GetTargetActor(Targets[SkipIndex]);
 
 	FTransform3d TargetToWorld = UE::ToolTarget::GetLocalToWorldTransform(Targets[SkipIndex]);
-	FTransform3d WorldToTarget = TargetToWorld.Inverse();
 
 	FSimpleShapeSet3d SimpleCollision;
 	UE::Geometry::FComponentCollisionSettings CollisionSettings;
@@ -403,13 +402,28 @@ void UCombineMeshesTool::UpdateExistingAsset()
 				{
 					ComponentDMesh.ReverseOrientation(true);
 				}
-				MeshTransforms::ApplyTransform(ComponentDMesh, WorldToTarget);
-				if (WorldToTarget.GetDeterminant() < 0)
+				MeshTransforms::ApplyTransformInverse(ComponentDMesh, TargetToWorld);
+				if (TargetToWorld.GetDeterminant() < 0)
 				{
 					ComponentDMesh.ReverseOrientation(true);
 				}
 				Transforms[0] = ComponentToWorld;
-				Transforms[1] = WorldToTarget;
+				if (TargetToWorld.GetRotation().IsIdentity() || TargetToWorld.GetScale3D().IsUniform())
+				{
+					// Inverse can be represented by a single FTransform3d
+					Transforms[1] = TargetToWorld.Inverse();
+				}
+				else
+				{
+					// Separate inverse into a rotation+translation part and a scale part
+					FQuat4d WorldToTargetR = TargetToWorld.GetRotation().Inverse();
+					FTransform3d WorldToTargetRT(WorldToTargetR, WorldToTargetR * (-TargetToWorld.GetTranslation()), FVector3d::One());
+					FTransform3d WorldToTargetS = FTransform3d::Identity;
+					WorldToTargetS.SetScale3D(FTransform3d::GetSafeScaleReciprocal(TargetToWorld.GetScale3D()));
+
+					Transforms[1] = WorldToTargetRT;
+					Transforms.Add(WorldToTargetS);
+				}
 				if (bOutputComponentSupportsCollision && UE::Geometry::ComponentTypeSupportsCollision(PrimitiveComponent))
 				{
 					UE::Geometry::AppendSimpleCollision(PrimitiveComponent, &SimpleCollision, Transforms);
diff --git a/Engine/Plugins/Runtime/MeshModelingToolset/Source/MeshModelingTools/Private/ToolActivities/PolyEditPlanarProjectionUVActivity.cpp b/Engine/Plugins/Runtime/MeshModelingToolset/Source/MeshModelingTools/Private/ToolActivities/PolyEditPlanarProjectionUVActivity.cpp
--- a/Engine/Plugins/Runtime/MeshModelingToolset/Source/MeshModelingTools/Private/ToolActivities/PolyEditPlanarProjectionUVActivity.cpp
+++ b/Engine/Plugins/Runtime/MeshModelingToolset/Source/MeshModelingTools/Private/ToolActivities/PolyEditPlanarProjectionUVActivity.cpp
@@ -172,28 +172,30 @@ void UPolyEditPlanarProjectionUVActivity::ApplySetUVs()
 	TArray<int32> ActiveTriangleSelection;
 	ActivityContext->CurrentTopology->GetSelectedTriangles(ActiveSelection, ActiveTriangleSelection);
 
-	// align projection frame to line user drew
-	FFrame3d PlanarFrame = SurfacePathMechanic->HitPath[0];
-	double UVScale = 1.0 / ActiveSelectionBounds.MaxDim();
-	FVector3d Delta = SurfacePathMechanic->HitPath[1].Origin - PlanarFrame.Origin;
-	double Dist = UE::Geometry::Normalize(Delta);
-	UVScale *= FMathd::Lerp(1.0, 25.0, Dist / ActiveSelectionBounds.MaxDim());
-	PlanarFrame.ConstrainedAlignAxis(0, Delta, PlanarFrame.Z());
-
-	// transform to local, use 3D point to transfer UV scale value
-	FTransform3d WorldTransform(ActivityContext->Preview->PreviewMesh->GetTransform());
-	FVector3d ScalePt = PlanarFrame.Origin + UVScale * PlanarFrame.Z();
-	FTransform3d ToLocalXForm(WorldTransform.Inverse());
-	PlanarFrame.Transform(ToLocalXForm);
-	ScalePt = ToLocalXForm.TransformPosition(ScalePt);
-	UVScale = Distance(ScalePt, PlanarFrame.Origin);
+	// Get world-position data about line user drew
+	FFrame3d WorldPlanarFrame = SurfacePathMechanic->HitPath[0];
+	double WorldUVScale = 1.0 / ActiveSelectionBounds.MaxDim();
+	FVector3d WorldDelta = SurfacePathMechanic->HitPath[1].Origin - WorldPlanarFrame.Origin;
+	double WorldDist = UE::Geometry::Normalize(WorldDelta);
+	WorldUVScale *= FMathd::Lerp(1.0, 25.0, WorldDist / ActiveSelectionBounds.MaxDim());
+	FVector3d WorldNormal = WorldPlanarFrame.Z();
+	
+	// Create a local frame w/ matching normal, aligned to the user's line
+	FTransformSRT3d WorldTransform(ActivityContext->Preview->PreviewMesh->GetTransform());
+	FVector3d LocalNormal = WorldTransform.InverseTransformNormal(WorldNormal);
+	FFrame3d LocalPlanarFrame(WorldTransform.InverseTransformPosition(WorldPlanarFrame.Origin), LocalNormal);
+	LocalPlanarFrame.ConstrainedAlignAxis(0, WorldTransform.InverseTransformVector(WorldDelta), LocalPlanarFrame.Z());
+	// Use a reference point to get a uniform scale factor based on the world scale. Note it will not match exactly if the WorldTransform has non-uniform scale.
+	FVector3d ScalePt = WorldPlanarFrame.Origin + WorldUVScale * (WorldPlanarFrame.X() + WorldPlanarFrame.Y()) * FMathd::InvSqrt2;
+	ScalePt = WorldTransform.InverseTransformPosition(ScalePt);
+	double LocalUVScale = Distance(ScalePt, LocalPlanarFrame.Origin);
 
 	// track changes
 	FDynamicMeshChangeTracker ChangeTracker(ActivityContext->CurrentMesh.Get());
 	ChangeTracker.BeginChange();
 	ChangeTracker.SaveTriangles(ActiveTriangleSelection, true);
 	FDynamicMeshEditor Editor(ActivityContext->CurrentMesh.Get());
-	Editor.SetTriangleUVsFromProjection(ActiveTriangleSelection, PlanarFrame, UVScale, FVector2f::Zero(), false, 0);
+	Editor.SetTriangleUVsFromProjection(ActiveTriangleSelection, LocalPlanarFrame, LocalUVScale, FVector2f::Zero(), false, 0);
 
 	// Emit undo (also updates relevant structures). We didn't change the mesh topology here but for now we use
 	// the same route as everything else. See :HandlePositionOnlyMeshChanges
diff --git a/Engine/Plugins/Runtime/MeshModelingToolset/Source/ModelingComponents/Private/BaseTools/BaseCreateFromSelectedTool.cpp b/Engine/Plugins/Runtime/MeshModelingToolset/Source/ModelingComponents/Private/BaseTools/BaseCreateFromSelectedTool.cpp
--- a/Engine/Plugins/Runtime/MeshModelingToolset/Source/ModelingComponents/Private/BaseTools/BaseCreateFromSelectedTool.cpp
+++ b/Engine/Plugins/Runtime/MeshModelingToolset/Source/ModelingComponents/Private/BaseTools/BaseCreateFromSelectedTool.cpp
@@ -187,8 +187,8 @@ void UBaseCreateFromSelectedTool::GenerateAsset(const FDynamicMeshOpResult& OpRe
 	FTransform3d NewTransform;
 	if (Targets.Num() == 1) // in the single-selection case, shove the result back into the original component space
 	{
-		FTransform3d ToSourceComponentSpace = UE::ToolTarget::GetLocalToWorldTransform(Targets[0]).Inverse();
-		MeshTransforms::ApplyTransform(*OpResult.Mesh, ToSourceComponentSpace);
+		FTransform3d FromSourceComponentSpace = UE::ToolTarget::GetLocalToWorldTransform(Targets[0]);
+		MeshTransforms::ApplyTransformInverse(*OpResult.Mesh, FromSourceComponentSpace);
 		NewTransform = UE::ToolTarget::GetLocalToWorldTransform(Targets[0]);
 	}
 	else // in the multi-selection case, center the pivot for the combined result
diff --git a/Engine/Source/Runtime/GeometryCore/Public/FrameTypes.h b/Engine/Source/Runtime/GeometryCore/Public/FrameTypes.h
--- a/Engine/Source/Runtime/GeometryCore/Public/FrameTypes.h
+++ b/Engine/Source/Runtime/GeometryCore/Public/FrameTypes.h
@@ -170,6 +170,13 @@ struct TFrame3
 		return FTransform((FQuat)Rotation, (FVector)Origin);
 	}
 
+	/** @return conversion of this Frame to an inverse FTransform */
+	FTransform ToInverseFTransform() const
+	{
+		TQuaternion<RealType> InverseRotation = Rotation.Inverse();
+		return FTransform((FQuat)InverseRotation, (FVector)(InverseRotation * (-Origin)));
+	}
+
 	/** @return conversion of this Frame to FPlane */
 	FPlane ToFPlane() const
 	{
@@ -182,6 +189,13 @@ struct TFrame3
 		return TTransformSRT3<RealType>(Rotation, Origin);
 	}
 
+	/** @return conversion of this Frame to an inverse TTransform */
+	TTransformSRT3<RealType> ToInverseTransform() const
+	{
+		TQuaternion<RealType> InverseRotation = Rotation.Inverse();
+		return TTransformSRT3<RealType>(InverseRotation, InverseRotation * (-Origin));
+	}
+
 	/** @return point at distances along frame axes */
 	TVector<RealType> PointAt(RealType X, RealType Y, RealType Z) const
 	{
@@ -338,7 +352,7 @@ struct TFrame3
 
 
 	/**
-	 * transform this frame by the given transform
+	 * transform this frame by the given transform. Note: Ignores scale, as TFrame3 does not support scaling
 	 */
 	void Transform(const FTransform& XForm)
 	{
@@ -348,7 +362,7 @@ struct TFrame3
 
 
 	/**
-	 * transform this frame by the given transform
+	 * transform this frame by the given transform. Note: Ignores scale, as TFrame3 does not support scaling
 	 */
 	void Transform(const TTransformSRT3<RealType>& XForm)
 	{
diff --git a/Engine/Source/Runtime/GeometryCore/Public/TransformTypes.h b/Engine/Source/Runtime/GeometryCore/Public/TransformTypes.h
--- a/Engine/Source/Runtime/GeometryCore/Public/TransformTypes.h
+++ b/Engine/Source/Runtime/GeometryCore/Public/TransformTypes.h
@@ -168,9 +168,20 @@ public:
 	}
 
 	/**
-	 * @return inverse of this transform
+	 * Attempts to return an inverse, but will give an incorrect result if the transform has both non-uniform scaling and rotation,
+	 * because TTransformSRT3<RealType> cannot represent the true inverse in that case.
 	 */
+	UE_DEPRECATED(5.1, "Use the InverseTransformPosition/Vector/Normal functions instead; TransformSRT3 cannot represent its own inverse.  If the exact previous behavior is needed, use InverseUnsafe.")
 	TTransformSRT3<RealType> Inverse() const
+	{
+		return InverseUnsafe();
+	}
+
+	/**
+	 * Attempts to return an inverse, but will give an incorrect result if the transform has both non-uniform scaling and rotation,
+	 * because TTransformSRT3<RealType> cannot represent the true inverse in that case.
+	 */
+	TTransformSRT3<RealType> InverseUnsafe() const
 	{
 		TQuaternion<RealType> InvRotation = Rotation.Inverse();
 		TVector<RealType> InvScale3D = GetSafeScaleReciprocal(Scale3D);
