MeshModelingTools: Fix selection getting cleared in PolyEd Inset and Disconnect and undo being a little broken for selection after CL 18570424.

diff --git a/Engine/Plugins/Runtime/MeshModelingToolset/Source/MeshModelingTools/Private/EditMeshPolygonsTool.cpp b/Engine/Plugins/Runtime/MeshModelingToolset/Source/MeshModelingTools/Private/EditMeshPolygonsTool.cpp
--- a/Engine/Plugins/Runtime/MeshModelingToolset/Source/MeshModelingTools/Private/EditMeshPolygonsTool.cpp
+++ b/Engine/Plugins/Runtime/MeshModelingToolset/Source/MeshModelingTools/Private/EditMeshPolygonsTool.cpp
@@ -1904,6 +1904,12 @@ void UEditMeshPolygonsTool::EmitCurrentMeshChangeAndUpdate(const FText& Transact
 	TUniquePtr<FDynamicMeshChange> MeshChangeIn,
 	const FGroupTopologySelection& OutputSelection)
 {
+	// We used to take this as a paremeter, but even if we happen to know that the FDynamicMeshChange doesn't
+	// involve topology changes, it acts via deleting/reinserting triangles in undo/redo, which changes the
+	// eids in a mesh and causes problems. So we always treat the group topology as modified in this function.
+	// TODO: Have an overload that uses a vertex change for non-topology-modifying cases.
+	constexpr bool bGroupTopologyModified = true;
+
 	// open top-level transaction
 	GetToolManager()->BeginUndoTransaction(TransactionLabel);
 
@@ -1912,14 +1918,16 @@ void UEditMeshPolygonsTool::EmitCurrentMeshChangeAndUpdate(const FText& Transact
 	// so that we can copy it ahead of time and reinstate it.
 	bool bReferencingSameSelection = (&SelectionMechanic->GetActiveSelection() == &OutputSelection);
 
-	bool bSelectionModified = !bReferencingSameSelection && SelectionMechanic->GetActiveSelection() != OutputSelection;
+	// Not actually relevant since our assumption of topology being modified means we always clear existing selection.
+	// bool bSelectionModified = !bReferencingSameSelection && SelectionMechanic->GetActiveSelection() != OutputSelection;
 
 	// In case we need to make a selection copy
 	FGroupTopologySelection TempSelection;
 	const FGroupTopologySelection* OutputSelectionToUse = &OutputSelection;
 
-	// If the selection is going to be cleared, we need to do it explicitly ourselves so that we can emit a change.
-	if (!SelectionMechanic->GetActiveSelection().IsEmpty() && bSelectionModified)
+	// Because we always consider group topology to be modified, the UpdateFromCurrentMesh() call down below
+	// will clear out our existing selection. However, we want to do this ourselves so that we issue a transaction.
+	if (!SelectionMechanic->GetActiveSelection().IsEmpty() /* && (bSelectionModified || bGroupTopologyModified) */)
 	{
 		if (bReferencingSameSelection)
 		{
@@ -1938,12 +1946,12 @@ void UEditMeshPolygonsTool::EmitCurrentMeshChangeAndUpdate(const FText& Transact
 		TransactionLabel);
 
 	// Update related structures
-	UpdateFromCurrentMesh(true);
-	ModifiedTopologyCounter += 1;
+	UpdateFromCurrentMesh(bGroupTopologyModified);
+	ModifiedTopologyCounter += bGroupTopologyModified;
 
-	// Set output selection either if we changed selections (to something non-empty), or if
-	// our selection got cleared due to bGroupTopologyModified.
-	if (!OutputSelectionToUse->IsEmpty() && bSelectionModified)
+	// Set output selection if there's a non-empty one. We know we've cleared the selection by
+	// this point due to treating topology as always modified.
+	if (!OutputSelectionToUse->IsEmpty() /* && (bSelectionModified || bGroupTopologyModified) */)
 	{
 		SelectionMechanic->BeginChange();
 		SelectionMechanic->SetSelection(*OutputSelectionToUse);
