StaticMeshLODResourcesMeshSurfaceAdapter generates a bad TriangleOffsetArray in certain cases
* TriangleOffsetArray range must be [0, NumTriangles)

diff --git a/Engine/Plugins/Runtime/MeshModelingToolset/Source/ModelingComponents/Private/Scene/MeshSceneAdapter.cpp b/Engine/Plugins/Runtime/MeshModelingToolset/Source/ModelingComponents/Private/Scene/MeshSceneAdapter.cpp
--- a/Engine/Plugins/Runtime/MeshModelingToolset/Source/ModelingComponents/Private/Scene/MeshSceneAdapter.cpp
+++ b/Engine/Plugins/Runtime/MeshModelingToolset/Source/ModelingComponents/Private/Scene/MeshSceneAdapter.cpp
@@ -884,7 +884,6 @@ struct FStaticMeshLODResourcesMeshSurfaceAdapter : public FStaticMeshLODResource
 		TriangleOffsetArray.Reserve(MeshIn->Sections.Num() + 1);
 		ValidSections.Reserve(MeshIn->Sections.Num());
 
-		int32 SectionOffset = 0;
 		for (const FStaticMeshSection& Section : MeshIn->Sections)
 		{
 			auto IsValidMaterial = [StaticMeshIn, Section]()
@@ -898,15 +897,13 @@ struct FStaticMeshLODResourcesMeshSurfaceAdapter : public FStaticMeshLODResource
 			bool bIsValidMaterial = bOnlySurfaceMaterials ? IsValidMaterial() : true;
 			if (bIsValidMaterial)
 			{
-				TriangleOffsetArray.Add(SectionOffset);
+				TriangleOffsetArray.Add(NumTriangles);
 				NumTriangles += Section.NumTriangles;
 				ValidSections.Add(&Section);
 			}
-
-			SectionOffset += Section.NumTriangles;
 		}
 
-		TriangleOffsetArray.Add(SectionOffset);
+		TriangleOffsetArray.Add(NumTriangles);
 	}
 };
 
