RemeshTool: if Preserve Sharp Edges in unchecked, allow seam edges to be collapsed by the remesher

diff --git a/Engine/Plugins/Runtime/MeshModelingToolset/Source/ModelingOperators/Private/CleaningOps/RemeshMeshOp.cpp b/Engine/Plugins/Runtime/MeshModelingToolset/Source/ModelingOperators/Private/CleaningOps/RemeshMeshOp.cpp
--- a/Engine/Plugins/Runtime/MeshModelingToolset/Source/ModelingOperators/Private/CleaningOps/RemeshMeshOp.cpp
+++ b/Engine/Plugins/Runtime/MeshModelingToolset/Source/ModelingOperators/Private/CleaningOps/RemeshMeshOp.cpp
@@ -120,7 +120,7 @@ void FRemeshMeshOp::CalculateResult(FProgressCancel* Progress)
 
 	FMeshConstraints Constraints;
 	constexpr bool bAllowSeamSplits = true;
-	constexpr bool bAllowSeamCollpase = false;
+	const bool bAllowSeamCollapse = !bPreserveSharpEdges;
 	const bool bAllowSeamSmoothing = !bPreserveSharpEdges;
 
 	FMeshConstraintsUtil::ConstrainAllBoundariesAndSeams(Constraints,
@@ -130,7 +130,7 @@ void FRemeshMeshOp::CalculateResult(FProgressCancel* Progress)
 		MaterialBoundaryConstraint,
 		bAllowSeamSplits,
 		bAllowSeamSmoothing,
-		bAllowSeamCollpase);
+		bAllowSeamCollapse);
 
 	if (bReprojectConstraints)
 	{
