ModelingTools: Improved handling in UEditMeshPolygonsTool to deal with cases where the TransformGizmo can't be acquired successfully, which was leading to crashes. Now a failed gizmo acquisition simply breaks user interactions, but doesn't crash the editor.

diff --git a/Engine/Plugins/Runtime/MeshModelingToolset/Source/MeshModelingTools/Private/EditMeshPolygonsTool.cpp b/Engine/Plugins/Runtime/MeshModelingToolset/Source/MeshModelingTools/Private/EditMeshPolygonsTool.cpp
--- a/Engine/Plugins/Runtime/MeshModelingToolset/Source/MeshModelingTools/Private/EditMeshPolygonsTool.cpp
+++ b/Engine/Plugins/Runtime/MeshModelingToolset/Source/MeshModelingTools/Private/EditMeshPolygonsTool.cpp
@@ -437,34 +437,42 @@ void UEditMeshPolygonsTool::Setup()
 	
 	TransformGizmo = UE::TransformGizmoUtil::CreateCustomRepositionableTransformGizmo(GizmoManager,
 		ETransformGizmoSubElements::FullTranslateRotateScale, this);
-	// Stop scaling at 0 rather than going negative
-	TransformGizmo->SetDisallowNegativeScaling(true);
-	// We allow non uniform scale even when the gizmo mode is set to "world" because we're not scaling components- we're
-	// moving vertices, so we don't care which axes we "scale" along.
-	TransformGizmo->SetIsNonUniformScaleAllowedFunction([]() {
-		return true;
-	});
-	// Hook up callbacks
-	TransformProxy = NewObject<UTransformProxy>(this);
-	TransformProxy->OnTransformChanged.AddUObject(this, &UEditMeshPolygonsTool::OnGizmoTransformChanged);
-	TransformProxy->OnBeginTransformEdit.AddUObject(this, &UEditMeshPolygonsTool::OnBeginGizmoTransform);
-	TransformProxy->OnEndTransformEdit.AddUObject(this, &UEditMeshPolygonsTool::OnEndGizmoTransform);
-	TransformProxy->OnEndPivotEdit.AddWeakLambda(this, [this](UTransformProxy* Proxy) {
-		LastTransformerFrame = FFrame3d(Proxy->GetTransform());
-		if (CommonProps->bLockRotation)
-		{
-			LockedTransfomerFrame = LastTransformerFrame;
-		}
-	});
-	TransformGizmo->SetActiveTarget(TransformProxy, GetToolManager());
-	TransformGizmo->SetVisibility(false);
+	if (ensure(TransformGizmo)) // If we don't get a valid gizmo a lot of interactions won't work, but at least we won't crash
+	{
+		// Stop scaling at 0 rather than going negative
+		TransformGizmo->SetDisallowNegativeScaling(true);
+		// We allow non uniform scale even when the gizmo mode is set to "world" because we're not scaling components- we're
+		// moving vertices, so we don't care which axes we "scale" along.
+		TransformGizmo->SetIsNonUniformScaleAllowedFunction([]() {
+			return true;
+			});
+
+		// Hook up callbacks
+		TransformProxy = NewObject<UTransformProxy>(this);
+		TransformProxy->OnTransformChanged.AddUObject(this, &UEditMeshPolygonsTool::OnGizmoTransformChanged);
+		TransformProxy->OnBeginTransformEdit.AddUObject(this, &UEditMeshPolygonsTool::OnBeginGizmoTransform);
+		TransformProxy->OnEndTransformEdit.AddUObject(this, &UEditMeshPolygonsTool::OnEndGizmoTransform);
+		TransformProxy->OnEndPivotEdit.AddWeakLambda(this, [this](UTransformProxy* Proxy) {
+			LastTransformerFrame = FFrame3d(Proxy->GetTransform());
+			if (CommonProps->bLockRotation)
+			{
+				LockedTransfomerFrame = LastTransformerFrame;
+			}
+			});
+		TransformGizmo->SetActiveTarget(TransformProxy, GetToolManager());
+		TransformGizmo->SetVisibility(false);
+	}
 
 	DragAlignmentMechanic = NewObject<UDragAlignmentMechanic>(this);
 	DragAlignmentMechanic->Setup(this);
 	DragAlignmentMechanic->InitializeDeformedMeshRayCast(
 		[this]() { return &GetSpatial(); },
 		WorldTransform, &LinearDeformer); // Should happen after LinearDeformer is initialized
-	DragAlignmentMechanic->AddToGizmo(TransformGizmo);
+
+	if (TransformGizmo)
+	{
+		DragAlignmentMechanic->AddToGizmo(TransformGizmo);
+	}
 
 	if (Topology->Groups.Num() < 2)
 	{
@@ -712,8 +720,11 @@ void UEditMeshPolygonsTool::UpdateGizmoFrame(const FFrame3d* UseFrame)
 
 	LastTransformerFrame = SetFrame;
 
-	// This resets the scale as well
-	TransformGizmo->ReinitializeGizmoTransform(SetFrame.ToFTransform());
+	if (TransformGizmo)
+	{
+		// This resets the scale as well
+		TransformGizmo->ReinitializeGizmoTransform(SetFrame.ToFTransform());
+	}
 }
 
 
@@ -773,7 +784,7 @@ void UEditMeshPolygonsTool::OnGizmoTransformChanged(UTransformProxy* Proxy, FTra
 		LastUpdateGizmoScale = FVector3d(Transform.GetScale3D());
 		GetToolManager()->PostInvalidation();
 		bGizmoUpdatePending = true;
-		bLastUpdateUsedWorldFrame = (TransformGizmo->CurrentCoordinateSystem == EToolContextCoordinateSystem::World);
+		bLastUpdateUsedWorldFrame = (TransformGizmo ? TransformGizmo->CurrentCoordinateSystem == EToolContextCoordinateSystem::World : false);
 	}
 }
 
@@ -786,15 +797,18 @@ void UEditMeshPolygonsTool::OnEndGizmoTransform(UTransformProxy* Proxy)
 
 	FFrame3d TransformFrame(Proxy->GetTransform());
 
-	if (CommonProps->bLockRotation)
-	{
-		FFrame3d SetFrame = TransformFrame;
-		SetFrame.Rotation = LockedTransfomerFrame.Rotation;
-		TransformGizmo->ReinitializeGizmoTransform(SetFrame.ToFTransform());
-	}
-	else
+	if (TransformGizmo)
 	{
-		TransformGizmo->SetNewChildScale(FVector::OneVector);
+		if (CommonProps->bLockRotation)
+		{
+			FFrame3d SetFrame = TransformFrame;
+			SetFrame.Rotation = LockedTransfomerFrame.Rotation;
+			TransformGizmo->ReinitializeGizmoTransform(SetFrame.ToFTransform());		
+		}
+		else
+		{
+			TransformGizmo->SetNewChildScale(FVector::OneVector);
+		}
 	}
 
 	LastTransformerFrame = TransformFrame;
@@ -1071,7 +1085,10 @@ void UEditMeshPolygonsTool::StartActivity(TObjectPtr<UInteractiveToolActivity> A
 	// call.
 	if (Activity->Start() == EToolActivityStartResult::Running)
 	{
-		TransformGizmo->SetVisibility(false);
+		if (TransformGizmo)
+		{
+			TransformGizmo->SetVisibility(false);
+		}
 		SelectionMechanic->SetIsEnabled(false);
 		SetToolPropertySourceEnabled(SelectionMechanic->Properties, false);
 		CurrentActivity = Activity;
@@ -1121,7 +1138,10 @@ void UEditMeshPolygonsTool::UpdateGizmoVisibility()
 {
 	if (SelectionMechanic->HasSelection())
 	{
-		TransformGizmo->SetVisibility(true);
+		if (TransformGizmo)
+		{
+			TransformGizmo->SetVisibility(true);
+		}
 
 		// update frame because we might be here due to an undo event/etc, rather than an explicit 
 		// selection change
@@ -1130,7 +1150,10 @@ void UEditMeshPolygonsTool::UpdateGizmoVisibility()
 	}
 	else
 	{
-		TransformGizmo->SetVisibility(false);
+		if (TransformGizmo)
+		{
+			TransformGizmo->SetVisibility(false);
+		}
 	}
 }
 
