ModellingComponents: Add drag alignment mechanic to plane control gizmo in PolyExt

diff --git a/Engine/Plugins/Runtime/MeshModelingToolset/Source/MeshModelingTools/Private/DrawPolygonTool.cpp b/Engine/Plugins/Runtime/MeshModelingToolset/Source/MeshModelingTools/Private/DrawPolygonTool.cpp
--- a/Engine/Plugins/Runtime/MeshModelingToolset/Source/MeshModelingTools/Private/DrawPolygonTool.cpp
+++ b/Engine/Plugins/Runtime/MeshModelingToolset/Source/MeshModelingTools/Private/DrawPolygonTool.cpp
@@ -35,6 +35,8 @@
 #include "Selection/ToolSelectionUtil.h"
 #include "ModelingObjectsCreationAPI.h"
 
+#include "Mechanics/DragAlignmentMechanic.h"
+
 using namespace UE::Geometry;
 
 #define LOCTEXT_NAMESPACE "UDrawPolygonTool"
@@ -130,6 +132,10 @@ void UDrawPolygonTool::Setup()
 	// listen for changes to the proxy and update the plane when that happens
 	PlaneTransformProxy->OnTransformChanged.AddUObject(this, &UDrawPolygonTool::PlaneTransformChanged);
 
+	DragAlignmentMechanic = NewObject<UDragAlignmentMechanic>(this);
+	DragAlignmentMechanic->Setup(this);
+	DragAlignmentMechanic->AddToGizmo(PlaneTransformGizmo);
+
 	// initialize material properties for new objects
 	MaterialProperties = NewObject<UNewMeshMaterialProperties>(this);
 	MaterialProperties->RestoreProperties(this);
@@ -191,6 +197,8 @@ void UDrawPolygonTool::Shutdown(EToolShutdownType ShutdownType)
 		delete SetPointInWorldConnector;
 	}
 
+	DragAlignmentMechanic->Shutdown();
+
 	GetToolManager()->GetPairedGizmoManager()->DestroyAllGizmosByOwner(this);
 
 	OutputTypeProperties->SaveProperties(this);
diff --git a/Engine/Plugins/Runtime/MeshModelingToolset/Source/MeshModelingTools/Public/DrawPolygonTool.h b/Engine/Plugins/Runtime/MeshModelingToolset/Source/MeshModelingTools/Public/DrawPolygonTool.h
--- a/Engine/Plugins/Runtime/MeshModelingToolset/Source/MeshModelingTools/Public/DrawPolygonTool.h
+++ b/Engine/Plugins/Runtime/MeshModelingToolset/Source/MeshModelingTools/Public/DrawPolygonTool.h
@@ -18,6 +18,7 @@
 
 
 class UCombinedTransformGizmo;
+class UDragAlignmentMechanic;
 class UTransformProxy;
 PREDECLARE_USE_GEOMETRY_CLASS(FDynamicMesh3);
 
@@ -319,6 +320,9 @@ protected:
 	UPROPERTY()
 	TObjectPtr<UPlaneDistanceFromHitMechanic> HeightMechanic;
 
+	UPROPERTY()
+	TObjectPtr<UDragAlignmentMechanic> DragAlignmentMechanic = nullptr;
+
 	/** Generate extruded meshes.  Returns true on success. */
 	bool GeneratePolygonMesh(const TArray<FVector3d>& Polygon, const TArray<TArray<FVector3d>>& PolygonHoles, FDynamicMesh3* ResultMeshOut, UE::Geometry::FFrame3d& WorldFrameOut, bool bIncludePreviewVtx, double ExtrudeDistance, bool bExtrudeSymmetric);
 
