ModellingComponents: Add optional corner point visualization to PolyEdit

diff --git a/Engine/Plugins/Runtime/MeshModelingToolset/Source/MeshModelingTools/Private/EditMeshPolygonsTool.cpp b/Engine/Plugins/Runtime/MeshModelingToolset/Source/MeshModelingTools/Private/EditMeshPolygonsTool.cpp
--- a/Engine/Plugins/Runtime/MeshModelingToolset/Source/MeshModelingTools/Private/EditMeshPolygonsTool.cpp
+++ b/Engine/Plugins/Runtime/MeshModelingToolset/Source/MeshModelingTools/Private/EditMeshPolygonsTool.cpp
@@ -348,6 +348,7 @@ void UEditMeshPolygonsTool::Setup()
 	SelectionMechanic = NewObject<UPolygonSelectionMechanic>(this);
 	SelectionMechanic->bAddSelectionFilterPropertiesToParentTool = false; // We'll do this ourselves later
 	SelectionMechanic->Setup(this);
+	SelectionMechanic->SetShowSelectableCorners(CommonProps->bShowSelectableCorners);
 	SelectionMechanic->Properties->RestoreProperties(this, GetPropertyCacheIdentifier(bTriangleMode));
 	SelectionMechanic->OnSelectionChanged.AddUObject(this, &UEditMeshPolygonsTool::OnSelectionModifiedEvent);
 	SelectionMechanic->OnFaceSelectionPreviewChanged.AddWeakLambda(this, [this]() {
@@ -1184,7 +1185,13 @@ void UEditMeshPolygonsTool::DrawHUD(FCanvas* Canvas, IToolsContextRenderAPI* Ren
 	SelectionMechanic->DrawHUD(Canvas, RenderAPI);
 }
 
-
+void UEditMeshPolygonsTool::OnPropertyModified(UObject* PropertySet, FProperty* Property)
+{
+	if (Property && (Property->GetFName() == GET_MEMBER_NAME_CHECKED(UPolyEditCommonProperties, bShowSelectableCorners)))
+	{
+		SelectionMechanic->SetShowSelectableCorners(CommonProps->bShowSelectableCorners);
+	}
+}
 
 
 //
diff --git a/Engine/Plugins/Runtime/MeshModelingToolset/Source/MeshModelingTools/Public/EditMeshPolygonsTool.h b/Engine/Plugins/Runtime/MeshModelingToolset/Source/MeshModelingTools/Public/EditMeshPolygonsTool.h
--- a/Engine/Plugins/Runtime/MeshModelingToolset/Source/MeshModelingTools/Public/EditMeshPolygonsTool.h
+++ b/Engine/Plugins/Runtime/MeshModelingToolset/Source/MeshModelingTools/Public/EditMeshPolygonsTool.h
@@ -75,6 +75,9 @@ public:
 	UPROPERTY(EditAnywhere, Category = Options)
 	bool bShowWireframe = false;
 
+	UPROPERTY(EditAnywhere, Category = Options)
+	bool bShowSelectableCorners = false;
+
 	/** Determines whether, on selection changes, the gizmo's rotation is taken from the object transform, or from the geometry
 	 elements selected. Only relevant with a local coordinate system and when rotation is not locked. */
 	UPROPERTY(EditAnywhere, Category = Gizmo, meta = (HideEditConditionToggle, EditCondition = "bLocalCoordSystem && !bLockRotation"))
@@ -452,6 +455,8 @@ public:
 	virtual bool HasCancel() const override { return true; }
 	virtual bool HasAccept() const override { return true; }
 
+	virtual void OnPropertyModified(UObject* PropertySet, FProperty* Property) override;
+
 	// IInteractiveToolCameraFocusAPI implementation
 	virtual FBox GetWorldSpaceFocusBox() override;
 	virtual bool GetWorldSpaceFocusPoint(const FRay& WorldRay, FVector& PointOut) override;
diff --git a/Engine/Plugins/Runtime/MeshModelingToolset/Source/ModelingComponents/Private/Selection/PolygonSelectionMechanic.cpp b/Engine/Plugins/Runtime/MeshModelingToolset/Source/ModelingComponents/Private/Selection/PolygonSelectionMechanic.cpp
--- a/Engine/Plugins/Runtime/MeshModelingToolset/Source/ModelingComponents/Private/Selection/PolygonSelectionMechanic.cpp
+++ b/Engine/Plugins/Runtime/MeshModelingToolset/Source/ModelingComponents/Private/Selection/PolygonSelectionMechanic.cpp
@@ -86,14 +86,16 @@ void UPolygonSelectionMechanic::Setup(UInteractiveTool* ParentToolIn)
 	// set up visualizers
 	PolyEdgesRenderer.LineColor = FLinearColor::Red;
 	PolyEdgesRenderer.LineThickness = 2.0;
+	PolyEdgesRenderer.PointColor = FLinearColor::Red;
+	PolyEdgesRenderer.PointSize = 10.0f;
 	HilightRenderer.LineColor = FLinearColor::Green;
 	HilightRenderer.LineThickness = 4.0f;
 	HilightRenderer.PointColor = FLinearColor::Green;
-	HilightRenderer.PointSize = 5.0f;
+	HilightRenderer.PointSize = 10.0f;
 	SelectionRenderer.LineColor = LinearColors::Gold3f();
 	SelectionRenderer.LineThickness = 4.0f;
 	SelectionRenderer.PointColor = LinearColors::Gold3f();
-	SelectionRenderer.PointSize = 5.0f;
+	SelectionRenderer.PointSize = 10.0f;
 
 	float HighlightedFacePercentDepthOffset = 0.5f;
 	HighlightedFaceMaterial = ToolSetupUtil::GetSelectionMaterial(FLinearColor::Green, ParentToolIn->GetToolManager(), HighlightedFacePercentDepthOffset);
@@ -250,6 +252,14 @@ void UPolygonSelectionMechanic::Render(IToolsContextRenderAPI* RenderAPI)
 			PolyEdgesRenderer.DrawLine(A, B);
 		}
 	}
+	if (bShowSelectableCorners)
+	{
+		for (const FGroupTopology::FCorner& Corner : Topology->Corners)
+		{
+			FVector3d A = TargetMesh->GetVertex(Corner.VertexID);
+			PolyEdgesRenderer.DrawPoint(A);
+		}
+	}
 	PolyEdgesRenderer.EndFrame();
 
 	if (PersistentSelection.IsEmpty() == false)
@@ -877,6 +887,11 @@ FAxisAlignedBox3d UPolygonSelectionMechanic::GetSelectionBounds(bool bWorld) con
 	}
 }
 
+void UPolygonSelectionMechanic::SetShowSelectableCorners(bool bShowCorners)
+{
+	bShowSelectableCorners = bShowCorners;
+}
+
 
 
 void FPolygonSelectionMechanicSelectionChange::Apply(UObject* Object)
diff --git a/Engine/Plugins/Runtime/MeshModelingToolset/Source/ModelingComponents/Public/Selection/PolygonSelectionMechanic.h b/Engine/Plugins/Runtime/MeshModelingToolset/Source/ModelingComponents/Public/Selection/PolygonSelectionMechanic.h
--- a/Engine/Plugins/Runtime/MeshModelingToolset/Source/ModelingComponents/Public/Selection/PolygonSelectionMechanic.h
+++ b/Engine/Plugins/Runtime/MeshModelingToolset/Source/ModelingComponents/Public/Selection/PolygonSelectionMechanic.h
@@ -299,6 +299,8 @@ public:
 	 */
 	FAxisAlignedBox3d GetSelectionBounds(bool bWorld) const;
 
+	void SetShowSelectableCorners(bool bShowCorners);
+
 	//
 	// Change Tracking
 	//
@@ -439,6 +441,8 @@ protected:
 	bool bCtrlToggle = false;
 	static const int32 ShiftModifierID = 1;
 	static const int32 CtrlModifierID = 2;
+
+	bool bShowSelectableCorners = true;
 public:
 	FToolDataVisualizer PolyEdgesRenderer;
 	FToolDataVisualizer HilightRenderer;
