PolyEd Plane Control Gizmo: Add hotkey R to hide plane control gizmo in PolyEd

diff --git a/Engine/Plugins/Runtime/MeshModelingToolset/Source/MeshModelingTools/Private/EditMeshPolygonsTool.cpp b/Engine/Plugins/Runtime/MeshModelingToolset/Source/MeshModelingTools/Private/EditMeshPolygonsTool.cpp
--- a/Engine/Plugins/Runtime/MeshModelingToolset/Source/MeshModelingTools/Private/EditMeshPolygonsTool.cpp
+++ b/Engine/Plugins/Runtime/MeshModelingToolset/Source/MeshModelingTools/Private/EditMeshPolygonsTool.cpp
@@ -292,13 +292,21 @@ void UEditMeshPolygonsTool::Setup()
 	// Initialize the common properties but don't add them yet, because we want them to be under the activity-specific ones.
 	CommonProps = NewObject<UPolyEditCommonProperties>(this);
 	CommonProps->RestoreProperties(this, GetPropertyCacheIdentifier(bTriangleMode));
+
 	CommonProps->WatchProperty(CommonProps->LocalFrameMode,
 		[this](ELocalFrameMode) { UpdateGizmoFrame(); });
 	CommonProps->WatchProperty(CommonProps->bLockRotation,
+		[this](bool) { LockedTransfomerFrame = LastTransformerFrame; });
+	CommonProps->WatchProperty(CommonProps->bGizmoVisible,
 		[this](bool)
 		{
-			LockedTransfomerFrame = LastTransformerFrame;
+			if (!CurrentActivity)
+			{
+				UpdateGizmoVisibility();
+				ResetUserMessage();
+			}
 		});
+
 	// We are going to SilentUpdate here because otherwise the Watches above will immediately fire
 	// and cause UpdateGizmoFrame() to be called emitting a spurious Transform change. 
 	CommonProps->SilentUpdateWatched();
@@ -489,6 +497,24 @@ void UEditMeshPolygonsTool::Setup()
 	}
 }
 
+void UEditMeshPolygonsTool::ResetUserMessage()
+{
+	// When the gizmo is hidden, notify the user and
+	// specify the toggle hotkey to prevent panic
+	if (!TransformGizmo->IsVisible())
+	{
+		GetToolManager()->DisplayMessage(
+			LOCTEXT("ToggleTransformGizmoNotify", "Transform Gizmo Hidden, Press 'R' to Unhide"),
+			EToolMessageLevel::UserNotification);
+	}
+	else
+	{
+		GetToolManager()->DisplayMessage(
+			DefaultMessage,
+			EToolMessageLevel::UserNotification);
+	}
+}
+
 bool UEditMeshPolygonsTool::IsToolInputSelectionUsable(const UPersistentMeshSelection* InputSelectionIn)
 {
 	// TODO: We currently don't support persistent selection on volume brushes because
@@ -660,6 +686,19 @@ void UEditMeshPolygonsTool::RegisterActions(FInteractiveToolActionSet& ActionSet
 		LOCTEXT("DeleteSelectionTooltip", "Delete Selection"),
 		EModifierKey::None, EKeys::Delete, OnDeletionKeyPress);
 
+	ActionSet.RegisterAction(this, (int32)EStandardToolActions::BaseClientDefinedActionID + 5,
+		TEXT("ToggleTransformGizmoAKey"),
+		LOCTEXT("ToggleTransformGizmoUIName", "Toggle Transform Gizmo Visibility"),
+		LOCTEXT("ToggleTransformGizmoTooltip", "Toggle Transform Gizmo Visibility"),
+		EModifierKey::None, EKeys::R,
+		[this]() 
+		{
+			if (!CurrentActivity)
+			{
+				CommonProps->bGizmoVisible = !CommonProps->bGizmoVisible;
+			}
+		});
+
 	// TODO: Esc should be made to exit out of current activity if one is active. However this
 	// requires a bit of work because we don't seem to be able to register conditional actions,
 	// and we don't want to always capture Esc.
@@ -1112,10 +1151,6 @@ void UEditMeshPolygonsTool::EndCurrentActivity(EToolShutdownType ShutdownType)
 		if (CurrentActivity->IsRunning())
 		{
 			CurrentActivity->End(ShutdownType);
-
-			// Reset info message.
-			GetToolManager()->DisplayMessage(DefaultMessage,
-				EToolMessageLevel::UserNotification);
 		}
 
 		CurrentActivity = nullptr;
@@ -1128,6 +1163,10 @@ void UEditMeshPolygonsTool::EndCurrentActivity(EToolShutdownType ShutdownType)
 		SetToolPropertySourceEnabled(SelectionMechanic->Properties, true);
 		UpdateGizmoVisibility();
 	}
+
+	// If an activity displays a notification, it should be
+	// overwritten with an appropriate notification once finished
+	ResetUserMessage();
 }
 
 void UEditMeshPolygonsTool::NotifyActivitySelfEnded(UInteractiveToolActivity* Activity)
@@ -1137,24 +1176,23 @@ void UEditMeshPolygonsTool::NotifyActivitySelfEnded(UInteractiveToolActivity* Ac
 
 void UEditMeshPolygonsTool::UpdateGizmoVisibility()
 {
-	if (SelectionMechanic->HasSelection())
+	// Only allow gizmo to become visible if something is selected,
+	// the gizmo isn't hidden, and there is no current activity.
+	if (SelectionMechanic->HasSelection() && CommonProps->bGizmoVisible && !CurrentActivity)
 	{
 		if (TransformGizmo)
 		{
 			TransformGizmo->SetVisibility(true);
 		}
 
-		// update frame because we might be here due to an undo event/etc, rather than an explicit 
-		// selection change
+		// Update frame because we might be here due to an undo event/etc,
+		// rather than an explicit selection change
 		LastGeometryFrame = SelectionMechanic->GetSelectionFrame(true, &LastGeometryFrame);
 		UpdateGizmoFrame();
 	}
-	else
+	else if(TransformGizmo)
 	{
-		if (TransformGizmo)
-		{
-			TransformGizmo->SetVisibility(false);
-		}
+		TransformGizmo->SetVisibility(false);
 	}
 }
 
@@ -1193,7 +1231,6 @@ void UEditMeshPolygonsTool::OnPropertyModified(UObject* PropertySet, FProperty*
 	}
 }
 
-
 //
 // Gizmo change tracking
 //
diff --git a/Engine/Plugins/Runtime/MeshModelingToolset/Source/MeshModelingTools/Public/EditMeshPolygonsTool.h b/Engine/Plugins/Runtime/MeshModelingToolset/Source/MeshModelingTools/Public/EditMeshPolygonsTool.h
--- a/Engine/Plugins/Runtime/MeshModelingToolset/Source/MeshModelingTools/Public/EditMeshPolygonsTool.h
+++ b/Engine/Plugins/Runtime/MeshModelingToolset/Source/MeshModelingTools/Public/EditMeshPolygonsTool.h
@@ -78,6 +78,10 @@ public:
 	UPROPERTY(EditAnywhere, Category = Options)
 	bool bShowSelectableCorners = false;
 
+	/** When true, allows the transform gizmo to be rendered */
+	UPROPERTY(EditAnywhere, Category = Options)
+	bool bGizmoVisible = false;
+
 	/** Determines whether, on selection changes, the gizmo's rotation is taken from the object transform, or from the geometry
 	 elements selected. Only relevant with a local coordinate system and when rotation is not locked. */
 	UPROPERTY(EditAnywhere, Category = Gizmo, meta = (HideEditConditionToggle, EditCondition = "bLocalCoordSystem && !bLockRotation"))
@@ -564,6 +568,8 @@ protected:
 
 	FText DefaultMessage;
 
+	void ResetUserMessage();
+
 	bool IsToolInputSelectionUsable(const UPersistentMeshSelection* InputSelection);
 	bool bSelectionStateDirty = false;
 	void OnSelectionModifiedEvent();
