Add additional capture locations to FSceneCapturePhotoSet::AddStandardExteriorCapturesFromBoundingBox(), on upper and side box edge midpoints. Enable by default in Approximate Actors and BakeRenderCaptureTool.

diff --git a/Engine/Plugins/Experimental/MeshModelingToolsetExp/Source/GeometryProcessingAdapters/Private/ApproximateActorsImpl.cpp b/Engine/Plugins/Experimental/MeshModelingToolsetExp/Source/GeometryProcessingAdapters/Private/ApproximateActorsImpl.cpp
--- a/Engine/Plugins/Experimental/MeshModelingToolsetExp/Source/GeometryProcessingAdapters/Private/ApproximateActorsImpl.cpp
+++ b/Engine/Plugins/Experimental/MeshModelingToolsetExp/Source/GeometryProcessingAdapters/Private/ApproximateActorsImpl.cpp
@@ -112,7 +112,7 @@ static TUniquePtr<FSceneCapturePhotoSet> CapturePhotoSet(
 
 	SceneCapture->AddStandardExteriorCapturesFromBoundingBox(
 		CaptureDimensions, FieldOfView, NearPlaneDist,
-		true, true, true);
+		true, true, true, true, true);
 	
 	return SceneCapture;
 }
diff --git a/Engine/Plugins/Experimental/MeshModelingToolsetExp/Source/MeshModelingToolsExp/Private/BakeRenderCaptureTool.cpp b/Engine/Plugins/Experimental/MeshModelingToolsetExp/Source/MeshModelingToolsExp/Private/BakeRenderCaptureTool.cpp
--- a/Engine/Plugins/Experimental/MeshModelingToolsetExp/Source/MeshModelingToolsExp/Private/BakeRenderCaptureTool.cpp
+++ b/Engine/Plugins/Experimental/MeshModelingToolsetExp/Source/MeshModelingToolsExp/Private/BakeRenderCaptureTool.cpp
@@ -219,7 +219,7 @@ static TUniquePtr<FSceneCapturePhotoSet> CapturePhotoSet(
 
 	SceneCapture->AddStandardExteriorCapturesFromBoundingBox(
 		CaptureDimensions, FieldOfView, NearPlaneDist,
-		true, true, true);
+		true, true, true, true, true);
 
 	return SceneCapture;
 }
diff --git a/Engine/Plugins/Runtime/MeshModelingToolset/Source/ModelingComponents/Private/Scene/SceneCapturePhotoSet.cpp b/Engine/Plugins/Runtime/MeshModelingToolset/Source/ModelingComponents/Private/Scene/SceneCapturePhotoSet.cpp
--- a/Engine/Plugins/Runtime/MeshModelingToolset/Source/ModelingComponents/Private/Scene/SceneCapturePhotoSet.cpp
+++ b/Engine/Plugins/Runtime/MeshModelingToolset/Source/ModelingComponents/Private/Scene/SceneCapturePhotoSet.cpp
@@ -54,7 +54,9 @@ void FSceneCapturePhotoSet::AddStandardExteriorCapturesFromBoundingBox(
 	double NearPlaneDist,
 	bool bFaces,
 	bool bUpperCorners,
-	bool bLowerCorners)
+	bool bLowerCorners,
+	bool bUpperEdges,
+	bool bSideEdges)
 {
 	TArray<FVector3d> Directions;
 
@@ -81,6 +83,20 @@ void FSceneCapturePhotoSet::AddStandardExteriorCapturesFromBoundingBox(
 		Directions.Add(Normalized(FVector3d(1, -1, 1)));
 		Directions.Add(Normalized(FVector3d(-1, -1, 1)));
 	}
+	if (bUpperEdges)
+	{
+		Directions.Add(Normalized(FVector3d(-1, 0, -1)));
+		Directions.Add(Normalized(FVector3d(1, 0, -1)));
+		Directions.Add(Normalized(FVector3d(0, -1, -1)));
+		Directions.Add(Normalized(FVector3d(0, 1, -1)));
+	}
+	if (bSideEdges)
+	{
+		Directions.Add(Normalized(FVector3d(1, 1, 0)));
+		Directions.Add(Normalized(FVector3d(-1, 1, 0)));
+		Directions.Add(Normalized(FVector3d(1, -1, 0)));
+		Directions.Add(Normalized(FVector3d(-1, -1, 0)));
+	}
 	AddExteriorCaptures(PhotoDimensions, HorizontalFOVDegrees, NearPlaneDist, Directions);
 }
 
@@ -97,16 +113,18 @@ void FSceneCapturePhotoSet::AddExteriorCaptures(
 	FScopedSlowTask Progress(Directions.Num(), LOCTEXT("ComputingViewpoints", "Computing Viewpoints..."));
 	Progress.MakeDialog(bAllowCancel);
 
-	// Workaround for Nanite scene proxies visibility
-	// Unregister all components to remove unwanted proxies from the scene. This is currently the only way to "hide" nanite meshes.
+	// Unregister all components to remove unwanted proxies from the scene. This was previously the only way to "hide" nanite meshes, now optional.
 	TSet<AActor*> VisibleActorsSet(VisibleActors);
 	TArray<AActor*> ActorsToRegister;
-	for (TActorIterator<AActor> Actor(TargetWorld); Actor; ++Actor)
+	if (bEnforceVisibilityViaUnregister)
 	{
-		if (!VisibleActorsSet.Contains(*Actor))
+		for (TActorIterator<AActor> Actor(TargetWorld); Actor; ++Actor)
 		{
-			Actor->UnregisterAllComponents();
-			ActorsToRegister.Add(*Actor);
+			if (!VisibleActorsSet.Contains(*Actor))
+			{
+				Actor->UnregisterAllComponents();
+				ActorsToRegister.Add(*Actor);
+			}
 		}
 	}
 
@@ -353,6 +371,12 @@ bool FSceneCapturePhotoSet::ComputeSample(
 	return true;
 }
 
+
+void FSceneCapturePhotoSet::SetEnableVisibilityByUnregisterMode(bool bEnable)
+{
+	bEnforceVisibilityViaUnregister = bEnable;
+}
+
 void FSceneCapturePhotoSet::SetEnableWriteDebugImages(bool bEnable, FString FolderName)
 {
 	bWriteDebugImages = bEnable;
diff --git a/Engine/Plugins/Runtime/MeshModelingToolset/Source/ModelingComponents/Public/Scene/SceneCapturePhotoSet.h b/Engine/Plugins/Runtime/MeshModelingToolset/Source/ModelingComponents/Public/Scene/SceneCapturePhotoSet.h
--- a/Engine/Plugins/Runtime/MeshModelingToolset/Source/ModelingComponents/Public/Scene/SceneCapturePhotoSet.h
+++ b/Engine/Plugins/Runtime/MeshModelingToolset/Source/ModelingComponents/Public/Scene/SceneCapturePhotoSet.h
@@ -59,7 +59,9 @@ public:
 		double NearPlaneDist,
 		bool bFaces,
 		bool bUpperCorners,
-		bool bLowerCorners);
+		bool bLowerCorners,
+		bool bUpperEdges,
+		bool bSideEdges);
 
 	/**
 	 * Add captures on the "view sphere", ie a sphere centered/sized such that the target actors
@@ -141,6 +143,14 @@ public:
 	 */
 	void SetEnableWriteDebugImages(bool bEnable, FString FolderName = FString());
 
+	/**
+	 * If enabled, any Component Scene Proxies in the level that are not meant to be included in
+	 * the capture (ie not added via SetCaptureSceneActors), will be unregistered to hide them.
+	 * This is generally not necessary, and disabled by default, but in some cases the Renderer
+	 * may not be able to fully exclude the effects of an object via hidden/visible flags.
+	 */
+	void SetEnableVisibilityByUnregisterMode(bool bEnable);
+
 	void SetAllowCancel(bool bAllowCancelIn)
 	{
 		bAllowCancel = bAllowCancelIn;
@@ -155,6 +165,8 @@ protected:
 	UWorld* TargetWorld = nullptr;
 	TArray<AActor*> VisibleActors;
 
+	bool bEnforceVisibilityViaUnregister = false;
+
 	bool bEnableBaseColor = true;
 	bool bEnableRoughness = false;
 	bool bEnableSpecular = false;
