PolyEd Plane Control Gizmo: Change the default visibility of the transform gizmo in PolyEd to visible.

diff --git a/Engine/Plugins/Runtime/MeshModelingToolset/Source/MeshModelingTools/Public/EditMeshPolygonsTool.h b/Engine/Plugins/Runtime/MeshModelingToolset/Source/MeshModelingTools/Public/EditMeshPolygonsTool.h
--- a/Engine/Plugins/Runtime/MeshModelingToolset/Source/MeshModelingTools/Public/EditMeshPolygonsTool.h
+++ b/Engine/Plugins/Runtime/MeshModelingToolset/Source/MeshModelingTools/Public/EditMeshPolygonsTool.h
@@ -80,7 +80,7 @@ public:
 
 	/** When true, allows the transform gizmo to be rendered */
 	UPROPERTY(EditAnywhere, Category = Options)
-	bool bGizmoVisible = false;
+	bool bGizmoVisible = true;
 
 	/** Determines whether, on selection changes, the gizmo's rotation is taken from the object transform, or from the geometry
 	 elements selected. Only relevant with a local coordinate system and when rotation is not locked. */
