Add Bridge Edges to PolyEd: Add single to single groupedge bridge action to PolyEd, without welding sides

diff --git a/Engine/Plugins/Runtime/MeshModelingToolset/Source/MeshModelingTools/Private/EditMeshPolygonsTool.cpp b/Engine/Plugins/Runtime/MeshModelingToolset/Source/MeshModelingTools/Private/EditMeshPolygonsTool.cpp
--- a/Engine/Plugins/Runtime/MeshModelingToolset/Source/MeshModelingTools/Private/EditMeshPolygonsTool.cpp
+++ b/Engine/Plugins/Runtime/MeshModelingToolset/Source/MeshModelingTools/Private/EditMeshPolygonsTool.cpp
@@ -26,6 +26,7 @@
 #include "MeshRegionBoundaryLoops.h"
 #include "ModelingToolTargetUtil.h" // UE::ToolTarget:: functions
 #include "Operations/SimpleHoleFiller.h"
+#include "Operations/MinimalHoleFiller.h"
 #include "Operations/PolygroupRemesh.h"
 #include "Selection/PersistentMeshSelection.h"
 #include "Selection/StoredMeshSelectionUtil.h"
@@ -685,7 +686,6 @@ void UEditMeshPolygonsTool::RegisterActions(FInteractiveToolActionSet& ActionSet
 		LOCTEXT("DeleteSelectionUIName", "Delete Selection"),
 		LOCTEXT("DeleteSelectionTooltip", "Delete Selection"),
 		EModifierKey::None, EKeys::Delete, OnDeletionKeyPress);
-
 	ActionSet.RegisterAction(this, (int32)EStandardToolActions::BaseClientDefinedActionID + 5,
 		TEXT("ToggleTransformGizmoAKey"),
 		LOCTEXT("ToggleTransformGizmoUIName", "Toggle Transform Gizmo Visibility"),
@@ -1083,6 +1083,9 @@ void UEditMeshPolygonsTool::OnTick(float DeltaTime)
 		case EEditMeshPolygonsToolActions::FillHole:
 			ApplyFillHole();
 			break;
+		case EEditMeshPolygonsToolActions::BridgeEdges:
+			ApplyBridgeEdges();
+			break;
 		case EEditMeshPolygonsToolActions::Retriangulate:
 			ApplyRetriangulate();
 			break;
@@ -1790,7 +1793,109 @@ void UEditMeshPolygonsTool::ApplyFillHole()
 		ChangeTracker.EndChange(), NewSelection);
 }
 
+void UEditMeshPolygonsTool::ApplyBridgeEdges()
+{
+	if (SelectionMechanic->GetActiveSelection().SelectedEdgeIDs.Num() != 2 || BeginMeshBoundaryEdgeEditChange(false) == false)
+	{
+		GetToolManager()->DisplayMessage(LOCTEXT("OnEdgeBridgeFailed", "Cannot Bridge current selection"), EToolMessageLevel::UserWarning);
+		return;
+	}
+
+	FDynamicMesh3* Mesh = CurrentMesh.Get();
+	FDynamicMeshChangeTracker ChangeTracker(Mesh);
+	ChangeTracker.BeginChange();
+	FGroupTopologySelection CurrentSelection = SelectionMechanic->GetActiveSelection();
+
+	TArray<int32> LoopVertices;
+	TArray<int32> LoopEdges;
+	TArray<int32> SelectedEdgeIDs = CurrentSelection.SelectedEdgeIDs.Array();
+	
+	// I think doing this will guarantee that every edge in the span stores the vertices corresponding to the connected triangles orientation
+	FEdgeSpan& SpanA = Topology->Edges[SelectedEdgeIDs[0]].Span;
+	FEdgeSpan& SpanB = Topology->Edges[SelectedEdgeIDs[1]].Span;
+	SpanA.SetCorrectOrientation();
+	SpanB.SetCorrectOrientation();
+
+	// Disallow bridging of edge loops for now
+	if (SpanA.Vertices[0] == SpanA.Vertices.Last() || SpanB.Vertices[0] == SpanB.Vertices.Last())
+	{
+		GetToolManager()->DisplayMessage(LOCTEXT("OnEdgeBridgeFailed", "Cannot Bridge current selection, selected edges must not be loops"), EToolMessageLevel::UserWarning);
+		return;
+	}
+
+	// Add all vertices from first edge
+	LoopVertices = SpanA.Vertices;
+
+	// If first vertex of second edge is not a duplicate of a terminating vertex of the first edge, add vertex
+	if (SpanB.Vertices[0] != SpanA.Vertices[0] && SpanB.Vertices[0] != SpanA.Vertices.Last())
+	{
+		LoopVertices.Add(SpanB.Vertices[0]);
+	}
+
+	// Definitely add the non-terminating vertices of second edge.
+	for (int Vertex = 1; Vertex < SpanB.Vertices.Num() - 1; ++Vertex)
+	{
+		LoopVertices.Add(SpanB.Vertices[Vertex]);
+	}
 
+	// If last vertex of second edge is not a duplicate of a terminating vertex of the first edge, add vertex
+	if (SpanB.Vertices.Last() != SpanA.Vertices[0] && SpanB.Vertices.Last() != SpanA.Vertices.Last())
+	{
+		LoopVertices.Add(SpanB.Vertices.Last());
+	}
+
+	FEdgeLoop::VertexLoopToEdgeLoop(Mesh, LoopVertices, LoopEdges);
+	FEdgeLoop Loop(Mesh, LoopVertices, LoopEdges);
+
+	// We could always use the minimal hole filler, but it doesn't quite do what "bridge" would suggest when
+	// the area to be bridged is concave (across two curved-inward edges). Meanwhile simple ear clipping
+	// seems to fail in some common cases for reasons that we should investigate. For now, start with ear
+	// clipping, and revert to minimal if needed.
+	FSimpleHoleFiller SimpleHoleFiller(Mesh, Loop, FSimpleHoleFiller::EFillType::PolygonEarClipping);
+	TArray<int32> NewTriangles;
+
+	// Fill the hole
+	if (!SimpleHoleFiller.Fill())
+	{
+		//Ear clipping doesn't add vertices, so don't need to delete isolated verts
+		FDynamicMeshEditor Editor(Mesh);
+		Editor.RemoveTriangles(SimpleHoleFiller.NewTriangles, false);
+
+		FMinimalHoleFiller MinimalHoleFiller(Mesh, Loop);
+
+		if (!MinimalHoleFiller.Fill())
+		{
+			Editor.RemoveTriangles(MinimalHoleFiller.NewTriangles, false);
+			GetToolManager()->DisplayMessage(LOCTEXT("OnEdgeBridgeFailed", "Cannot Bridge current selection"), EToolMessageLevel::UserWarning);
+			return;
+		}
+		else
+		{
+			NewTriangles = MinimalHoleFiller.NewTriangles;
+		}
+	}
+	else {
+		NewTriangles = SimpleHoleFiller.NewTriangles;
+	}
+
+	// Compute normals and UVs
+	if (Mesh->HasAttributes())
+	{
+		TArray<FVector3d> VertexPositions;
+		Loop.GetVertices(VertexPositions);
+		FVector3d PlaneOrigin;
+		FVector3d PlaneNormal;
+		PolygonTriangulation::ComputePolygonPlane<double>(VertexPositions, PlaneNormal, PlaneOrigin);
+
+		FDynamicMeshEditor Editor(Mesh);
+		FFrame3d ProjectionFrame(PlaneOrigin, PlaneNormal);
+		Editor.SetTriangleNormals(NewTriangles);
+		Editor.SetTriangleUVsFromProjection(NewTriangles, ProjectionFrame, UVScaleFactor);
+	}
+
+	EmitCurrentMeshChangeAndUpdate(LOCTEXT("PolyMeshBridgeEdgeChange", "Bridge Edge"),
+		ChangeTracker.EndChange(), CurrentSelection);
+}
 
 
 void UEditMeshPolygonsTool::ApplyPokeSingleFace()
diff --git a/Engine/Plugins/Runtime/MeshModelingToolset/Source/MeshModelingTools/Public/EditMeshPolygonsTool.h b/Engine/Plugins/Runtime/MeshModelingToolset/Source/MeshModelingTools/Public/EditMeshPolygonsTool.h
--- a/Engine/Plugins/Runtime/MeshModelingToolset/Source/MeshModelingTools/Public/EditMeshPolygonsTool.h
+++ b/Engine/Plugins/Runtime/MeshModelingToolset/Source/MeshModelingTools/Public/EditMeshPolygonsTool.h
@@ -129,6 +129,7 @@ enum class EEditMeshPolygonsToolActions
 	WeldEdges,
 	StraightenEdge,
 	FillHole,
+	BridgeEdges,
 	BevelEdges,
 
 	PlanarProjectionUV,
@@ -368,7 +369,9 @@ public:
 
 	UFUNCTION(CallInEditor, Category = EdgeEdits, meta = (DisplayName = "Bevel", DisplayPriority = 4))
 	void Bevel() { PostAction(EEditMeshPolygonsToolActions::BevelEdges); }
-
+	
+	UFUNCTION(CallInEditor, Category = EdgeEdits, meta = (DisplayName = "Bridge", DisplayPriority = 5))
+	void Bridge() { PostAction(EEditMeshPolygonsToolActions::BridgeEdges); }
 };
 
 
@@ -631,6 +634,7 @@ protected:
 	void ApplyWeldEdges();
 	void ApplyStraightenEdges();
 	void ApplyFillHole();
+	void ApplyBridgeEdges();
 
 	void ApplyFlipSingleEdge();
 	void ApplyCollapseSingleEdge();
