EditMeshPolygonsTool.cpp warning: Text conflict from LOCTEXT macro for namespace and key

diff --git a/Engine/Plugins/Runtime/MeshModelingToolset/Source/MeshModelingTools/Private/EditMeshPolygonsTool.cpp b/Engine/Plugins/Runtime/MeshModelingToolset/Source/MeshModelingTools/Private/EditMeshPolygonsTool.cpp
--- a/Engine/Plugins/Runtime/MeshModelingToolset/Source/MeshModelingTools/Private/EditMeshPolygonsTool.cpp
+++ b/Engine/Plugins/Runtime/MeshModelingToolset/Source/MeshModelingTools/Private/EditMeshPolygonsTool.cpp
@@ -1795,9 +1795,11 @@ void UEditMeshPolygonsTool::ApplyFillHole()
 
 void UEditMeshPolygonsTool::ApplyBridgeEdges()
 {
+	const FText BridgeFailMessage = LOCTEXT("OnEdgeBridgeFailed", "Cannot Bridge current selection");
+
 	if (SelectionMechanic->GetActiveSelection().SelectedEdgeIDs.Num() != 2 || BeginMeshBoundaryEdgeEditChange(false) == false)
 	{
-		GetToolManager()->DisplayMessage(LOCTEXT("OnEdgeBridgeFailed", "Cannot Bridge current selection"), EToolMessageLevel::UserWarning);
+		GetToolManager()->DisplayMessage(BridgeFailMessage, EToolMessageLevel::UserWarning);
 		return;
 	}
 
@@ -1819,7 +1821,7 @@ void UEditMeshPolygonsTool::ApplyBridgeEdges()
 	// Disallow bridging of edge loops for now
 	if (SpanA.Vertices[0] == SpanA.Vertices.Last() || SpanB.Vertices[0] == SpanB.Vertices.Last())
 	{
-		GetToolManager()->DisplayMessage(LOCTEXT("OnEdgeBridgeFailed", "Cannot Bridge current selection, selected edges must not be loops"), EToolMessageLevel::UserWarning);
+		GetToolManager()->DisplayMessage(BridgeFailMessage, EToolMessageLevel::UserWarning);
 		return;
 	}
 
@@ -1866,7 +1868,7 @@ void UEditMeshPolygonsTool::ApplyBridgeEdges()
 		if (!MinimalHoleFiller.Fill())
 		{
 			Editor.RemoveTriangles(MinimalHoleFiller.NewTriangles, false);
-			GetToolManager()->DisplayMessage(LOCTEXT("OnEdgeBridgeFailed", "Cannot Bridge current selection"), EToolMessageLevel::UserWarning);
+			GetToolManager()->DisplayMessage(BridgeFailMessage, EToolMessageLevel::UserWarning);
 			return;
 		}
 		else
