Fix Editor Mac which broke in Horde Issue #180297: Compile errors in Module.ModelingComponents.cpp and SceneCapturePhotoSet.cpp

#rnx

diff --git a/Engine/Plugins/Runtime/MeshModelingToolset/Source/ModelingComponents/Private/Scene/SceneCapturePhotoSet.cpp b/Engine/Plugins/Runtime/MeshModelingToolset/Source/ModelingComponents/Private/Scene/SceneCapturePhotoSet.cpp
--- a/Engine/Plugins/Runtime/MeshModelingToolset/Source/ModelingComponents/Private/Scene/SceneCapturePhotoSet.cpp
+++ b/Engine/Plugins/Runtime/MeshModelingToolset/Source/ModelingComponents/Private/Scene/SceneCapturePhotoSet.cpp
@@ -321,9 +321,9 @@ bool FSceneCapturePhotoSet::ComputeSampleLocation(
 	double MinDot = 1.0;
 
 	int32 NumPhotos = PhotoSetParams.Num();
-	for (int32 pi = 0; pi < NumPhotos; ++pi)
+	for (int32 Index = 0; Index < NumPhotos; ++Index)
 	{
-		const FSpatialPhotoParams& Params = PhotoSetParams[pi];
+		const FSpatialPhotoParams& Params = PhotoSetParams[Index];
 		check(Params.Dimensions.IsSquare());
 
 		FVector3d ViewDirection = Params.Frame.X();
@@ -358,7 +358,7 @@ bool FSceneCapturePhotoSet::ComputeSampleLocation(
 				{
 					PhotoCoords.X = (u/2.0 + 0.5) * (double)Params.Dimensions.GetWidth();
 					PhotoCoords.Y = (v/2.0 + 0.5) * (double)Params.Dimensions.GetHeight();
-					PhotoIndex = pi;
+					PhotoIndex = Index;
 					MinDot = ViewDot;
 				}
 			}
diff --git a/Engine/Source/Runtime/GeometryCore/Public/Image/SpatialPhotoSet.h b/Engine/Source/Runtime/GeometryCore/Public/Image/SpatialPhotoSet.h
--- a/Engine/Source/Runtime/GeometryCore/Public/Image/SpatialPhotoSet.h
+++ b/Engine/Source/Runtime/GeometryCore/Public/Image/SpatialPhotoSet.h
@@ -133,9 +133,9 @@ bool TSpatialPhotoSet<PixelType, RealType>::ComputeSampleLocation(
 	double MinDot = 1.0;
 
 	int32 NumPhotos = Num();
-	for (int32 pi = 0; pi < NumPhotos; ++pi)
+	for (int32 Index = 0; Index < NumPhotos; ++Index)
 	{
-		const TSpatialPhoto<PixelType>& Photo = *Photos[pi];
+		const TSpatialPhoto<PixelType>& Photo = *Photos[Index];
 		check(Photo.Dimensions.IsSquare());
 
 		FVector3d ViewDirection = Photo.Frame.X();
@@ -170,7 +170,7 @@ bool TSpatialPhotoSet<PixelType, RealType>::ComputeSampleLocation(
 				{
 					PhotoCoords.X = (u/2.0 + 0.5) * (double)Photo.Dimensions.GetWidth();
 					PhotoCoords.Y = (v/2.0 + 0.5) * (double)Photo.Dimensions.GetHeight();
-					PhotoIndex = pi;
+					PhotoIndex = Index;
 					MinDot = ViewDot;
 				}
 			}
