SceneSnappingManager: add ::PauseSceneGeometryUpdates() and ::UnPauseSceneGeometryUpdates(), these temporarily disable/re-enable geometry updates (eg BVH rebuilds) on DynamicMeshComponents, for use in live-edit situations where we can't defer mesh updates (because render geometry needs to update) but don't want a per-frame BVH build

diff --git a/Engine/Plugins/Runtime/MeshModelingToolset/Source/ModelingComponents/Private/Snapping/ModelingSceneSnappingManager.cpp b/Engine/Plugins/Runtime/MeshModelingToolset/Source/ModelingComponents/Private/Snapping/ModelingSceneSnappingManager.cpp
--- a/Engine/Plugins/Runtime/MeshModelingToolset/Source/ModelingComponents/Private/Snapping/ModelingSceneSnappingManager.cpp
+++ b/Engine/Plugins/Runtime/MeshModelingToolset/Source/ModelingComponents/Private/Snapping/ModelingSceneSnappingManager.cpp
@@ -250,11 +250,43 @@ void UModelingSceneSnappingManager::HandleGlobalComponentTransformChangedDelegat
 
 void UModelingSceneSnappingManager::HandleDynamicMeshModifiedDelegate(UDynamicMeshComponent* Component)
 {
-	const bool bDeferRebuild = false;
-	SpatialCache->NotifyGeometryUpdate(Component, bDeferRebuild);
+	if (bQueueModifiedDynamicMeshUpdates)
+	{
+		PendingModifiedDynamicMeshes.Add(Component);
+	}
+	else
+	{
+		const bool bDeferRebuild = false;
+		SpatialCache->NotifyGeometryUpdate(Component, bDeferRebuild);
+	}
 }
 
 
+
+
+
+void UModelingSceneSnappingManager::PauseSceneGeometryUpdates()
+{
+	ensure(bQueueModifiedDynamicMeshUpdates == false);
+	bQueueModifiedDynamicMeshUpdates = true;
+}
+
+
+void UModelingSceneSnappingManager::UnPauseSceneGeometryUpdates(bool bImmediateRebuilds)
+{
+	if (ensure(bQueueModifiedDynamicMeshUpdates))
+	{
+		for (UDynamicMeshComponent* Component : PendingModifiedDynamicMeshes)
+		{
+			SpatialCache->NotifyGeometryUpdate(Component, !bImmediateRebuilds);
+		}
+		PendingModifiedDynamicMeshes.Reset();
+		bQueueModifiedDynamicMeshUpdates = false;
+	}
+}
+
+
+
 void UModelingSceneSnappingManager::BuildSpatialCacheForWorld(
 	UWorld* World,
 	TFunctionRef<bool(AActor*)> ActorFilter,
diff --git a/Engine/Plugins/Runtime/MeshModelingToolset/Source/ModelingComponents/Public/Snapping/ModelingSceneSnappingManager.h b/Engine/Plugins/Runtime/MeshModelingToolset/Source/ModelingComponents/Public/Snapping/ModelingSceneSnappingManager.h
--- a/Engine/Plugins/Runtime/MeshModelingToolset/Source/ModelingComponents/Public/Snapping/ModelingSceneSnappingManager.h
+++ b/Engine/Plugins/Runtime/MeshModelingToolset/Source/ModelingComponents/Public/Snapping/ModelingSceneSnappingManager.h
@@ -92,6 +92,26 @@ public:
 		TFunctionRef<bool(AActor*)> ActorFilter,
 		TFunctionRef<bool(UPrimitiveComponent*)> PrimitiveFilter );
 
+
+	//
+	// APIs for configuring the SceneSnappingManager behavior
+	//
+public:
+	/***
+	 * Temporarily disable live updates of any modified Actors/Components. Any changes that are
+	 * detected while updates are paused will be executed on Unpause.
+	 * This is mainly intended for use in interactive operations, eg while live-editing a mesh,
+	 * changes may be posted every frame but the cache BVHs/etc only need to be updated on mouse-up, etc
+	 */
+	virtual void PauseSceneGeometryUpdates();
+	/**
+	 * Re-enable live updates that were prevented by PauseSceneGeometryUpdates(), and execute
+	 * any pending updates that were requested while in the pause state.
+	 * @param bImmediateRebuilds If true (default), things like BVH updates will happen immediately, instead of simply being marked as pending
+	 */
+	virtual void UnPauseSceneGeometryUpdates(bool bImmediateRebuilds = true);
+
+
 protected:
 
 	UPROPERTY()
@@ -117,6 +137,10 @@ protected:
 
 	TMap<UPrimitiveComponent*, TWeakObjectPtr<UDynamicMeshComponent>> DynamicMeshComponents;
 	void HandleDynamicMeshModifiedDelegate(UDynamicMeshComponent* Component);
+
+	// PauseSceneGeometryUpdates / UnPauseSceneGeometryUpdates support
+	bool bQueueModifiedDynamicMeshUpdates = false;
+	TSet<UDynamicMeshComponent*> PendingModifiedDynamicMeshes;
 };
 
 
