ModelingComponents: add new geometric-element Selection system

GeometrySelection.h/cpp: FGeometrySelection, FGeometrySelectionEditor, representation and manipulation of a generic element-level selection, stored as a set of uint64. Editor makes changes to Selection while tracking deltas.
GeometrySelectionUtil.h/cpp: utility and helper functions for editing/processing a mesh+selection
GeometrySelector.h: FGeometryIdentifier type used to identify target objects with element-level selection support. IGeometrySelector interface that defines API for element selection for a type of target object. Factory type for spawing instances. IGeometrySelectionTransformer interface that defines API for element-level transformation
GeometrySelectionManager.h/cpp: UGeometrySelectionManager that maintains set of 'active targets' with associated Selector, Selection, and SelectionEditor, and provides high-level API for creating and querying element selections, 3D transforms of selections, selection undo/redo, and (currently) limited debug rendering
GeometrySelectionChanges.h/cpp: FChange types for element selection changes
DynamicMeshSelector.h/cpp: implementation of IGeometrySelector for UDynamicMesh/Component

diff --git a/Engine/Plugins/Runtime/MeshModelingToolset/Source/ModelingComponents/Private/Selection/DynamicMeshSelector.cpp b/Engine/Plugins/Runtime/MeshModelingToolset/Source/ModelingComponents/Private/Selection/DynamicMeshSelector.cpp
new file mode 100644
--- /dev/null
+++ b/Engine/Plugins/Runtime/MeshModelingToolset/Source/ModelingComponents/Private/Selection/DynamicMeshSelector.cpp
@@ -0,0 +1,478 @@
+// Copyright Epic Games, Inc. All Rights Reserved.
+
+
+#include "Selection/DynamicMeshSelector.h"
+#include "Selection/GeometrySelectionUtil.h"
+#include "ToolContextInterfaces.h"
+#include "DynamicMesh/DynamicMeshAttributeSet.h"
+#include "DynamicMesh/MeshIndexUtil.h"
+#include "DynamicMesh/MeshNormals.h"
+#include "Changes/MeshVertexChange.h"
+#include "DynamicMesh/ColliderMesh.h"
+#include "GroupTopology.h"
+
+using namespace UE::Geometry;
+
+#define LOCTEXT_NAMESPACE "FDynamicMeshSelector"
+
+
+void FDynamicMeshSelector::Initialize(
+	FGeometryIdentifier SourceGeometryIdentifierIn,
+	UDynamicMesh* TargetMeshIn, 
+	TUniqueFunction<UE::Geometry::FTransformSRT3d()> GetWorldTransformFuncIn)
+{
+	check(TargetMeshIn != nullptr);
+
+	SourceGeometryIdentifier = SourceGeometryIdentifierIn;
+	TargetMesh = TargetMeshIn;
+	GetWorldTransformFunc = MoveTemp(GetWorldTransformFuncIn);
+
+	RegisterMeshChangedHandler();
+}
+
+void FDynamicMeshSelector::RegisterMeshChangedHandler()
+{
+	TargetMesh_OnMeshChangedHandle = TargetMesh->OnMeshChanged().AddLambda([this](UDynamicMesh* Mesh, FDynamicMeshChangeInfo ChangeInfo)
+	{
+		ColliderMesh.Reset();
+
+		if (ChangeInfo.Type != EDynamicMeshChangeType::MeshVertexChange)
+		{
+			GroupTopology.Reset();
+		}
+
+		// publish geometry-modified event
+		NotifyGeometryModified();
+	});
+}
+
+
+void FDynamicMeshSelector::Shutdown()
+{
+	if (TargetMesh.IsValid())
+	{
+		TargetMesh->OnMeshChanged().Remove(TargetMesh_OnMeshChangedHandle);
+		TargetMesh_OnMeshChangedHandle.Reset();
+		
+		ColliderMesh.Reset();
+	}
+	TargetMesh = nullptr;
+}
+
+bool FDynamicMeshSelector::Sleep()
+{
+	check(TargetMesh.IsValid());
+
+	TargetMesh->OnMeshChanged().Remove(TargetMesh_OnMeshChangedHandle);
+	TargetMesh_OnMeshChangedHandle.Reset();
+
+	SleepingTargetMesh = TargetMesh;
+	TargetMesh = nullptr;
+
+	ColliderMesh.Reset();
+	GroupTopology.Reset();
+
+	return true;
+}
+
+bool FDynamicMeshSelector::Restore()
+{
+	check(SleepingTargetMesh.IsValid());
+
+	TargetMesh = SleepingTargetMesh.Get();
+	SleepingTargetMesh = nullptr;
+
+	RegisterMeshChangedHandler();
+
+	return true;
+}
+
+
+bool FDynamicMeshSelector::RayHitTest(
+	const FRay3d& WorldRay,
+	FInputRayHit& HitResultOut )
+{
+	FTransformSRT3d WorldTransform = GetWorldTransformFunc();
+	HitResultOut = FInputRayHit();
+	FRay3d LocalRay = WorldTransform.InverseTransformRay(WorldRay);
+	double RayHitT; int32 HitTriangleID; FVector3d HitBaryCoords;
+	if (GetColliderMesh()->FindNearestHitTriangle(LocalRay, RayHitT, HitTriangleID, HitBaryCoords))
+	{
+		HitResultOut.bHit = true;
+		HitResultOut.HitIdentifier = HitTriangleID;
+		HitResultOut.HitOwner = TargetMesh.Get();
+		FVector3d WorldPosition = WorldTransform.TransformPosition(LocalRay.PointAt(RayHitT));
+		HitResultOut.HitDepth = WorldRay.GetParameter(WorldPosition);
+	}
+	return HitResultOut.bHit;
+}
+
+
+void FDynamicMeshSelector::UpdateSelectionViaRaycast(
+	const FRay3d& WorldRay,
+	FGeometrySelectionEditor& SelectionEditor,
+	const FGeometrySelectionUpdateConfig& UpdateConfig,
+	FGeometrySelectionUpdateResult& ResultOut)
+{
+	FRay3d LocalRay = GetWorldTransformFunc().InverseTransformRay(WorldRay);
+
+	if (SelectionEditor.GetTopologyType() == EGeometryTopologyType::Polygroup)
+	{
+		UE::Geometry::UpdateGroupSelectionViaRaycast(
+			GetColliderMesh(), GetGroupTopology(), &SelectionEditor,
+			LocalRay, UpdateConfig, ResultOut);
+	}
+	else
+	{
+		UE::Geometry::UpdateTriangleSelectionViaRaycast(
+			GetColliderMesh(), &SelectionEditor,
+			LocalRay, UpdateConfig, ResultOut);
+	}
+}
+
+
+void FDynamicMeshSelector::GetSelectionFrame(const FGeometrySelection& Selection, FFrame3d& SelectionFrame, bool bTransformToWorld)
+{
+	if (Selection.TopologyType == EGeometryTopologyType::Polygroup)
+	{
+		const FGroupTopology* Topology = GetGroupTopology();
+		FGroupTopologySelection TopoSelection;
+		TargetMesh->ProcessMesh([&](const UE::Geometry::FDynamicMesh3& SourceMesh)
+		{
+			UE::Geometry::ConvertPolygroupSelectionToTopologySelection(Selection, SourceMesh, Topology, TopoSelection);
+		});
+
+		SelectionFrame = Topology->GetSelectionFrame(TopoSelection);
+	}
+	else
+	{
+		TargetMesh->ProcessMesh([&](const UE::Geometry::FDynamicMesh3& SourceMesh)
+		{
+			UE::Geometry::GetTriangleSelectionFrame(Selection, SourceMesh, SelectionFrame);
+		});
+	}
+
+	if (bTransformToWorld)
+	{
+		SelectionFrame.Transform(GetLocalToWorldTransform());
+	}
+}
+
+void FDynamicMeshSelector::AccumulateSelectionBounds(const FGeometrySelection& Selection, FGeometrySelectionBounds& BoundsInOut, bool bTransformToWorld)
+{
+	FTransform UseTransform = (bTransformToWorld) ? GetLocalToWorldTransform() : FTransform::Identity;
+	UE::Geometry::FAxisAlignedBox3d TargetWorldBounds = UE::Geometry::FAxisAlignedBox3d::Empty();
+	int32 ElementCount = 0;
+
+	if (Selection.TopologyType == EGeometryTopologyType::Polygroup)
+	{
+		const FGroupTopology* Topology = GetGroupTopology();
+		TargetMesh->ProcessMesh([&](const UE::Geometry::FDynamicMesh3& SourceMesh)
+		{
+			UE::Geometry::EnumeratePolygroupSelectionVertices(Selection, SourceMesh, GetGroupTopology(), UseTransform,
+				[&](uint32 VertexID, const FVector3d& Position) { 
+					TargetWorldBounds.Contain(Position); 
+					ElementCount++;
+				}
+			);
+		});
+	}
+	else
+	{
+		TargetMesh->ProcessMesh([&](const UE::Geometry::FDynamicMesh3& SourceMesh)
+		{
+			UE::Geometry::EnumerateTriangleSelectionVertices(Selection, SourceMesh, UseTransform,
+				[&](uint32 VertexID, const FVector3d& Position) { 
+					TargetWorldBounds.Contain(Position); 
+					ElementCount++;
+				}
+			);
+		});
+	}
+
+	if (ElementCount > 0)		// relying on this because in (eg) single-vertex case, the box will still be "empty"
+	{
+		BoundsInOut.WorldBounds.Contain(TargetWorldBounds);
+	}
+}
+
+
+
+void FDynamicMeshSelector::AccumulateSelectionElements(const FGeometrySelection& Selection, FGeometrySelectionElements& Elements, bool bTransformToWorld)
+{
+	FTransform UseWorldTransform = GetLocalToWorldTransform();
+	const FTransform* ApplyTransform = (bTransformToWorld) ? &UseWorldTransform : nullptr;
+
+	if (Selection.TopologyType == EGeometryTopologyType::Polygroup)
+	{
+		const FGroupTopology* Topology = GetGroupTopology();
+		TargetMesh->ProcessMesh([&](const UE::Geometry::FDynamicMesh3& SourceMesh)
+		{
+			UE::Geometry::EnumeratePolygroupSelectionElements(Selection, SourceMesh, Topology,
+				[&](uint32 VertexID, const FVector3d& Point) { Elements.Points.Add(Point); },
+				[&](uint32 EdgeID, const FSegment3d& Segment) { Elements.Segments.Add(Segment); },
+				[&](uint32 TriangleID, const FTriangle3d& Triangle) { Elements.Triangles.Add(Triangle); },
+				ApplyTransform
+			);
+		});
+	}
+	else
+	{
+		TargetMesh->ProcessMesh([&](const UE::Geometry::FDynamicMesh3& SourceMesh)
+		{
+			UE::Geometry::EnumerateTriangleSelectionElements(Selection, SourceMesh, 
+				[&](uint32 VertexID, const FVector3d& Point) { Elements.Points.Add(Point); },
+				[&](uint32 EdgeID, const FSegment3d& Segment) { Elements.Segments.Add(Segment); },
+				[&](uint32 TriangleID, const FTriangle3d& Triangle) { Elements.Triangles.Add(Triangle); },
+				ApplyTransform
+			);
+		});
+	}
+}
+
+
+void FDynamicMeshSelector::UpdateColliderMesh()
+{
+	// should we transform to world??
+	TargetMesh->ProcessMesh([&](const UE::Geometry::FDynamicMesh3& SourceMesh)
+	{
+		FColliderMesh::FBuildOptions BuildOptions;
+		BuildOptions.bBuildAABBTree = BuildOptions.bBuildVertexMap = BuildOptions.bBuildTriangleMap = true;
+		ColliderMesh = MakePimpl<FColliderMesh>(SourceMesh, BuildOptions);
+	});
+}
+
+const FColliderMesh* FDynamicMeshSelector::GetColliderMesh()
+{
+	if (!ColliderMesh.IsValid())
+	{
+		UpdateColliderMesh();
+	}
+	return ColliderMesh.Get();
+}
+
+
+void FDynamicMeshSelector::UpdateGroupTopology()
+{
+	// TODO would be preferable to not have to use the raw mesh pointer here, however
+	// FGroupTopology currently needs the pointer to do mesh queries
+
+	GroupTopology = MakePimpl<FGroupTopology>(TargetMesh->GetMeshPtr(), true);
+}
+
+const FGroupTopology* FDynamicMeshSelector::GetGroupTopology()
+{
+	if (!GroupTopology.IsValid())
+	{
+		UpdateGroupTopology();
+	}
+	return GroupTopology.Get();
+}
+
+
+
+bool FDynamicMeshComponentSelectorFactory::CanBuildForTarget(FGeometryIdentifier TargetIdentifier) const
+{
+	if (TargetIdentifier.TargetType == FGeometryIdentifier::ETargetType::PrimitiveComponent)
+	{
+		UDynamicMeshComponent* Component = TargetIdentifier.GetAsComponentType<UDynamicMeshComponent>();
+		if (Component)
+		{
+			return true;
+		}
+	}
+	return false;
+}
+
+TUniquePtr<IGeometrySelector> FDynamicMeshComponentSelectorFactory::BuildForTarget(FGeometryIdentifier TargetIdentifier) const
+{
+	UDynamicMeshComponent* Component = TargetIdentifier.GetAsComponentType<UDynamicMeshComponent>();
+	check(Component != nullptr);
+
+	TUniquePtr<FDynamicMeshSelector> Selector = MakeUnique<FDynamicMeshSelector>();
+	Selector->Initialize(TargetIdentifier, Component->GetDynamicMesh(),
+		[Component]() { return IsValid(Component) ? (UE::Geometry::FTransformSRT3d)Component->GetComponentTransform() : FTransformSRT3d::Identity(); }
+	);
+	return Selector;
+}
+
+
+
+
+
+
+
+/**
+* FDynamicMeshSelectionTransformer is a basic Transformer implementation for a
+* FDynamicMeshSelector. 
+*/
+class FDynamicMeshSelectionTransformer : public IGeometrySelectionTransformer
+{
+public:
+	FDynamicMeshSelector* Selector;
+
+	TArray<int32> MeshVertices;
+	TArray<FVector3d> InitialPositions;
+	TSet<int32> TriangleROI;
+	TSet<int32> OverlayNormals;
+
+	TArray<FVector3d> UpdatedPositions;
+
+	TPimplPtr<FMeshVertexChangeBuilder> ActiveVertexChange;
+
+	virtual IGeometrySelector* GetSelector() const override
+	{
+		return Selector;
+	}
+
+	virtual void BeginTransform(const FGeometrySelection& Selection) override;
+
+	virtual void UpdateTransform( TFunctionRef<FVector3d(int32 VertexID, const FVector3d& InitialPosition, const FTransform& WorldTransform)> PositionTransformFunc ) override;
+	void UpdatePendingVertexChange(bool bFinal);
+
+	virtual void EndTransform(IToolsContextTransactionsAPI* TransactionsAPI) override;
+};
+
+
+
+void FDynamicMeshSelectionTransformer::BeginTransform(const FGeometrySelection& Selection)
+{
+	const FGroupTopology* UseTopology = nullptr;
+	if (Selection.TopologyType == EGeometryTopologyType::Polygroup)
+	{
+		UseTopology = Selector->GetGroupTopology();
+	}
+
+	TSet<int32> VertexIDs;
+	Selector->GetDynamicMesh()->ProcessMesh([&](const FDynamicMesh3& SourceMesh)
+	{
+		// get set of selected vertex IDs, and then vector (only converting to vector for VertexToTriangleOneRing...)
+		if (Selection.TopologyType == EGeometryTopologyType::Polygroup)
+		{
+			UE::Geometry::EnumeratePolygroupSelectionVertices(Selection, SourceMesh, UseTopology, FTransform::Identity,
+				[&](uint32 VertexID, const FVector3d& Position) { VertexIDs.Add((int32)VertexID); }
+			);
+		}
+		else
+		{
+			UE::Geometry::EnumerateTriangleSelectionVertices(Selection, SourceMesh, FTransform::Identity,
+				[&](uint32 VertexID, const FVector3d& Position) { VertexIDs.Add((int32)VertexID); }
+			);
+		}
+		MeshVertices = VertexIDs.Array();
+
+		// save initial positions
+		InitialPositions.Reserve(MeshVertices.Num());
+		for (int32 vid : MeshVertices)
+		{
+			InitialPositions.Add(SourceMesh.GetVertex(vid));
+		}
+		UpdatedPositions.SetNum(MeshVertices.Num());
+
+		// get triangle ROI
+		UE::Geometry::VertexToTriangleOneRing(&SourceMesh, MeshVertices, TriangleROI);
+
+		// save overlay normals
+		if (SourceMesh.HasAttributes() && SourceMesh.Attributes()->PrimaryNormals() != nullptr)
+		{
+			const FDynamicMeshNormalOverlay* Normals = SourceMesh.Attributes()->PrimaryNormals();
+			UE::Geometry::TrianglesToOverlayElements(Normals, TriangleROI, OverlayNormals);
+		}
+	});
+
+	ActiveVertexChange = MakePimpl<FMeshVertexChangeBuilder>(EMeshVertexChangeComponents::VertexPositions | EMeshVertexChangeComponents::OverlayNormals);
+	UpdatePendingVertexChange(false);
+}
+
+void FDynamicMeshSelectionTransformer::UpdateTransform( 
+	TFunctionRef<FVector3d(int32 VertexID, const FVector3d& InitialPosition, const FTransform& WorldTransform)> PositionTransformFunc 
+)
+{
+	int32 N = MeshVertices.Num();
+	FTransform WorldTransform = Selector->GetLocalToWorldTransform();
+
+	for (int32 k = 0; k < N; ++k)
+	{
+		UpdatedPositions[k] = PositionTransformFunc(MeshVertices[k], InitialPositions[k], WorldTransform);
+	}
+
+	Selector->GetDynamicMesh()->EditMesh([&](FDynamicMesh3& EditMesh)
+	{
+		for (int32 k = 0; k < N; ++k)
+		{
+			EditMesh.SetVertex(MeshVertices[k], UpdatedPositions[k]);
+		}
+
+		// TODO: the right thing to do here would be to only recompute the modified overlay normals...but we have no utility for that??
+		//FMeshNormals::QuickRecomputeOverlayNormals(EditMesh);
+
+	}, EDynamicMeshChangeType::DeformationEdit,
+		//EDynamicMeshAttributeChangeFlags::VertexPositions, false);
+	   EDynamicMeshAttributeChangeFlags::VertexPositions | EDynamicMeshAttributeChangeFlags::NormalsTangents, false);
+
+	UpdatePendingVertexChange(false);
+}
+
+void FDynamicMeshSelectionTransformer::UpdatePendingVertexChange(bool bFinal)
+{
+	// update the vertex change
+	Selector->GetDynamicMesh()->ProcessMesh([&](const FDynamicMesh3& SourceMesh)
+	{
+		ActiveVertexChange->SaveVertices(&SourceMesh, MeshVertices, !bFinal);
+		ActiveVertexChange->SaveOverlayNormals(&SourceMesh, OverlayNormals, !bFinal);
+	});
+}
+
+void FDynamicMeshSelectionTransformer::EndTransform(IToolsContextTransactionsAPI* TransactionsAPI)
+{
+	UpdatePendingVertexChange(true);
+
+	if (TransactionsAPI != nullptr)
+	{
+		TUniquePtr<FToolCommandChange> Result = MoveTemp(ActiveVertexChange->Change);
+		TransactionsAPI->AppendChange(Selector->GetDynamicMesh(), MoveTemp(Result), LOCTEXT("DynamicMeshTransformChange", "Transform"));
+	}
+
+	ActiveVertexChange.Reset();
+}
+
+
+
+
+
+
+
+IGeometrySelectionTransformer* FDynamicMeshSelector::InitializeTransformation(const FGeometrySelection& Selection)
+{
+	check(!ActiveTransformer);
+
+	// If we are transforming a DynamicMeshComponent, we want to defer collision updates, otherwise 
+	// complex collision will be rebuilt every frame
+	FGeometryIdentifier ParentIdentifier = GetSourceGeometryIdentifier();
+	if (ParentIdentifier.ObjectType == FGeometryIdentifier::EObjectType::DynamicMeshComponent)
+	{
+		ParentIdentifier.GetAsComponentType<UDynamicMeshComponent>()->SetTransientDeferCollisionUpdates(true);
+	}
+
+	ActiveTransformer = MakePimpl<FDynamicMeshSelectionTransformer>();
+	ActiveTransformer->Selector = this;
+
+	return ActiveTransformer.Get();
+}
+
+void FDynamicMeshSelector::ShutdownTransformation(IGeometrySelectionTransformer* Transformer)
+{
+	ActiveTransformer.Reset();
+
+	FGeometryIdentifier ParentIdentifier = GetSourceGeometryIdentifier();
+	if (ParentIdentifier.ObjectType == FGeometryIdentifier::EObjectType::DynamicMeshComponent)
+	{
+		ParentIdentifier.GetAsComponentType<UDynamicMeshComponent>()->SetTransientDeferCollisionUpdates(false);
+		ParentIdentifier.GetAsComponentType<UDynamicMeshComponent>()->UpdateCollision(true);
+	}
+}
+
+
+
+
+#undef LOCTEXT_NAMESPACE 
\ No newline at end of file
diff --git a/Engine/Plugins/Runtime/MeshModelingToolset/Source/ModelingComponents/Private/Selection/GeometrySelection.cpp b/Engine/Plugins/Runtime/MeshModelingToolset/Source/ModelingComponents/Private/Selection/GeometrySelection.cpp
new file mode 100644
--- /dev/null
+++ b/Engine/Plugins/Runtime/MeshModelingToolset/Source/ModelingComponents/Private/Selection/GeometrySelection.cpp
@@ -0,0 +1,51 @@
+// Copyright Epic Games, Inc. All Rights Reserved.
+
+#include "Selection/GeometrySelection.h"
+#include "DynamicMesh/DynamicMesh3.h"
+
+
+void FGeometrySelectionEditor::Initialize(FGeometrySelection* TargetSelectionIn)
+{
+	check(TargetSelectionIn != nullptr);
+	TargetSelection = TargetSelectionIn;
+}
+
+void FGeometrySelectionEditor::ClearSelection(FGeometrySelectionDelta& DeltaOut)
+{
+	for (uint64 ID : TargetSelection->Selection)
+	{
+		DeltaOut.Removed.Add(ID);
+	}
+
+	TargetSelection->Reset();
+}
+
+
+bool FGeometrySelectionEditor::IsSelected(uint64 ID) const
+{
+	return TargetSelection->Selection.Contains(ID);
+}
+
+bool FGeometrySelectionEditor::Replace(const FGeometrySelection& NewSelection, FGeometrySelectionDelta& DeltaOut)
+{
+	// currently hard-enforcing this...maybe we can allow switches? 
+	if (ensure(TargetSelection->ElementType == NewSelection.ElementType) == false ||
+		ensure(TargetSelection->TopologyType == NewSelection.TopologyType) == false)
+	{
+		return false;
+	}
+
+	for (uint64 ID : TargetSelection->Selection)
+	{
+		DeltaOut.Removed.Add(ID);
+	}
+	TargetSelection->Selection.Reset();
+
+	for (uint64 ID : NewSelection.Selection)
+	{
+		TargetSelection->Selection.Add(ID);
+		DeltaOut.Added.Add(ID);
+	}
+
+	return true;
+}
\ No newline at end of file
diff --git a/Engine/Plugins/Runtime/MeshModelingToolset/Source/ModelingComponents/Private/Selection/GeometrySelectionChanges.cpp b/Engine/Plugins/Runtime/MeshModelingToolset/Source/ModelingComponents/Private/Selection/GeometrySelectionChanges.cpp
new file mode 100644
--- /dev/null
+++ b/Engine/Plugins/Runtime/MeshModelingToolset/Source/ModelingComponents/Private/Selection/GeometrySelectionChanges.cpp
@@ -0,0 +1,109 @@
+// Copyright Epic Games, Inc. All Rights Reserved.
+
+
+#include "Selection/GeometrySelectionChanges.h"
+#include "Selection/GeometrySelectionManager.h"
+#include "Selection/GeometrySelectionUtil.h"
+
+#define LOCTEXT_NAMESPACE "GeometrySelectionChanges"
+
+
+void FGeometrySelectionDeltaChange::Apply(UObject* Object)
+{
+	UGeometrySelectionManager* Manager = Cast<UGeometrySelectionManager>(Object);
+	if ( ensureMsgf(Manager, TEXT("FGeometrySelectionDeltaChange must be emitted on a UGeometrySelectionManager")) )
+	{
+		Manager->ApplyChange(this);
+	}
+}
+
+void FGeometrySelectionDeltaChange::Revert(UObject* Object)
+{
+	UGeometrySelectionManager* Manager = Cast<UGeometrySelectionManager>(Object);
+	if ( ensureMsgf(Manager, TEXT("FGeometrySelectionDeltaChange must be emitted on a UGeometrySelectionManager")) )
+	{
+		Manager->RevertChange(this);
+	}
+}
+
+FString FGeometrySelectionDeltaChange::ToString() const
+{
+	return TEXT("FGeometrySelectionDeltaChange");
+}
+
+bool FGeometrySelectionDeltaChange::HasExpired(UObject* Object) const
+{
+	UGeometrySelectionManager* Manager = Cast<UGeometrySelectionManager>(Object);
+	return (Manager == nullptr || IsValid(Manager) == false || Manager->HasBeenShutDown());
+}
+
+
+void FGeometrySelectionDeltaChange::ApplyChange(FGeometrySelectionEditor* Editor, FGeometrySelectionDelta& ApplyDelta)
+{
+	if (Delta.Removed.Num() > 0)
+	{
+		UE::Geometry::UpdateSelectionWithNewElements(Editor, EGeometrySelectionChangeType::Remove, Delta.Removed, &ApplyDelta);
+	}
+	if (Delta.Added.Num() > 0)
+	{
+		UE::Geometry::UpdateSelectionWithNewElements(Editor, EGeometrySelectionChangeType::Add, Delta.Added, &ApplyDelta);
+	}
+}
+void FGeometrySelectionDeltaChange::RevertChange(FGeometrySelectionEditor* Editor, FGeometrySelectionDelta& RevertDelta)
+{
+	if (Delta.Added.Num() > 0)
+	{
+		UE::Geometry::UpdateSelectionWithNewElements(Editor, EGeometrySelectionChangeType::Remove, Delta.Added, &RevertDelta);
+	}
+	if (Delta.Removed.Num() > 0)
+	{
+		UE::Geometry::UpdateSelectionWithNewElements(Editor, EGeometrySelectionChangeType::Add, Delta.Removed, &RevertDelta);
+	}
+}
+
+
+
+
+
+
+void FGeometrySelectionReplaceChange::Apply(UObject* Object)
+{
+	UGeometrySelectionManager* Manager = Cast<UGeometrySelectionManager>(Object);
+	if ( ensureMsgf(Manager, TEXT("FGeometrySelectionReplaceChange must be emitted on a UGeometrySelectionManager")) )
+	{
+		Manager->ApplyChange(this);
+	}
+}
+
+void FGeometrySelectionReplaceChange::Revert(UObject* Object)
+{
+	UGeometrySelectionManager* Manager = Cast<UGeometrySelectionManager>(Object);
+	if ( ensureMsgf(Manager, TEXT("FGeometrySelectionReplaceChange must be emitted on a UGeometrySelectionManager")) )
+	{
+		Manager->RevertChange(this);
+	}
+}
+
+FString FGeometrySelectionReplaceChange::ToString() const
+{
+	return TEXT("FGeometrySelectionReplaceChange");
+}
+
+bool FGeometrySelectionReplaceChange::HasExpired(UObject* Object) const
+{
+	UGeometrySelectionManager* Manager = Cast<UGeometrySelectionManager>(Object);
+	return (Manager == nullptr || IsValid(Manager) == false || Manager->HasBeenShutDown());
+}
+
+void FGeometrySelectionReplaceChange::ApplyChange(FGeometrySelectionEditor* Editor, FGeometrySelectionDelta& ApplyDelta)
+{
+	Editor->Replace(After, ApplyDelta);
+}
+void FGeometrySelectionReplaceChange::RevertChange(FGeometrySelectionEditor* Editor, FGeometrySelectionDelta& RevertDelta)
+{
+	Editor->Replace(Before, RevertDelta);
+}
+
+
+
+#undef LOCTEXT_NAMESPACE
\ No newline at end of file
diff --git a/Engine/Plugins/Runtime/MeshModelingToolset/Source/ModelingComponents/Private/Selection/GeometrySelectionManager.cpp b/Engine/Plugins/Runtime/MeshModelingToolset/Source/ModelingComponents/Private/Selection/GeometrySelectionManager.cpp
new file mode 100644
--- /dev/null
+++ b/Engine/Plugins/Runtime/MeshModelingToolset/Source/ModelingComponents/Private/Selection/GeometrySelectionManager.cpp
@@ -0,0 +1,875 @@
+// Copyright Epic Games, Inc. All Rights Reserved.
+
+#include "Selection/GeometrySelectionManager.h"
+#include "Selection/DynamicMeshSelector.h"
+#include "ToolContextInterfaces.h"
+#include "ToolDataVisualizer.h"
+
+// for debug drawing
+#include "SceneManagement.h"		
+#include "DynamicMeshBuilder.h"
+#include "ToolSetupUtil.h"
+
+#include "Selection/PersistentMeshSelection.h"
+
+
+using namespace UE::Geometry;
+
+
+#define LOCTEXT_NAMESPACE "UGeometrySelectionManager"
+
+
+void UGeometrySelectionManager::Initialize( UInteractiveToolsContext* ToolsContextIn, IToolsContextTransactionsAPI* TransactionsAPIIn )
+{
+	ToolsContext = ToolsContextIn;
+	TransactionsAPI = TransactionsAPIIn;
+}
+
+void UGeometrySelectionManager::RegisterSelectorFactory(TUniquePtr<IGeometrySelectorFactory> Factory)
+{
+	Factories.Add(MoveTemp(Factory));
+	ResetTargetCache();
+}
+
+void UGeometrySelectionManager::Shutdown()
+{
+	OnSelectionModified.Clear();
+	ToolsContext = nullptr;
+	TransactionsAPI = nullptr;
+
+	ResetTargetCache();
+}
+
+bool UGeometrySelectionManager::HasBeenShutDown() const
+{
+	return (ToolsContext == nullptr);
+}
+
+
+class FGeometrySelectionManager_SelectionTypeChange : public FToolCommandChange
+{
+public:
+	EGeometryElementType FromElementType = EGeometryElementType::Face;
+	EGeometryElementType ToElementType = EGeometryElementType::Face;
+
+	UGeometrySelectionManager::EMeshTopologyMode FromTopologyMode = UGeometrySelectionManager::EMeshTopologyMode::None;
+	UGeometrySelectionManager::EMeshTopologyMode ToTopologyMode = UGeometrySelectionManager::EMeshTopologyMode::None;
+
+	/** Makes the change to the object */
+	virtual void Apply(UObject* Object) override
+	{
+		if (FromElementType != ToElementType)
+		{
+			CastChecked<UGeometrySelectionManager>(Object)->SetSelectionElementTypeInternal(ToElementType);
+		}
+		if (FromTopologyMode != ToTopologyMode)
+		{
+			CastChecked<UGeometrySelectionManager>(Object)->SetMeshTopologyModeInternal(ToTopologyMode);
+		}
+	}
+
+	/** Reverts change to the object */
+	virtual void Revert(UObject* Object) override
+	{
+		if (FromElementType != ToElementType)
+		{
+			CastChecked<UGeometrySelectionManager>(Object)->SetSelectionElementTypeInternal(FromElementType);
+		}
+		if (FromTopologyMode != ToTopologyMode)
+		{
+			CastChecked<UGeometrySelectionManager>(Object)->SetMeshTopologyModeInternal(FromTopologyMode);
+		}
+	}
+
+	/** Describes this change (for debugging) */
+	virtual FString ToString() const override { return TEXT("FGeometrySelectionManager_SelectionTypeChange"); }
+
+	virtual bool HasExpired(UObject* Object) const override
+	{
+		UGeometrySelectionManager* Manager = Cast<UGeometrySelectionManager>(Object);
+		return (Manager == nullptr || IsValid(Manager) == false || Manager->HasBeenShutDown());
+	}
+};
+
+
+
+void UGeometrySelectionManager::SetSelectionElementTypeInternal(EGeometryElementType NewElementType)
+{
+	if (SelectionElementType != NewElementType)
+	{
+		SelectionElementType = NewElementType;
+
+		for (TSharedPtr<FGeometrySelectionTarget> Target : ActiveTargetReferences)
+		{
+			Target->Selection.ElementType = SelectionElementType;
+			// do we need to update/rebuild Editors here? They currently do not have state...
+		}
+	}
+}
+
+
+void UGeometrySelectionManager::SetSelectionElementType(EGeometryElementType NewElementType)
+{
+	if (SelectionElementType != NewElementType)
+	{
+		GetTransactionsAPI()->BeginUndoTransaction(LOCTEXT("ChangeElementType", "Selection Type"));
+
+		if (HasSelection())
+		{
+			ClearSelection();
+		}
+
+		// We have to undo/redo the change to the selection type because if we want to 'undo' this later and restore
+		// the current selection, we need the active element type to be correct. Note that it goes *after* the Clear
+		// so that when we undo, we change to the correct type before we restore
+		TUniquePtr<FGeometrySelectionManager_SelectionTypeChange> TypeChange = MakeUnique<FGeometrySelectionManager_SelectionTypeChange>();
+		TypeChange->FromElementType = SelectionElementType;
+		TypeChange->ToElementType = NewElementType;
+		GetTransactionsAPI()->AppendChange(this, MoveTemp(TypeChange), LOCTEXT("ChangeElementType", "Selection Type"));
+
+		SetSelectionElementTypeInternal(NewElementType);
+
+		GetTransactionsAPI()->EndUndoTransaction();
+	}
+}
+
+
+
+
+void UGeometrySelectionManager::SetMeshTopologyModeInternal(EMeshTopologyMode NewTopologyMode)
+{
+	if (MeshTopologyMode != NewTopologyMode)
+	{
+		MeshTopologyMode = NewTopologyMode;
+
+		for (TSharedPtr<FGeometrySelectionTarget> Target : ActiveTargetReferences)
+		{
+			Target->Selection.TopologyType = GetSelectionTopologyType();
+			// do we need to update/rebuild Editors here? They currently do not have state...
+		}
+	}
+}
+
+
+void UGeometrySelectionManager::SetMeshTopologyMode(EMeshTopologyMode NewTopologyMode)
+{
+	if (MeshTopologyMode != NewTopologyMode)
+	{
+		GetTransactionsAPI()->BeginUndoTransaction(LOCTEXT("ChangeSelectionMode", "Selection Mode"));
+
+		if (HasSelection())
+		{
+			ClearSelection();
+		}
+
+		// We have to undo/redo the change to the selection type because if we want to 'undo' this later and restore
+		// the current selection, we need the active element type to be correct. Note that it goes *after* the Clear
+		// so that when we undo, we change to the correct type before we restore
+		TUniquePtr<FGeometrySelectionManager_SelectionTypeChange> TypeChange = MakeUnique<FGeometrySelectionManager_SelectionTypeChange>();
+		TypeChange->FromTopologyMode = MeshTopologyMode;
+		TypeChange->FromTopologyMode = NewTopologyMode;
+		GetTransactionsAPI()->AppendChange(this, MoveTemp(TypeChange),LOCTEXT("ChangeSelectionMode", "Selection Mode"));
+
+		SetMeshTopologyModeInternal(NewTopologyMode);
+
+		GetTransactionsAPI()->EndUndoTransaction();
+	}
+}
+
+
+
+EGeometryTopologyType UGeometrySelectionManager::GetSelectionTopologyType() const
+{
+	if (MeshTopologyMode == EMeshTopologyMode::Polygroup)
+	{
+		return EGeometryTopologyType::Polygroup;
+	}
+	else
+	{
+		return EGeometryTopologyType::Triangle;
+	}
+}
+
+
+
+bool UGeometrySelectionManager::HasSelection() const
+{
+	for (TSharedPtr<FGeometrySelectionTarget> Target : ActiveTargetReferences)
+	{
+		if (Target->Selection.IsEmpty() == false)
+		{
+			return true;
+		}
+	}
+	return false;
+}
+
+
+
+
+
+class FGeometrySelectionManager_ActiveTargetsChange : public FToolCommandChange
+{
+public:
+	TArray<FGeometryIdentifier> TargetsBefore;
+	TArray<FGeometryIdentifier> TargetsAfter;
+
+	virtual void Apply(UObject* Object) override
+	{
+		CastChecked<UGeometrySelectionManager>(Object)->SetTargetsOnUndoRedo(TargetsAfter);
+	}
+
+	virtual void Revert(UObject* Object) override
+	{
+		CastChecked<UGeometrySelectionManager>(Object)->SetTargetsOnUndoRedo(TargetsBefore);
+	}
+
+	virtual FString ToString() const override { return TEXT("FGeometrySelectionManager_ActiveTargetsChange"); }
+
+	virtual bool HasExpired(UObject* Object) const override
+	{
+		UGeometrySelectionManager* Manager = Cast<UGeometrySelectionManager>(Object);
+		return (Manager == nullptr || IsValid(Manager) == false || Manager->HasBeenShutDown());
+	}
+};
+
+
+
+
+
+bool UGeometrySelectionManager::HasActiveTargets() const
+{
+	return (ActiveTargetReferences.Num() > 0);
+}
+
+
+
+void UGeometrySelectionManager::ClearActiveTargets()
+{
+	// generally at this point is it too late to clear the selection, because it will emit an
+	// undo that cannot be redone later, because on redo the Targets will not exist yet
+	// (one possibility would be to emit separate changes for when the target set is modified?? would that work w/ delete?? )
+	ensure(HasSelection() == false);
+
+	for (TSharedPtr<FGeometrySelectionTarget> Target : ActiveTargetReferences)
+	{
+		SleepOrShutdownTarget(Target.Get(), false);
+	}
+
+	ActiveTargetReferences.Reset();
+	ActiveTargetMap.Reset();
+
+	UpdateSelectionRenderCacheOnTargetChange();
+
+	OnSelectionModified.Broadcast();
+}
+
+bool UGeometrySelectionManager::AddActiveTarget(FGeometryIdentifier TargetIdentifier)
+{
+	if (ActiveTargetMap.Contains(TargetIdentifier))
+	{
+		return false;
+	}
+
+	// need to have a selector factory that can build for this target
+	const IGeometrySelectorFactory* UseFactory = nullptr;
+	for (const TUniquePtr<IGeometrySelectorFactory>& Factory : Factories)
+	{
+		if (Factory->CanBuildForTarget(TargetIdentifier))
+		{
+			UseFactory = Factory.Get();
+			break;
+		}
+	}
+	if (UseFactory == nullptr)
+	{
+		return false;
+	}
+
+	TSharedPtr<FGeometrySelectionTarget> SelectionTarget = GetCachedTarget(TargetIdentifier, UseFactory);
+	check(SelectionTarget != nullptr);
+
+	ActiveTargetMap.Add(TargetIdentifier, SelectionTarget);
+	ActiveTargetReferences.Add(SelectionTarget);
+
+	SelectionTarget->OnGeometryModifiedHandle = 
+		SelectionTarget->Selector->GetOnGeometryModifed().AddUObject(this, &UGeometrySelectionManager::OnTargetGeometryModified);
+
+	UpdateSelectionRenderCacheOnTargetChange();
+
+	return true;
+}
+
+
+void UGeometrySelectionManager::SynchronizeActiveTargets(
+	const TArray<FGeometryIdentifier>& DesiredActiveSet,
+	TFunctionRef<void()> WillChangeActiveTargetsCallback)
+{
+	TArray<FGeometryIdentifier> Before = GetCurrentTargetIdentifiers();
+
+	// currently only support single selection
+	if (DesiredActiveSet.Num() == 1)
+	{
+		// if we do not already have this target, select it
+		if ( ActiveTargetMap.Contains(DesiredActiveSet[0]) == false )
+		{
+			WillChangeActiveTargetsCallback();
+			ClearActiveTargets();
+			AddActiveTarget(DesiredActiveSet[0]);
+		}
+	}
+	else
+	{
+		WillChangeActiveTargetsCallback();
+		ClearActiveTargets();
+	}
+
+	TArray<FGeometryIdentifier> After = GetCurrentTargetIdentifiers();
+	if (Before != After)
+	{
+		TUniquePtr<FGeometrySelectionManager_ActiveTargetsChange> Change = MakeUnique<FGeometrySelectionManager_ActiveTargetsChange>();
+		Change->TargetsBefore = Before;
+		Change->TargetsAfter = After;
+		GetTransactionsAPI()->AppendChange(this, MoveTemp(Change), LOCTEXT("Change Targets", "Change Targets"));
+	}
+}
+
+TArray<FGeometryIdentifier> UGeometrySelectionManager::GetCurrentTargetIdentifiers() const
+{
+	TArray<FGeometryIdentifier> Result;
+	for (TSharedPtr<FGeometrySelectionTarget> Target : ActiveTargetReferences)
+	{
+		Result.Add(Target->TargetIdentifier);
+	}
+	return Result;
+}
+
+void UGeometrySelectionManager::SetTargetsOnUndoRedo(TArray<FGeometryIdentifier> NewTargets)
+{
+	check(HasSelection() == false);
+	ClearActiveTargets();
+	for (FGeometryIdentifier Identifier : NewTargets)
+	{
+		AddActiveTarget(Identifier);
+	}
+}
+
+
+
+void UGeometrySelectionManager::SleepOrShutdownTarget(FGeometrySelectionTarget* Target, bool bForceShutdown)
+{
+	if (Target->Selector->SupportsSleep() && bForceShutdown == false)
+	{
+		bool bOK = Target->Selector->Sleep();
+		check(bOK);
+	}
+	else
+	{
+		Target->Selector->GetOnGeometryModifed().Remove(Target->OnGeometryModifiedHandle);
+		Target->Selector->Shutdown();
+	}
+}
+
+TSharedPtr<UGeometrySelectionManager::FGeometrySelectionTarget> UGeometrySelectionManager::GetCachedTarget(FGeometryIdentifier TargetIdentifier, const IGeometrySelectorFactory* UseFactory)
+{
+	if (TargetCache.Contains(TargetIdentifier))
+	{
+		TSharedPtr<FGeometrySelectionTarget> FoundTarget = TargetCache[TargetIdentifier];
+		FoundTarget->Selection.Reset();
+		FoundTarget->Selector->Restore();
+
+		// ensure these are current, as they may have changed while Target was asleep
+		FoundTarget->Selection.ElementType = GetSelectionElementType();
+		FoundTarget->Selection.TopologyType = GetSelectionTopologyType();
+
+		return FoundTarget;
+	}
+
+	// if we are in a situation where we don't have a cache, currently we need the Factory to exist?
+	check(UseFactory != nullptr);
+
+	// selector has to be built properly
+	TUniquePtr<IGeometrySelector> Selector = UseFactory->BuildForTarget(TargetIdentifier);
+	// not going to handle this for now...
+	check(Selector.IsValid());
+
+	TSharedPtr<FGeometrySelectionTarget> SelectionTarget = MakeShared<FGeometrySelectionTarget>();
+	SelectionTarget->Selector = MoveTemp(Selector);
+	SelectionTarget->TargetIdentifier = TargetIdentifier;
+	SelectionTarget->SelectionIdentifer = SelectionTarget->Selector->GetIdentifier();
+	SelectionTarget->Selection.ElementType = GetSelectionElementType();
+	SelectionTarget->Selection.TopologyType = GetSelectionTopologyType();
+
+	SelectionTarget->SelectionEditor = MakeUnique<FGeometrySelectionEditor>();
+	SelectionTarget->SelectionEditor->Initialize(&SelectionTarget->Selection);
+
+	if (SelectionTarget->Selector->SupportsSleep())
+	{
+		TargetCache.Add(TargetIdentifier, SelectionTarget);
+	}
+
+	return SelectionTarget;
+}
+
+void UGeometrySelectionManager::ResetTargetCache()
+{
+	for (TPair<FGeometryIdentifier, TSharedPtr<FGeometrySelectionTarget>> Pair : TargetCache)
+	{
+		SleepOrShutdownTarget(Pair.Value.Get(), true);
+	}
+	TargetCache.Reset();
+}
+
+
+
+bool UGeometrySelectionManager::RayHitTest(
+	const FRay3d& WorldRay,
+	FInputRayHit& HitResultOut )
+{
+	HitResultOut = FInputRayHit();
+	if (ActiveTargetReferences.Num() == 0)
+	{
+		return false;
+	}
+
+	bool bHit = ActiveTargetReferences[0]->Selector->RayHitTest(WorldRay, HitResultOut);
+	if (bHit)
+	{
+		HitResultOut.HitOwner = ActiveTargetReferences[0].Get();
+		HitResultOut.HitObject = (ActiveTargetReferences[0]->TargetIdentifier.TargetType == FGeometryIdentifier::ETargetType::PrimitiveComponent) ?
+			ActiveTargetReferences[0]->TargetIdentifier.TargetObject : nullptr;
+	}
+
+	// currently only going to support one object, not sure how to support more yet...
+	return bHit;
+}
+
+
+void UGeometrySelectionManager::ClearSelection()
+{
+	if (!HasSelection())
+	{
+		return;
+	}
+
+	GetTransactionsAPI()->BeginUndoTransaction(LOCTEXT("ClearSelection", "Clear Selection"));
+
+	for (TSharedPtr<FGeometrySelectionTarget> Target : ActiveTargetReferences)
+	{
+		FGeometrySelectionDelta ClearDelta;
+		Target->SelectionEditor->ClearSelection(ClearDelta);
+		if (ClearDelta.IsEmpty() == false)
+		{
+			TUniquePtr<FGeometrySelectionDeltaChange> ClearChange = MakeUnique<FGeometrySelectionDeltaChange>();
+			ClearChange->Identifier = Target->SelectionIdentifer;
+			ClearChange->Delta = MoveTemp(ClearDelta);
+			GetTransactionsAPI()->AppendChange(this, MoveTemp(ClearChange), LOCTEXT("ClearSelection", "Selection"));
+		}
+	}
+
+	GetTransactionsAPI()->EndUndoTransaction();
+
+	bSelectionRenderCachesDirty = true;
+	OnSelectionModified.Broadcast();
+}
+
+
+void UGeometrySelectionManager::UpdateSelectionViaRaycast(
+	const FRay3d& WorldRay,
+	const FGeometrySelectionUpdateConfig& UpdateConfig,
+	FGeometrySelectionUpdateResult& ResultOut
+)
+{
+	ResultOut.bSelectionModified = false;
+
+	if (ActiveTargetReferences.Num() == 0)
+	{
+		return;
+	}
+
+	// currently only going to support one object, not sure how to support more yet...
+	FGeometrySelectionTarget* Target = ActiveTargetReferences[0].Get();
+
+	Target->Selector->UpdateSelectionViaRaycast(
+		WorldRay, *Target->SelectionEditor, UpdateConfig, ResultOut
+	);
+
+	if (ResultOut.bSelectionModified)
+	{
+		TUniquePtr<FGeometrySelectionDeltaChange> DeltaChange = MakeUnique<FGeometrySelectionDeltaChange>();
+		DeltaChange->Identifier = Target->SelectionIdentifer;
+		DeltaChange->Delta = ResultOut.SelectionDelta;
+
+		GetTransactionsAPI()->BeginUndoTransaction(LOCTEXT("UpdateSelectionViaRaycast", "Change Selection"));
+		GetTransactionsAPI()->AppendChange(this, MoveTemp(DeltaChange), LOCTEXT("UpdateSelectionViaRaycast", "Change Selection"));
+		GetTransactionsAPI()->EndUndoTransaction();
+
+		bSelectionRenderCachesDirty = true;
+		OnSelectionModified.Broadcast();
+	}
+	else if (ResultOut.bSelectionMissed && UpdateConfig.ChangeType == EGeometrySelectionChangeType::Replace)
+	{
+		ClearSelection();
+	}
+}
+
+
+
+
+
+bool UGeometrySelectionManager::GetSelectionBounds(FGeometrySelectionBounds& BoundsOut) const
+{
+	BoundsOut = FGeometrySelectionBounds();
+
+	for (TSharedPtr<FGeometrySelectionTarget> Target : ActiveTargetReferences)
+	{
+		Target->Selector->AccumulateSelectionBounds(Target->Selection, BoundsOut, true);
+	}
+
+	return (BoundsOut.WorldBounds.IsEmpty() == false);
+}
+
+
+void UGeometrySelectionManager::GetSelectionWorldFrame(FFrame3d& SelectionFrame) const
+{
+	SelectionFrame = FFrame3d();
+	if (HasSelection())
+	{
+		// only handling this case for now
+		//if (ActiveTargetReferences.Num() == 1)
+		TSharedPtr<FGeometrySelectionTarget> Target = ActiveTargetReferences[0];
+		Target->Selector->GetSelectionFrame(Target->Selection, SelectionFrame, true);
+	}
+}
+
+
+bool UGeometrySelectionManager::HasSelectionForComponent(UPrimitiveComponent* Component) const
+{
+	if (HasSelection())
+	{
+		for (TSharedPtr<FGeometrySelectionTarget> Target : ActiveTargetReferences)
+		{
+			if (Target->TargetIdentifier.TargetObject == Component)
+			{
+				return Target->Selection.IsEmpty();
+			}
+		}
+	}
+	return false;
+}
+
+bool UGeometrySelectionManager::GetSelectionForComponent(UPrimitiveComponent* Component, FGeometrySelection& SelectionOut) const
+{
+	if (HasSelection())
+	{
+		for (TSharedPtr<FGeometrySelectionTarget> Target : ActiveTargetReferences)
+		{
+			if (Target->TargetIdentifier.TargetObject == Component)
+			{
+				SelectionOut = Target->Selection;
+				return ! Target->Selection.IsEmpty();
+			}
+		}
+	}
+	return false;
+}
+
+
+
+bool UGeometrySelectionManager::BeginTransformation()
+{
+	if (!ensure(IsInActiveTransformation() == false))
+	{
+		return false;
+	}
+	if (HasSelection() == false)
+	{
+		return false;
+	}
+
+	bool bHaveTransformers = false;
+	for (TSharedPtr<FGeometrySelectionTarget> Target : ActiveTargetReferences)
+	{
+		IGeometrySelectionTransformer* Transformer = Target->Selector->InitializeTransformation(Target->Selection);
+		if (Transformer != nullptr)
+		{
+			Transformer->BeginTransform(Target->Selection);
+			ActiveTransformations.Add(Transformer);
+			bHaveTransformers = true;
+		}
+		else
+		{
+			ActiveTransformations.Add(nullptr);
+		}
+	}
+
+	if (!bHaveTransformers)
+	{
+		ActiveTransformations.Reset();
+		return false;
+	}
+
+	return true;
+}
+
+void UGeometrySelectionManager::UpdateTransformation( 
+	TFunctionRef<FVector3d(int32 VertexID, const FVector3d&, const FTransform&)> PositionTransformFunc )
+{
+	if ( ! ensure(IsInActiveTransformation()) ) return;
+
+	for (int32 k = 0; k < ActiveTargetReferences.Num(); ++k)
+	{
+		if (ActiveTransformations[k] != nullptr)
+		{
+			ActiveTransformations[k]->UpdateTransform(PositionTransformFunc);
+		}
+	}
+
+	bSelectionRenderCachesDirty = true;
+}
+
+void UGeometrySelectionManager::EndTransformation()
+{
+	if ( ! ensure(IsInActiveTransformation()) ) return;
+
+	GetTransactionsAPI()->BeginUndoTransaction(LOCTEXT("EndTransformation", "Transform Selection"));
+
+	for (int32 k = 0; k < ActiveTargetReferences.Num(); ++k)
+	{
+		if (ActiveTransformations[k] != nullptr)
+		{
+			ActiveTransformations[k]->EndTransform(GetTransactionsAPI());
+
+			ActiveTargetReferences[k]->Selector->ShutdownTransformation(ActiveTransformations[k]);
+		}
+	}
+	ActiveTransformations.Reset();
+
+	GetTransactionsAPI()->EndUndoTransaction();
+
+	bSelectionRenderCachesDirty = true;
+}
+
+
+
+
+void UGeometrySelectionManager::ApplyChange(IGeometrySelectionChange* Change)
+{
+	// We should not get here because selection changes should have been expired.
+	if ( ! ensure(HasBeenShutDown() == false) )
+	{
+		return;
+	}
+
+	FGeometryIdentifier Identifer = Change->GetIdentifier();
+
+	for (int32 k = 0; k < ActiveTargetReferences.Num(); ++k)
+	{
+		if (ActiveTargetReferences[k]->SelectionIdentifer == Identifer) 
+		{
+			FGeometrySelectionDelta ApplyDelta;
+			Change->ApplyChange( ActiveTargetReferences[k]->SelectionEditor.Get(), ApplyDelta);
+
+			if (ApplyDelta.IsEmpty() == false)
+			{
+				bSelectionRenderCachesDirty = true;
+				OnSelectionModified.Broadcast();
+			}
+
+			break;
+		}
+	}
+}
+
+
+void UGeometrySelectionManager::RevertChange(IGeometrySelectionChange* Change)
+{
+	// We should not get here because selection changes should have been expired.
+	if ( ! ensure(HasBeenShutDown() == false) )
+	{
+		return;
+	}
+
+	FGeometryIdentifier Identifer = Change->GetIdentifier();
+
+	for (int32 k = 0; k < ActiveTargetReferences.Num(); ++k)
+	{
+		if (ActiveTargetReferences[k]->SelectionIdentifer == Identifer) 
+		{
+			FGeometrySelectionDelta RevertDelta;
+			Change->RevertChange( ActiveTargetReferences[k]->SelectionEditor.Get(), RevertDelta);
+
+			if (RevertDelta.IsEmpty() == false)
+			{
+				bSelectionRenderCachesDirty = true;
+				OnSelectionModified.Broadcast();
+			}
+
+			break;
+		}
+	}
+}
+
+
+void UGeometrySelectionManager::OnTargetGeometryModified(IGeometrySelector* Selector)
+{
+	bSelectionRenderCachesDirty = true;
+}
+
+
+void UGeometrySelectionManager::UpdateSelectionRenderCacheOnTargetChange()
+{
+	CachedSelectionRenderElements.Reset();
+	CachedSelectionRenderElements.SetNum(ActiveTargetReferences.Num());
+	bSelectionRenderCachesDirty = true;
+}
+
+void UGeometrySelectionManager::RebuildSelectionRenderCaches()
+{
+	if (bSelectionRenderCachesDirty == false)
+	{
+		return;
+	}
+
+	check(ActiveTargetReferences.Num() == CachedSelectionRenderElements.Num());
+	for (int32 k = 0; k < ActiveTargetReferences.Num(); ++k)
+	{
+		TSharedPtr<FGeometrySelectionTarget> Target = ActiveTargetReferences[k];
+
+		FGeometrySelectionElements& Elements = CachedSelectionRenderElements[k];
+		Elements.Reset();
+
+		Target->Selector->AccumulateSelectionElements(Target->Selection, Elements, true);
+	}
+
+	bSelectionRenderCachesDirty = false;
+}
+
+
+void UGeometrySelectionManager::DebugPrintSelection()
+{
+	if (ActiveTargetReferences.Num() == 0)
+	{
+		UE_LOG(LogGeometry, Warning, TEXT("[SelectionManager] No Active Selection"));
+		return;
+	}
+	
+	int32 NumSelected = 0;
+	for (TSharedPtr<FGeometrySelectionTarget> Target : ActiveTargetReferences)
+	{
+		NumSelected += Target->Selection.Num();
+	}
+	UE_LOG(LogGeometry, Warning, TEXT("[SelectionManager] %d selected items in %d active targets"), NumSelected, ActiveTargetReferences.Num());
+}
+
+
+void UGeometrySelectionManager::DebugRender(IToolsContextRenderAPI* RenderAPI)
+{
+	// disable selection during xform to avoid overhead
+	if (IsInActiveTransformation())
+	{
+		return;
+	}
+
+	//const UMaterialInterface* TriangleMaterial = ToolSetupUtil::GetSelectionMaterial(FLinearColor(1.0f, 0, 0, 0.5f), nullptr, 0.5f);
+
+	RebuildSelectionRenderCaches();
+
+	for ( const FGeometrySelectionElements& Elements : CachedSelectionRenderElements )
+	{
+		// batch render all the triangles, vastly more efficient than drawing one by one!
+
+		FPrimitiveDrawInterface* CurrentPDI = RenderAPI->GetPrimitiveDrawInterface();
+		FDynamicMeshBuilder MeshBuilder(CurrentPDI->View->GetFeatureLevel());
+		int32 DepthPriority = SDPG_World; // SDPG_Foreground;  // SDPG_World
+		FVector2f UVs[3] = { FVector2f(0,0), FVector2f(0,1), FVector2f(1,1)	};
+		FVector3f Normal = FVector3f(0, 0, 1);
+		FVector3f Tangent = FVector3f(1, 0, 0);	
+		for (const FTriangle3d& Triangle : Elements.Triangles)
+		{
+			int32 V0 = MeshBuilder.AddVertex(FDynamicMeshVertex((FVector3f)Triangle.V[0], Tangent, Normal, UVs[0], FColor::White));
+			int32 V1 = MeshBuilder.AddVertex(FDynamicMeshVertex((FVector3f)Triangle.V[1], Tangent, Normal, UVs[1], FColor::White));
+			int32 V2 = MeshBuilder.AddVertex(FDynamicMeshVertex((FVector3f)Triangle.V[2], Tangent, Normal, UVs[2], FColor::White));
+			MeshBuilder.AddTriangle(V0, V1, V2);
+		}
+		//FMaterialRenderProxy* MaterialRenderProxy = TriangleMaterial->GetRenderProxy();		// currently does not work, material does not render
+		FMaterialRenderProxy* MaterialRenderProxy = GEngine->ConstraintLimitMaterialX->GetRenderProxy();
+		MeshBuilder.Draw(CurrentPDI, FMatrix::Identity, MaterialRenderProxy, DepthPriority, false, false);
+
+
+		FToolDataVisualizer Visualizer;
+		Visualizer.bDepthTested = false;
+		Visualizer.BeginFrame(RenderAPI);
+
+		Visualizer.SetLineParameters(FLinearColor(0,1,0,1), 2.0f);
+		for (const FSegment3d& Segment : Elements.Segments)
+		{
+			Visualizer.DrawLine(Segment.StartPoint(), Segment.EndPoint());
+		}
+
+		Visualizer.SetPointParameters(FLinearColor(0,0,1,1), 10.0f);
+		for (const FVector3d& Point : Elements.Points)
+		{
+			Visualizer.DrawPoint(Point);
+		}
+
+		Visualizer.EndFrame();
+	}
+
+}
+
+
+
+UPersistentMeshSelection* UGeometrySelectionManager::GetActiveSingleSelectionConverted_Legacy(UPrimitiveComponent* ForComponent)
+{
+	OldSelection = nullptr;
+
+	if (HasSelection() == false)
+	{
+		return nullptr;
+	}
+
+	for (TSharedPtr<FGeometrySelectionTarget> Target : ActiveTargetReferences)
+	{
+		if (Target->TargetIdentifier.TargetObject == ForComponent)
+		{
+			const FGeometrySelection& GeoSelection = Target->Selection;
+
+			OldSelection = NewObject<UPersistentMeshSelection>(this);
+			FGenericMeshSelection Selection;
+			Selection.SourceComponent = ForComponent;
+
+			if (GeoSelection.TopologyType == EGeometryTopologyType::Triangle)
+			{
+				Selection.TopologyType = FGenericMeshSelection::ETopologyType::FTriangleGroupTopology;
+				if (GeoSelection.ElementType == EGeometryElementType::Face)
+				{
+					for (uint64 tid : GeoSelection.Selection)
+					{
+						Selection.FaceIDs.Add((int32)tid);
+					}
+				}
+			}
+			else
+			{
+				Selection.TopologyType = FGenericMeshSelection::ETopologyType::FGroupTopology;
+				if (GeoSelection.ElementType == EGeometryElementType::Face)
+				{
+					for (uint64 EncodedID : GeoSelection.Selection)
+					{
+						FGeoSelectionID SelectionID(EncodedID);
+						Selection.FaceIDs.Add((int32)SelectionID.TopologyID);
+					}
+				}
+			}
+
+			
+			OldSelection->SetSelection(Selection);
+		}
+	}
+
+	return OldSelection;
+}
+
+
+#undef LOCTEXT_NAMESPACE
\ No newline at end of file
diff --git a/Engine/Plugins/Runtime/MeshModelingToolset/Source/ModelingComponents/Private/Selection/GeometrySelectionUtil.cpp b/Engine/Plugins/Runtime/MeshModelingToolset/Source/ModelingComponents/Private/Selection/GeometrySelectionUtil.cpp
new file mode 100644
--- /dev/null
+++ b/Engine/Plugins/Runtime/MeshModelingToolset/Source/ModelingComponents/Private/Selection/GeometrySelectionUtil.cpp
@@ -0,0 +1,831 @@
+// Copyright Epic Games, Inc. All Rights Reserved.
+
+
+#include "Selection/GeometrySelectionUtil.h"
+#include "DynamicMesh/DynamicMesh3.h"
+#include "DynamicMesh/ColliderMesh.h"
+#include "DynamicMesh/MeshNormals.h"
+#include "TriangleTypes.h"
+#include "SegmentTypes.h"
+#include "GroupTopology.h"
+#include "Selections/MeshConnectedComponents.h"
+
+
+void UE::Geometry::UpdateTriangleSelectionViaRaycast(
+	const FColliderMesh* ColliderMesh,
+	FGeometrySelectionEditor* Editor,
+	const FRay3d& LocalRay,
+	const FGeometrySelectionUpdateConfig& UpdateConfig,
+	FGeometrySelectionUpdateResult& ResultOut )
+{
+	check(Editor->GetTopologyType() == EGeometryTopologyType::Triangle);
+
+	ResultOut.bSelectionMissed = true;
+
+	double RayHitT; int32 HitTriangleID; FVector3d HitBaryCoords;
+	if (ColliderMesh->FindNearestHitTriangle(LocalRay, RayHitT, HitTriangleID, HitBaryCoords))
+	{
+		HitTriangleID = ColliderMesh->GetSourceTriangleID(HitTriangleID);
+		if (HitTriangleID == IndexConstants::InvalidID)
+		{
+			return;
+		}
+
+		if (Editor->GetElementType() == EGeometryElementType::Face)
+		{
+			ResultOut.bSelectionModified = UpdateSelectionWithNewElements(Editor, UpdateConfig.ChangeType, TArray<uint64>{(uint64)HitTriangleID}, &ResultOut.SelectionDelta);
+			ResultOut.bSelectionMissed = false;
+		}
+		else if (Editor->GetElementType() == EGeometryElementType::Vertex)
+		{
+			FVector3d HitPos = LocalRay.PointAt(RayHitT);
+			FIndex3i TriVerts = ColliderMesh->GetTriangle(HitTriangleID);
+			int32 NearestIdx = 0;
+			double NearestDistSqr = DistanceSquared(ColliderMesh->GetVertex(TriVerts[0]), HitPos);
+			for (int32 k = 1; k < 3; ++k)
+			{
+				double DistSqr = DistanceSquared(ColliderMesh->GetVertex(TriVerts[k]), HitPos);
+				if (DistSqr < NearestDistSqr)
+				{
+					NearestDistSqr = DistSqr;
+					NearestIdx = k;
+				}
+			}
+
+			ResultOut.bSelectionModified = UpdateSelectionWithNewElements(Editor, UpdateConfig.ChangeType, TArray<uint64>{(uint64)TriVerts[NearestIdx]}, &ResultOut.SelectionDelta);
+			ResultOut.bSelectionMissed = false;
+		}
+		else if (Editor->GetElementType() == EGeometryElementType::Edge)
+		{
+			FVector3d HitPos = LocalRay.PointAt(RayHitT);
+			FIndex3i TriVerts = ColliderMesh->GetTriangle(HitTriangleID);
+			FVector3d Positions[3];
+			ColliderMesh->GetTriVertices(HitTriangleID, Positions[0], Positions[1], Positions[2]);
+			int32 NearestIdx = 0;
+			double NearestDistSqr = FSegment3d(Positions[0], Positions[1]).DistanceSquared(HitPos);
+			for (int32 k = 1; k < 3; ++k)
+			{
+				double DistSqr = FSegment3d(Positions[k], Positions[(k+1)%3]).DistanceSquared(HitPos);
+				if (DistSqr < NearestDistSqr)
+				{
+					NearestDistSqr = DistSqr;
+					NearestIdx = k;
+				}
+			}
+			FMeshTriEdgeID TriEdgeID(HitTriangleID, NearestIdx);
+
+			ResultOut.bSelectionModified = UpdateSelectionWithNewElements(Editor, UpdateConfig.ChangeType, TArray<uint64>{(uint64)TriEdgeID.Encoded()}, &ResultOut.SelectionDelta);
+			ResultOut.bSelectionMissed = false;
+		}
+	}
+}
+
+
+
+
+
+void UE::Geometry::UpdateGroupSelectionViaRaycast(
+	const FColliderMesh* ColliderMesh,
+	const FGroupTopology* GroupTopology,
+	FGeometrySelectionEditor* Editor,
+	const FRay3d& LocalRay,
+	const FGeometrySelectionUpdateConfig& UpdateConfig,
+	FGeometrySelectionUpdateResult& ResultOut)
+{
+	check(Editor->GetTopologyType() == EGeometryTopologyType::Polygroup);
+
+	ResultOut.bSelectionMissed = true;
+
+	double RayHitT; int32 HitTriangleID; FVector3d HitBaryCoords;
+	if (ColliderMesh->FindNearestHitTriangle(LocalRay, RayHitT, HitTriangleID, HitBaryCoords))
+	{
+		HitTriangleID = ColliderMesh->GetSourceTriangleID(HitTriangleID);
+		if (HitTriangleID == IndexConstants::InvalidID)
+		{
+			return;
+		}
+		int32 GroupID = GroupTopology->GetGroupID(HitTriangleID);
+
+		if (Editor->GetElementType() == EGeometryElementType::Face)
+		{
+			FGeoSelectionID GroupTriID((uint32)HitTriangleID, (uint32)GroupID);
+			ResultOut.bSelectionModified = UpdateSelectionWithNewElements(Editor, UpdateConfig.ChangeType, TArray<uint64>{GroupTriID.Encoded()}, &ResultOut.SelectionDelta);
+			ResultOut.bSelectionMissed = false;
+		}
+		else if (Editor->GetElementType() == EGeometryElementType::Vertex)
+		{
+			FVector3d HitPos = LocalRay.PointAt(RayHitT);
+			FIndex3i TriVerts = ColliderMesh->GetTriangle(HitTriangleID);
+			int32 NearestIdx = -1;
+			double NearestDistSqr = DistanceSquared(ColliderMesh->GetVertex(TriVerts[0]), HitPos);
+			for (int32 k = 0; k < 3; ++k)
+			{
+				if (GroupTopology->GetCornerIDFromVertexID(TriVerts[k]) != IndexConstants::InvalidID)
+				{
+					double DistSqr = DistanceSquared(ColliderMesh->GetVertex(TriVerts[k]), HitPos);
+					if (DistSqr < NearestDistSqr)
+					{
+						NearestDistSqr = DistSqr;
+						NearestIdx = k;
+					}
+				}
+			}
+			if (NearestIdx >= 0)
+			{
+				// do we need a group here?
+				int32 VertexID = TriVerts[NearestIdx];
+				ResultOut.bSelectionModified = UpdateSelectionWithNewElements(Editor, UpdateConfig.ChangeType, TArray<uint64>{(uint64)VertexID}, &ResultOut.SelectionDelta);
+				ResultOut.bSelectionMissed = false;
+			}
+		}
+		else if (Editor->GetElementType() == EGeometryElementType::Edge)
+		{
+			FVector3d HitPos = LocalRay.PointAt(RayHitT);
+			FIndex3i TriVerts = ColliderMesh->GetTriangle(HitTriangleID);
+			FVector3d Positions[3];
+			ColliderMesh->GetTriVertices(HitTriangleID, Positions[0], Positions[1], Positions[2]);
+			int32 NearestIdx = -1;
+			double NearestDistSqr = TNumericLimits<double>::Max();
+			for (int32 k = 0; k < 3; ++k)
+			{
+				if ( GroupTopology->IsGroupEdge(FMeshTriEdgeID(HitTriangleID, k), true) )
+				{
+					double DistSqr = FSegment3d(Positions[k], Positions[(k + 1) % 3]).DistanceSquared(HitPos);
+					if (DistSqr < NearestDistSqr)
+					{
+						NearestDistSqr = DistSqr;
+						NearestIdx = k;
+					}
+				}
+			}
+			if ( NearestIdx >= 0 )
+			{
+				// do we need a group here?
+				FMeshTriEdgeID TriEdgeID(HitTriangleID, NearestIdx);
+				ResultOut.bSelectionModified = UpdateSelectionWithNewElements(Editor, UpdateConfig.ChangeType, TArray<uint64>{(uint64)TriEdgeID.Encoded()}, & ResultOut.SelectionDelta);
+				ResultOut.bSelectionMissed = false;
+			}
+		}
+	}
+}
+
+
+
+
+
+
+
+bool UE::Geometry::UpdateSelectionWithNewElements(
+	FGeometrySelectionEditor* Editor,
+	EGeometrySelectionChangeType ChangeType,
+	const TArray<uint64>& NewIDs,
+	FGeometrySelectionDelta* Delta)
+{
+	FGeometrySelectionDelta LocalDelta;
+	FGeometrySelectionDelta& UseDelta = (Delta != nullptr) ? (*Delta) : LocalDelta;
+
+	if (ChangeType == EGeometrySelectionChangeType::Replace)
+	{
+		// [TODO] this could be optimized...
+		Editor->ClearSelection(UseDelta);
+		return Editor->Select(NewIDs, UseDelta);
+	}
+	else if (ChangeType == EGeometrySelectionChangeType::Add)
+	{
+		return Editor->Select(NewIDs, UseDelta);
+	}
+	else if (ChangeType == EGeometrySelectionChangeType::Remove)
+	{
+		return Editor->Deselect(NewIDs, UseDelta);
+	}
+	else
+	{
+		ensure(false);
+		return false;
+	}
+	
+}
+
+
+
+bool UE::Geometry::EnumerateTriangleSelectionVertices(
+	const FGeometrySelection& MeshSelection,
+	const UE::Geometry::FDynamicMesh3& Mesh,
+	const FTransform& ApplyTransform,
+	TFunctionRef<void(uint64, const FVector3d&)> VertexFunc)
+{
+	if ( ensure( MeshSelection.TopologyType == EGeometryTopologyType::Triangle ) == false )
+	{
+		return false;
+	}
+
+	if (MeshSelection.ElementType == EGeometryElementType::Face)
+	{
+		for (uint64 TriangleID : MeshSelection.Selection)
+		{
+			if (Mesh.IsTriangle(TriangleID))
+			{
+				FIndex3i Triangle = Mesh.GetTriangle((int32)TriangleID);
+				VertexFunc((uint64)Triangle.A, ApplyTransform.TransformPosition(Mesh.GetVertex(Triangle.A)));
+				VertexFunc((uint64)Triangle.B, ApplyTransform.TransformPosition(Mesh.GetVertex(Triangle.B)));
+				VertexFunc((uint64)Triangle.C, ApplyTransform.TransformPosition(Mesh.GetVertex(Triangle.C)));
+			}
+		}
+	}
+	else if (MeshSelection.ElementType == EGeometryElementType::Edge)
+	{
+		for (uint64 EncodedID : MeshSelection.Selection)
+		{
+			FMeshTriEdgeID TriEdgeID((uint32)EncodedID);
+			int32 EdgeID = Mesh.IsTriangle(TriEdgeID.TriangleID) ? Mesh.GetTriEdge(TriEdgeID.TriangleID, TriEdgeID.TriEdgeIndex) : IndexConstants::InvalidID;
+			if (Mesh.IsEdge(EdgeID))
+			{
+				FIndex2i EdgeV = Mesh.GetEdgeV(EdgeID);
+				VertexFunc((uint64)EdgeV.A, ApplyTransform.TransformPosition(Mesh.GetVertex(EdgeV.A)));
+				VertexFunc((uint64)EdgeV.B, ApplyTransform.TransformPosition(Mesh.GetVertex(EdgeV.B)));
+			}
+		}
+	}
+	else if (MeshSelection.ElementType == EGeometryElementType::Vertex)
+	{
+		for (uint64 VertexID : MeshSelection.Selection)
+		{
+			if (Mesh.IsVertex((int32)VertexID))
+			{
+				VertexFunc((uint64)VertexID, ApplyTransform.TransformPosition(Mesh.GetVertex(VertexID)));
+			}
+		}
+	}
+	else
+	{
+		return false;
+	}
+
+	return true;
+}
+
+
+
+bool UE::Geometry::EnumeratePolygroupSelectionVertices(
+	const FGeometrySelection& MeshSelection,
+	const UE::Geometry::FDynamicMesh3& Mesh,
+	const FGroupTopology* GroupTopology,
+	const FTransform& ApplyTransform,
+	TFunctionRef<void(uint64, const FVector3d&)> VertexFunc)
+{
+	if ( ensure( MeshSelection.TopologyType == EGeometryTopologyType::Polygroup ) == false )
+	{
+		return false;
+	}
+
+	if (MeshSelection.ElementType == EGeometryElementType::Face)
+	{
+		for (uint64 EncodedID : MeshSelection.Selection)
+		{
+			FGeoSelectionID GroupTriID(EncodedID);
+			int32 SeedTriangleID = (int32)GroupTriID.GeometryID, GroupID = (int32)GroupTriID.TopologyID;
+			if (Mesh.IsTriangle(SeedTriangleID))
+			{
+				for (int32 TriangleID : GroupTopology->GetGroupFaces(GroupID))
+				{
+					FIndex3i Triangle = Mesh.GetTriangle((int32)TriangleID);
+					VertexFunc((uint64)Triangle.A, ApplyTransform.TransformPosition(Mesh.GetVertex(Triangle.A)));
+					VertexFunc((uint64)Triangle.B, ApplyTransform.TransformPosition(Mesh.GetVertex(Triangle.B)));
+					VertexFunc((uint64)Triangle.C, ApplyTransform.TransformPosition(Mesh.GetVertex(Triangle.C)));
+				}
+			}
+		}
+	}
+	else if (MeshSelection.ElementType == EGeometryElementType::Edge)
+	{
+		for (uint64 EncodedID : MeshSelection.Selection)
+		{
+			FMeshTriEdgeID TriEdgeID( FGeoSelectionID(EncodedID).GeometryID );
+			int32 SeedEdgeID = Mesh.IsTriangle(TriEdgeID.TriangleID) ? Mesh.GetTriEdge(TriEdgeID.TriangleID, TriEdgeID.TriEdgeIndex) : IndexConstants::InvalidID;
+			if (Mesh.IsEdge(SeedEdgeID))
+			{
+				int32 GroupEdgeID = GroupTopology->FindGroupEdgeID(SeedEdgeID);
+				for (int32 VertexID : GroupTopology->GetGroupEdgeVertices(GroupEdgeID))
+				{
+					FVector3d V = Mesh.GetVertex(VertexID);
+					VertexFunc((uint64)VertexID, ApplyTransform.TransformPosition(V));
+				}
+			}
+		}
+	}
+	else if (MeshSelection.ElementType == EGeometryElementType::Vertex)
+	{
+		for (uint64 EncodedID : MeshSelection.Selection)
+		{
+			int32 VertexID = (int32)FGeoSelectionID(EncodedID).GeometryID;
+			if (Mesh.IsVertex(VertexID))
+			{
+				VertexFunc((uint64)VertexID, ApplyTransform.TransformPosition(Mesh.GetVertex(VertexID)));
+			}
+		}
+	}
+	else
+	{
+		return false;
+	}
+
+	return true;
+}
+
+
+
+bool UE::Geometry::EnumerateSelectionTriangles(
+	const FGeometrySelection& MeshSelection,
+	const UE::Geometry::FDynamicMesh3& Mesh,
+	TFunctionRef<void(int32)> TriangleFunc,
+	const UE::Geometry::FPolygroupSet* UseGroupSet
+)
+{
+	if (MeshSelection.TopologyType == EGeometryTopologyType::Triangle)
+	{
+		return EnumerateTriangleSelectionTriangles(MeshSelection, Mesh, TriangleFunc);
+	}
+	else if (MeshSelection.TopologyType == EGeometryTopologyType::Polygroup)
+	{
+		return EnumeratePolygroupSelectionTriangles(MeshSelection, Mesh, 
+			(UseGroupSet != nullptr) ? *UseGroupSet : FPolygroupSet(&Mesh), 
+			TriangleFunc);
+	}
+	return false;
+}
+
+
+bool UE::Geometry::EnumerateTriangleSelectionTriangles(
+	const FGeometrySelection& MeshSelection,
+	const UE::Geometry::FDynamicMesh3& Mesh,
+	TFunctionRef<void(int32)> TriangleFunc)
+{
+	if ( ensure( MeshSelection.TopologyType == EGeometryTopologyType::Triangle ) == false )
+	{
+		return false;
+	}
+
+	if (MeshSelection.ElementType == EGeometryElementType::Face)
+	{
+		for (uint64 TriangleID : MeshSelection.Selection)
+		{
+			if (Mesh.IsTriangle(TriangleID))
+			{
+				TriangleFunc(TriangleID);
+			}
+		}
+	}
+	else if (MeshSelection.ElementType == EGeometryElementType::Edge)
+	{
+		for (uint64 EncodedID : MeshSelection.Selection)
+		{
+			FMeshTriEdgeID TriEdgeID( FGeoSelectionID(EncodedID).GeometryID );
+			int32 EdgeID = Mesh.IsTriangle(TriEdgeID.TriangleID) ? Mesh.GetTriEdge(TriEdgeID.TriangleID, TriEdgeID.TriEdgeIndex) : IndexConstants::InvalidID;
+			Mesh.EnumerateEdgeTriangles(EdgeID, [&](int32 TriangleID)
+			{
+				TriangleFunc(TriangleID);
+			});
+		}
+	}
+	else if (MeshSelection.ElementType == EGeometryElementType::Vertex)
+	{
+		for (uint64 VertexID : MeshSelection.Selection)
+		{
+			Mesh.EnumerateVertexTriangles(VertexID, [&](int32 TriangleID)
+			{
+				TriangleFunc(TriangleID);
+			});
+		}
+	}
+	else
+	{
+		return false;
+	}
+
+	return true;
+}
+
+
+bool UE::Geometry::EnumeratePolygroupSelectionTriangles(
+	const FGeometrySelection& MeshSelection,
+	const UE::Geometry::FDynamicMesh3& Mesh,
+	const UE::Geometry::FPolygroupSet& GroupSet,
+	TFunctionRef<void(int32)> TriangleFunc
+)
+{
+	if ( ensure( MeshSelection.TopologyType == EGeometryTopologyType::Polygroup ) == false )
+	{
+		return false;
+	}
+
+	TArray<int32> SeedGroups;
+	TArray<int32> SeedTriangles;
+	TSet<int32> UniqueSeedGroups;
+
+	if (MeshSelection.ElementType == EGeometryElementType::Face)
+	{
+		for (uint64 EncodedID : MeshSelection.Selection)
+		{
+			FGeoSelectionID SelectionID(EncodedID);
+			int32 SeedTriangleID = (int32)SelectionID.GeometryID;
+			if (Mesh.IsTriangle(SeedTriangleID))
+			{
+				int32 GroupID = GroupSet.GetGroup(SeedTriangleID);
+				check(GroupID == (int32)SelectionID.TopologyID);		// sanity-check that we are using the right group
+				if ( GroupID >= 0 && UniqueSeedGroups.Contains(GroupID) == false)
+				{
+					UniqueSeedGroups.Add(GroupID);
+					SeedGroups.Add(GroupID);
+					SeedTriangles.Add(SeedTriangleID);
+				}
+			}
+		}
+	}
+	else if (MeshSelection.ElementType == EGeometryElementType::Edge)
+	{
+		for (uint64 EncodedID : MeshSelection.Selection)
+		{
+			FMeshTriEdgeID TriEdgeID( FGeoSelectionID(EncodedID).GeometryID );
+			int32 SeedEdgeID = Mesh.IsTriangle(TriEdgeID.TriangleID) ? Mesh.GetTriEdge(TriEdgeID.TriangleID, TriEdgeID.TriEdgeIndex) : IndexConstants::InvalidID;
+			if (Mesh.IsEdge(SeedEdgeID))
+			{
+				Mesh.EnumerateEdgeTriangles(SeedEdgeID, [&](int32 TriangleID)
+				{
+					int32 GroupID = GroupSet.GetGroup(TriangleID);
+					if (GroupID >= 0 && UniqueSeedGroups.Contains(GroupID) == false)
+					{
+						UniqueSeedGroups.Add(GroupID);
+						SeedGroups.Add(GroupID);
+						SeedTriangles.Add(TriangleID);
+					}
+				});
+			}
+		}
+	}
+	else if (MeshSelection.ElementType == EGeometryElementType::Vertex)
+	{
+		for (uint64 EncodedID : MeshSelection.Selection)
+		{
+			int32 VertexID = (int32)FGeoSelectionID(EncodedID).GeometryID;
+			if (Mesh.IsVertex(VertexID))
+			{
+				Mesh.EnumerateVertexTriangles(VertexID, [&](int32 TriangleID)
+				{
+					int32 GroupID = GroupSet.GetGroup(TriangleID);
+					if (GroupID >= 0 && UniqueSeedGroups.Contains(GroupID) == false)
+					{
+						UniqueSeedGroups.Add(GroupID);
+						SeedGroups.Add(GroupID);
+						SeedTriangles.Add(TriangleID);
+					}
+				});
+			}
+		}
+	}
+	else
+	{
+		return false;
+	}
+
+
+	TSet<int> TempROI;		// if we could provide this as input we would not need a temporary roi...
+	TArray<int32> QueueBuffer;
+	int32 NumGroups = SeedGroups.Num();
+	for (int32 k = 0; k < NumGroups; ++k)
+	{
+		check(GroupSet.GetGroup(SeedTriangles[k]) == SeedGroups[k]);
+		int32 GroupID = SeedGroups[k];
+		FMeshConnectedComponents::GrowToConnectedTriangles(&Mesh, 
+			TArray<int>{SeedTriangles[k]}, TempROI, &QueueBuffer, 
+			[&](int32 T1, int32 T2) { return GroupSet.GetGroup(T2) == GroupID; });
+		for (int32 tid : TempROI)
+		{
+			TriangleFunc(tid);
+		}
+	}
+
+	return true;
+}
+
+
+
+
+bool UE::Geometry::EnumerateTriangleSelectionElements(
+	const FGeometrySelection& MeshSelection,
+	const UE::Geometry::FDynamicMesh3& Mesh,
+	TFunctionRef<void(int32, const FVector3d&)> VertexFunc,
+	TFunctionRef<void(int32, const FSegment3d&)> EdgeFunc,
+	TFunctionRef<void(int32, const FTriangle3d&)> TriangleFunc,
+	const FTransform* ApplyTransform
+)
+{
+	if ( ensure( MeshSelection.TopologyType == EGeometryTopologyType::Triangle ) == false )
+	{
+		return false;
+	}
+
+	if (MeshSelection.ElementType == EGeometryElementType::Face)
+	{
+		for (uint64 EncodedID : MeshSelection.Selection)
+		{
+			int32 TriangleID = (int32)FGeoSelectionID(EncodedID).GeometryID;
+			if (Mesh.IsTriangle(TriangleID))
+			{
+				FVector3d A, B, C;
+				Mesh.GetTriVertices((int32)TriangleID, A, B, C);
+				if (ApplyTransform != nullptr)
+				{
+					A = ApplyTransform->TransformPosition(A);
+					B = ApplyTransform->TransformPosition(B);
+					C = ApplyTransform->TransformPosition(C);
+				}
+				TriangleFunc(TriangleID, FTriangle3d(A, B, C));
+			}
+		}
+	}
+	else if (MeshSelection.ElementType == EGeometryElementType::Edge)
+	{
+		for (uint64 EncodedID : MeshSelection.Selection)
+		{
+			FMeshTriEdgeID TriEdgeID( FGeoSelectionID(EncodedID).GeometryID );
+			int32 EdgeID = Mesh.IsTriangle(TriEdgeID.TriangleID) ? Mesh.GetTriEdge(TriEdgeID.TriangleID, TriEdgeID.TriEdgeIndex) : IndexConstants::InvalidID;
+			if (Mesh.IsEdge(EdgeID))
+			{
+				FVector3d A, B;
+				Mesh.GetEdgeV(EdgeID, A, B);
+				if (ApplyTransform != nullptr)
+				{
+					A = ApplyTransform->TransformPosition(A);
+					B = ApplyTransform->TransformPosition(B);
+				}
+				EdgeFunc(EdgeID, FSegment3d(A, B));
+			}
+		}
+	}
+	else if (MeshSelection.ElementType == EGeometryElementType::Vertex)
+	{
+		for (uint64 EncodedID : MeshSelection.Selection)
+		{
+			int32 VertexID = (int32)FGeoSelectionID(EncodedID).GeometryID;
+			if (Mesh.IsVertex(VertexID))
+			{
+				FVector3d A = Mesh.GetVertex(VertexID);
+				VertexFunc(VertexID, (ApplyTransform != nullptr) ? ApplyTransform->TransformPosition(A) : A);
+			}
+		}
+	}
+	else
+	{
+		return false;
+	}
+
+	return true;
+}
+
+
+
+
+
+bool UE::Geometry::EnumeratePolygroupSelectionElements(
+	const FGeometrySelection& MeshSelection,
+	const UE::Geometry::FDynamicMesh3& Mesh,
+	const FGroupTopology* GroupTopology,
+	TFunctionRef<void(int32, const FVector3d&)> VertexFunc,
+	TFunctionRef<void(int32, const FSegment3d&)> EdgeFunc,
+	TFunctionRef<void(int32, const FTriangle3d&)> TriangleFunc,
+	const FTransform* ApplyTransform
+)
+{
+	if ( ensure( MeshSelection.TopologyType == EGeometryTopologyType::Polygroup ) == false )
+	{
+		return false;
+	}
+
+	if (MeshSelection.ElementType == EGeometryElementType::Face)
+	{
+		for (uint64 EncodedID : MeshSelection.Selection)
+		{
+			FGeoSelectionID SelectionID(EncodedID);
+			int32 SeedTriangleID = (int32)SelectionID.GeometryID, GroupID = (int32)SelectionID.TopologyID;
+			if (Mesh.IsTriangle(SeedTriangleID))
+			{
+				for (int32 TriangleID : GroupTopology->GetGroupFaces(GroupID))
+				{
+					FVector3d A, B, C;
+					Mesh.GetTriVertices((int32)TriangleID, A, B, C);
+					if (ApplyTransform != nullptr)
+					{
+						A = ApplyTransform->TransformPosition(A);
+						B = ApplyTransform->TransformPosition(B);
+						C = ApplyTransform->TransformPosition(C);
+					}
+					TriangleFunc(TriangleID, FTriangle3d(A, B, C));
+				}
+			}
+		}
+	}
+	else if (MeshSelection.ElementType == EGeometryElementType::Edge)
+	{
+		for (uint64 EncodedID : MeshSelection.Selection)
+		{
+			FMeshTriEdgeID TriEdgeID( FGeoSelectionID(EncodedID).GeometryID );
+			int32 SeedEdgeID = Mesh.IsTriangle(TriEdgeID.TriangleID) ? Mesh.GetTriEdge(TriEdgeID.TriangleID, TriEdgeID.TriEdgeIndex) : IndexConstants::InvalidID;
+			if (Mesh.IsEdge(SeedEdgeID))
+			{
+				int32 GroupEdgeID = GroupTopology->FindGroupEdgeID(SeedEdgeID);
+				for (int32 EdgeID : GroupTopology->GetGroupEdgeEdges(GroupEdgeID))
+				{
+					FVector3d A, B;
+					Mesh.GetEdgeV(EdgeID, A, B);
+					if (ApplyTransform != nullptr)
+					{
+						A = ApplyTransform->TransformPosition(A);
+						B = ApplyTransform->TransformPosition(B);
+					}
+					EdgeFunc(EdgeID, FSegment3d(A, B));
+				}
+			}
+		}
+	}
+	else if (MeshSelection.ElementType == EGeometryElementType::Vertex)
+	{
+		for (uint64 EncodedID : MeshSelection.Selection)
+		{
+			int32 VertexID = (int32)FGeoSelectionID(EncodedID).GeometryID;
+			if (Mesh.IsVertex(VertexID))
+			{
+				FVector3d A = Mesh.GetVertex(VertexID);
+				VertexFunc(VertexID, (ApplyTransform != nullptr) ? ApplyTransform->TransformPosition(A) : A);
+			}
+		}
+	}
+	else
+	{
+		return false;
+	}
+
+	return true;
+}
+
+
+
+
+bool UE::Geometry::ConvertPolygroupSelectionToTopologySelection(
+	const FGeometrySelection& MeshSelection,
+	const UE::Geometry::FDynamicMesh3& Mesh,
+	const FGroupTopology* GroupTopology,
+	FGroupTopologySelection& TopologySelectionOut)
+{
+	if ( ensure( MeshSelection.TopologyType == EGeometryTopologyType::Polygroup ) == false )
+	{
+		return false;
+	}
+
+	if (MeshSelection.ElementType == EGeometryElementType::Face)
+	{
+		for (uint64 EncodedID : MeshSelection.Selection)
+		{
+			int32 GroupID = FGeoSelectionID(EncodedID).TopologyID;
+			if (GroupTopology->FindGroupByID(GroupID) != nullptr)
+			{
+				TopologySelectionOut.SelectedGroupIDs.Add(GroupID);
+			}
+		}
+	}
+	else if (MeshSelection.ElementType == EGeometryElementType::Edge)
+	{
+		for (uint64 EncodedID : MeshSelection.Selection)
+		{
+			FMeshTriEdgeID MeshEdgeID( FGeoSelectionID(EncodedID).GeometryID );
+			if ( Mesh.IsTriangle((int32)MeshEdgeID.TriangleID) )
+			{
+				int32 GroupEdgeID = GroupTopology->FindGroupEdgeID(MeshEdgeID);
+				if (GroupEdgeID >= 0)
+				{
+					TopologySelectionOut.SelectedEdgeIDs.Add(GroupEdgeID);
+				}
+			}
+		}
+	}
+	else if (MeshSelection.ElementType == EGeometryElementType::Vertex)
+	{
+		for (uint64 EncodedID : MeshSelection.Selection)
+		{
+			int32 VertexID = (int32)FGeoSelectionID(EncodedID).GeometryID;
+			if ( Mesh.IsVertex(VertexID) )
+			{
+				int32 CornerID = GroupTopology->GetCornerIDFromVertexID(VertexID);
+				if (CornerID >= 0)
+				{
+					TopologySelectionOut.SelectedCornerIDs.Add(CornerID);
+				}
+			}
+		}
+	}
+	else
+	{
+		return false;
+	}
+
+	return true;
+}
+
+
+
+
+bool UE::Geometry::GetTriangleSelectionFrame(
+	const FGeometrySelection& MeshSelection,
+	const UE::Geometry::FDynamicMesh3& Mesh,
+	FFrame3d& SelectionFrameOut)
+{
+	if ( ensure( MeshSelection.TopologyType == EGeometryTopologyType::Triangle ) == false )
+	{
+		return false;
+	}
+
+	FVector3d AccumulatedOrigin = FVector3d::Zero();
+	FVector3d AccumulatedNormal = FVector3d::Zero();
+	FVector3d AxisHint = FVector3d::Zero();
+	double AccumWeight = 0;
+	
+	if (MeshSelection.ElementType == EGeometryElementType::Face)
+	{
+		for (uint64 EncodedID : MeshSelection.Selection)
+		{
+			int32 TriangleID = (int32)FGeoSelectionID(EncodedID).GeometryID;
+			if (Mesh.IsTriangle(TriangleID))
+			{
+				FVector3d Normal, Centroid; double Area; 
+				Mesh.GetTriInfo(TriangleID, Normal, Area, Centroid);
+				if (Normal.SquaredLength() > 0.9)
+				{
+					Area = FMath::Max(Area, 0.000001);
+					AccumulatedOrigin += Area * Centroid;
+					AccumulatedNormal += Area * Normal;
+					AccumWeight += Area;
+				}
+			}
+		}
+	}
+	else if (MeshSelection.ElementType == EGeometryElementType::Edge)
+	{
+		for (uint64 EncodedID : MeshSelection.Selection)
+		{
+			FMeshTriEdgeID TriEdgeID( FGeoSelectionID(EncodedID).GeometryID );
+			int32 EdgeID = Mesh.IsTriangle(TriEdgeID.TriangleID) ? Mesh.GetTriEdge(TriEdgeID.TriangleID, TriEdgeID.TriEdgeIndex) : IndexConstants::InvalidID;
+			if ( Mesh.IsEdge(EdgeID) )
+			{
+				FVector3d A, B;
+				Mesh.GetEdgeV(EdgeID, A, B);
+				AccumulatedOrigin += (A + B) * 0.5;
+				AccumulatedNormal += Mesh.GetEdgeNormal(EdgeID);
+				AxisHint += Normalized(B - A);
+				AccumWeight += 1.0;
+			}
+		}
+	}
+	else if (MeshSelection.ElementType == EGeometryElementType::Vertex)
+	{
+		for (uint64 EncodedID : MeshSelection.Selection)
+		{
+			int32 VertexID = (int32)FGeoSelectionID(EncodedID).GeometryID;
+			if (Mesh.IsVertex(VertexID))
+			{
+				AccumulatedOrigin += Mesh.GetVertex(VertexID);
+				AccumulatedNormal += FMeshNormals::ComputeVertexNormal(Mesh, VertexID);		// this could return area!
+				AccumWeight += 1.0;
+			}
+		}
+	}
+	else
+	{
+		return false;
+	}
+
+	// todo use AxisHint!
+
+	SelectionFrameOut = FFrame3d();
+	if (AccumWeight > 0)
+	{
+		AccumulatedOrigin /= (double)AccumWeight;
+		Normalize(AccumulatedNormal);
+
+		// We set our frame Z to be accumulated normal, and the other two axes are unconstrained, so
+		// we want to set them to something that will make our frame generally more useful. If the normal
+		// is aligned with world Z, then the entire frame might as well be aligned with world.
+		if (1 - AccumulatedNormal.Dot(FVector3d::UnitZ()) < KINDA_SMALL_NUMBER)
+		{
+			SelectionFrameOut = FFrame3d(AccumulatedOrigin, FQuaterniond::Identity());
+		}
+		else
+		{
+			// Otherwise, let's place one of the other axes into the XY plane so that the frame is more
+			// useful for translation. We somewhat arbitrarily choose Y for this. 
+			FVector3d FrameY = Normalized(AccumulatedNormal.Cross(FVector3d::UnitZ())); // orthogonal to world Z and frame Z 
+			FVector3d FrameX = FrameY.Cross(AccumulatedNormal); // safe to not normalize because already orthogonal
+			SelectionFrameOut = FFrame3d(AccumulatedOrigin, FrameX, FrameY, AccumulatedNormal);
+		}
+	}
+
+	return true;
+}
\ No newline at end of file
diff --git a/Engine/Plugins/Runtime/MeshModelingToolset/Source/ModelingComponents/Public/Selection/DynamicMeshSelector.h b/Engine/Plugins/Runtime/MeshModelingToolset/Source/ModelingComponents/Public/Selection/DynamicMeshSelector.h
new file mode 100644
--- /dev/null
+++ b/Engine/Plugins/Runtime/MeshModelingToolset/Source/ModelingComponents/Public/Selection/DynamicMeshSelector.h
@@ -0,0 +1,147 @@
+// Copyright Epic Games, Inc. All Rights Reserved.
+
+#pragma once
+
+#include "CoreMinimal.h"
+#include "Misc/Change.h"
+#include "Templates/PimplPtr.h"
+#include "Selection/GeometrySelector.h"
+#include "Selection/GeometrySelection.h"
+#include "Components/DynamicMeshComponent.h"
+#include "DynamicMesh/DynamicMeshAABBTree3.h"
+#include "TransformTypes.h"
+
+
+class FDynamicMeshSelectionTransformer;
+class FMeshVertexChangeBuilder;
+PREDECLARE_GEOMETRY(class FGroupTopology);
+PREDECLARE_GEOMETRY(class FColliderMesh);
+
+
+/**
+ * FDynamicMeshSelector is an implementation of IGeometrySelector for a UDynamicMesh.
+ * Note that the Selector itself does *not* require that the target object be a UDynamicMeshComponent.
+ * Access to the World transform is provided by a TUniqueFunction set up in the Factory.
+ */
+class MODELINGCOMPONENTS_API FDynamicMeshSelector : public IGeometrySelector
+{
+public:
+	virtual ~FDynamicMeshSelector()
+	{
+		Shutdown();
+	}
+
+	/**
+	 * Initialize the FDynamicMeshSelector for a given source/target UDynamicMesh.
+	 * @param SourceGeometryIdentifier identifier for the object that the TargetMesh came from (eg DynamicMeshComponent or other UDynamicMesh source)
+	 * @param TargetMesh the target UDynamicMesh
+	 * @param GetWorldTransformFunc function that provides the Local to World Transform
+	 */
+	virtual void Initialize(
+		FGeometryIdentifier SourceGeometryIdentifier,
+		UDynamicMesh* TargetMesh,
+		TUniqueFunction<UE::Geometry::FTransformSRT3d()> GetWorldTransformFunc);
+
+	/**
+	 * @return FGeometryIdentifier for the parent of this Selector (eg a UDynamicMeshComponent in the common case)
+	 */
+	virtual FGeometryIdentifier GetSourceGeometryIdentifier() const
+	{
+		return SourceGeometryIdentifier;
+	}
+
+	//
+	// IGeometrySelector API implementation
+	//
+
+	virtual void Shutdown() override;
+	virtual bool SupportsSleep() const { return true; }
+	virtual bool Sleep();
+	virtual bool Restore();
+
+
+	virtual FGeometryIdentifier GetIdentifier() const override
+	{
+		FGeometryIdentifier Identifier;
+		Identifier.TargetType = FGeometryIdentifier::ETargetType::MeshContainer;
+		Identifier.ObjectType = FGeometryIdentifier::EObjectType::DynamicMesh;
+		Identifier.TargetObject = TargetMesh;
+		return Identifier;
+	}
+
+	virtual bool RayHitTest(
+		const FRay3d& WorldRay,
+		FInputRayHit& HitResultOut) override;
+
+
+	virtual void UpdateSelectionViaRaycast(
+		const FRay3d& WorldRay,
+		FGeometrySelectionEditor& SelectionEditor,
+		const FGeometrySelectionUpdateConfig& UpdateConfig,
+		FGeometrySelectionUpdateResult& ResultOut) override;
+
+
+	virtual FTransform GetLocalToWorldTransform() const
+	{
+		return (FTransform)GetWorldTransformFunc();
+	}
+
+	virtual void GetSelectionFrame(const FGeometrySelection& Selection, UE::Geometry::FFrame3d& SelectionFrame, bool bTransformToWorld) override;
+	virtual void AccumulateSelectionBounds(const FGeometrySelection& Selection, FGeometrySelectionBounds& BoundsInOut, bool bTransformToWorld) override;
+	virtual void AccumulateSelectionElements(const FGeometrySelection& Selection, FGeometrySelectionElements& Elements, bool bTransformToWorld) override;
+
+	virtual IGeometrySelectionTransformer* InitializeTransformation(const FGeometrySelection& Selection) override;
+	virtual void ShutdownTransformation(IGeometrySelectionTransformer* Transformer) override;
+
+
+protected:
+	FGeometryIdentifier SourceGeometryIdentifier;
+	TUniqueFunction<UE::Geometry::FTransformSRT3d()> GetWorldTransformFunc;
+
+	TWeakObjectPtr<UDynamicMesh> TargetMesh;
+	UDynamicMesh* GetDynamicMesh() const { return TargetMesh.Get(); }
+
+	FDelegateHandle TargetMesh_OnMeshChangedHandle;
+	void RegisterMeshChangedHandler();
+
+	//
+	// FColliderMesh is used to store a hit-testable AABBTree independent of the UDynamicMesh 
+	//
+	TPimplPtr<UE::Geometry::FColliderMesh> ColliderMesh;
+	void UpdateColliderMesh();
+	const UE::Geometry::FColliderMesh* GetColliderMesh();
+
+	//
+	// GroupTopology will be built on-demand if polygroup selection queries are made
+	//
+	TPimplPtr<UE::Geometry::FGroupTopology> GroupTopology;
+	void UpdateGroupTopology();
+	const UE::Geometry::FGroupTopology* GetGroupTopology();
+
+	// support for sleep/restore
+	TWeakObjectPtr<UDynamicMesh> SleepingTargetMesh = nullptr;
+
+
+	TPimplPtr<FDynamicMeshSelectionTransformer> ActiveTransformer;
+
+	// give Transformer access to internals 
+	friend class FDynamicMeshSelectionTransformer;
+};
+
+
+/**
+ * FDynamicMeshComponentSelectorFactory constructs FDynamicMeshSelector instances 
+ * for UDynamicMeshComponents
+ */
+class MODELINGCOMPONENTS_API FDynamicMeshComponentSelectorFactory : public IGeometrySelectorFactory
+{
+public:
+	virtual bool CanBuildForTarget(FGeometryIdentifier TargetIdentifier) const;
+
+	virtual TUniquePtr<IGeometrySelector> BuildForTarget(FGeometryIdentifier TargetIdentifier) const;
+};
+
+
+
+
+
diff --git a/Engine/Plugins/Runtime/MeshModelingToolset/Source/ModelingComponents/Public/Selection/GeometrySelection.h b/Engine/Plugins/Runtime/MeshModelingToolset/Source/ModelingComponents/Public/Selection/GeometrySelection.h
new file mode 100644
--- /dev/null
+++ b/Engine/Plugins/Runtime/MeshModelingToolset/Source/ModelingComponents/Public/Selection/GeometrySelection.h
@@ -0,0 +1,338 @@
+// Copyright Epic Games, Inc. All Rights Reserved.
+
+#pragma once
+
+#include "Math/IntVector.h"
+#include "TriangleTypes.h"
+#include "SegmentTypes.h"
+#include "BoxTypes.h"
+#include "DynamicMesh/InfoTypes.h"
+
+PREDECLARE_GEOMETRY(class FDynamicMesh3);
+
+
+/**
+ * FGeoSelectionID provides a pair of 32-bit unsigned integers that can
+ * be packed into a 64-bit unsigned integer for use with FGeometrySelection.
+ * This is generally intended to be used to encode a mesh geometry ID
+ * (eg TriangleID, EdgeID, VertexID) combined with a "Topology ID",
+ * eg something like a Face Group ID. However none of this is enforced
+ * and so a caller can use these two integers for any purpose.
+ * 
+ * Note that since the ints are unsigned, IndexConstants::InvalidID
+ * is not directly representable (-1 will become positive 0xFFFFFFFF). 
+ */
+struct MODELINGCOMPONENTS_API FGeoSelectionID
+{
+	/** Topology ID, stored in upper 32 bits when packed into 64-bits */
+	uint32 TopologyID;
+	/** Geometry ID, stored in lower 32 bits when packed into 64-bits */
+	uint32 GeometryID;
+
+	FGeoSelectionID()
+	{
+		GeometryID = 0;
+		TopologyID = 0;
+	}
+
+	/**
+	 * Initialize the TopologyID and GeometryID with the given values
+	 */
+	explicit FGeoSelectionID(uint32 GeometryIDIn, uint32 TopologyIDIn = 0)
+	{
+		GeometryID = GeometryIDIn;
+		TopologyID = TopologyIDIn;
+	}
+
+	/**
+	 * Initialize the TopologyID and GeometryID by unpacking the 64-bit packed EncodedID
+	 */
+	explicit FGeoSelectionID(uint64 EncodedID)
+	{
+		GeometryID = EncodedID & 0x00000000FFFFFFFF;
+		TopologyID = (EncodedID & 0xFFFFFFFF00000000) >> 32;
+	}
+
+	/**
+	 * @return the packed 64-bit representation of [TopologyID, GeometryID]
+	 */
+	uint64 Encoded() const
+	{
+		return ((uint64)TopologyID << 32) | (uint64)GeometryID;
+	}
+
+	/**
+	 * @return a FGeoSelectionID initialized with the given TriangleID, and a TopologyID of 0
+	 */
+	static FGeoSelectionID MeshTriangle(int32 TriangleID)
+	{
+		return FGeoSelectionID((uint32)TriangleID, 0);
+	}
+
+	/**
+	 * @return a FGeoSelectionID initialized with the given VertexID, and a TopologyID of 0
+	 */
+	static FGeoSelectionID MeshVertex(int32 VertexID)
+	{
+		return FGeoSelectionID((uint32)VertexID, 0);
+	}
+
+	/**
+	 * @return a FGeoSelectionID initialized with the given 32-bit encoding of the MeshTriEdgeID, and a TopologyID of 0
+	 */
+	static FGeoSelectionID MeshEdge(UE::Geometry::FMeshTriEdgeID EdgeKey)
+	{
+		return FGeoSelectionID(EdgeKey.Encoded());
+	}
+
+	/**
+	 * @return a FGeoSelectionID initialized with the given TriangleID and GroupID used as the TopologyID
+	 */
+	static FGeoSelectionID GroupFace(int32 TriangleID, int32 GroupID)
+	{
+		return FGeoSelectionID((uint32)TriangleID, (uint32)GroupID);
+	}
+
+	friend uint32 GetTypeHash(const FGeoSelectionID& TopoKey)
+	{
+		return ::GetTypeHash(TopoKey.Encoded());
+	}
+};
+
+
+
+
+/**
+ * Type of selected Elements in a FGeometrySelection
+ */
+enum class MODELINGCOMPONENTS_API EGeometryElementType
+{
+	/** Mesh Vertices, Polygroup Corners, ... */
+	Vertex = 1,
+	/** Mesh Edges, Polygroup Edges, ... */
+	Edge = 2,
+	/** Mesh Triangles, Polygroup Faces, ... */
+	Face = 4
+};
+
+/**
+ * Type of selected Topology in a FGeometrySelection.
+ */
+enum class MODELINGCOMPONENTS_API EGeometryTopologyType
+{
+	Triangle = 1,
+	Polygroup = 2
+};
+
+
+/**
+ * FGeometrySelection represents a subset of geometric elements of a larger
+ * object, for example a Mesh (currently the only use case). The main selection
+ * is represented via 64-bit unsigned integers. The integers are stored in a TSet
+ * for efficient unique adds and removals. No assumptions are made about the values,
+ * they could be (eg) mesh indices, IDs of some type, or even pointer values.
+ */
+struct MODELINGCOMPONENTS_API FGeometrySelection
+{
+	/** Type of geometric element represented by this selection, if applicable */
+	EGeometryElementType ElementType = EGeometryElementType::Face;
+
+	/** Type of geometric topology this selection is defined relative to, if applicable */
+	EGeometryTopologyType TopologyType = EGeometryTopologyType::Triangle;
+
+	/** Set of selected items/elements */
+	TSet<uint64> Selection;
+
+	/**
+	 * @return true if Selection is empty
+	 */
+	bool IsEmpty() const
+	{
+		return Selection.Num() == 0;
+	}
+
+	/**
+	 * @return number of elements in Selection
+	 */
+	int32 Num() const
+	{
+		return Selection.Num();
+	}
+
+	/**
+	 * Clear the Selection (may not release memory)
+	 */
+	void Reset()
+	{
+		Selection.Reset();
+	}
+
+	/**
+	 * Initialize the Element and Topology types for this Selection based on another Selection
+	 */
+	void InitializeTypes(const FGeometrySelection& FromSelection)
+	{
+		ElementType = FromSelection.ElementType;
+		TopologyType = FromSelection.TopologyType;
+	}
+};
+
+
+
+/**
+ * Type of change, relative to a FGeometrySelection
+ */
+enum class MODELINGCOMPONENTS_API EGeometrySelectionChangeType
+{
+	/** Elements Added to Selection */
+	Add,
+	/** Elements Removed from Selection */
+	Remove,
+	/** Selected Elements replaced with a new set of Selected Elements */
+	Replace
+};
+
+
+/**
+ * FGeometrySelectionUpdateConfig is passed to various Selection Editing functions/classes
+ * to indicate what type of change should be applied to a FGeometrySelection
+ * (based on additional parameters, generally)
+ */
+struct MODELINGCOMPONENTS_API FGeometrySelectionUpdateConfig
+{
+	EGeometrySelectionChangeType ChangeType = EGeometrySelectionChangeType::Add;
+};
+
+
+/**
+ * FGeometrySelectionDelta represents a change to the set of elements in a FGeometrySelection.
+ * The delta is ordered, ie if the Delta was to be re-applied, the Removed elements should
+ * be removed before the Added elements are added.
+ * 
+ * (Currently there is no way to swap the order)
+ */
+struct MODELINGCOMPONENTS_API FGeometrySelectionDelta
+{
+	/**
+	 * Elements removed from a FGeometrySelection during some selection edit
+	 */
+	TArray<uint64> Removed;
+	/**
+	 * Elements added to a FGeometrySelection during some selection edit
+	 */
+	TArray<uint64> Added;
+
+	/** @return true if the Delta is empty (nothing Removed or Added) */
+	bool IsEmpty() const { return Removed.Num() == 0 && Added.Num() == 0; }
+};
+
+
+
+/**
+ * 3D Bounding information for a FGeometrySelection
+ */
+struct MODELINGCOMPONENTS_API FGeometrySelectionBounds
+{
+	UE::Geometry::FAxisAlignedBox3d WorldBounds = UE::Geometry::FAxisAlignedBox3d::Empty();
+};
+
+
+/**
+ * 3D Geometry representing a FGeometrySelection, for example
+ * suitable for passing to rendering code, etc
+ */
+struct MODELINGCOMPONENTS_API FGeometrySelectionElements
+{
+	TArray<UE::Geometry::FTriangle3d> Triangles;
+	TArray<UE::Geometry::FSegment3d> Segments;
+	TArray<FVector3d> Points;
+
+	void Reset()
+	{
+		Triangles.Reset();
+		Segments.Reset();
+		Points.Reset();
+	}
+};
+
+
+struct MODELINGCOMPONENTS_API FGeometrySelectionUpdateResult
+{
+	bool bSelectionMissed = false;
+	bool bSelectionModified = false;
+
+	FGeometrySelectionDelta SelectionDelta;
+};
+
+
+
+/**
+ * FGeometrySelectionEditor is a utility/helper class used for modifying
+ * a FGeometrySelection. The various functions can be used to add or remove
+ * to the Selection, while also tracking what changed, returned via
+ * FGeometrySelectionDelta structs. 
+ */
+class MODELINGCOMPONENTS_API FGeometrySelectionEditor
+{
+public:
+	/**
+	 * Initialize the Editor with the given Selection. The
+	 * TargetSelectionIn must live longer than the FGeometrySelectionEditor
+	 */
+	void Initialize(FGeometrySelection* TargetSelectionIn);
+
+	/** @return the Element Type of the Target Selection */
+	EGeometryElementType GetElementType() const { return TargetSelection->ElementType; }
+	/** @return the Topology Type of the Target Selection */
+	EGeometryTopologyType GetTopologyType() const { return TargetSelection->TopologyType; }
+
+	/** @return true if the given ID is currently selected in the Target Selection */
+	bool IsSelected(uint64 ID) const;
+
+	/** Clear the Target Selection and return change information in DeltaOut */
+	void ClearSelection(FGeometrySelectionDelta& DeltaOut);
+
+	/** Add the items in the List to the Target Selection and return change information in DeltaOut */
+	template<typename ListType>
+	bool Select(const ListType& List, FGeometrySelectionDelta& DeltaOut)
+	{
+		int32 NumAdded = 0;
+		for (uint64 ID : List)
+		{
+			if (TargetSelection->Selection.Contains(ID) == false)
+			{
+				TargetSelection->Selection.Add(ID);
+				DeltaOut.Added.Add(ID);
+				NumAdded++;
+			}
+		}
+		return (NumAdded > 0);
+	}
+
+	/** Remove the items in the List from the Target Selection and return change information in DeltaOut */
+	template<typename ListType>
+	bool Deselect(const ListType& List, FGeometrySelectionDelta& DeltaOut)
+	{
+		int32 TotalRemoved = 0;
+		for (uint64 ID : List)
+		{
+			int32 NumRemoved = TargetSelection->Selection.Remove(ID);
+			if (NumRemoved > 0)
+			{
+				DeltaOut.Removed.Add(ID);
+			}
+			TotalRemoved += NumRemoved;
+		}
+		return (TotalRemoved > 0);
+	}
+
+	/** Replace the current set of selected items with those in the NewSelection, and return change information in DeltaOut */
+	bool Replace(const FGeometrySelection& NewSelection, FGeometrySelectionDelta& DeltaOut);
+
+
+protected:
+	FGeometrySelection* TargetSelection = nullptr;
+};
+
+
+
diff --git a/Engine/Plugins/Runtime/MeshModelingToolset/Source/ModelingComponents/Public/Selection/GeometrySelectionChanges.h b/Engine/Plugins/Runtime/MeshModelingToolset/Source/ModelingComponents/Public/Selection/GeometrySelectionChanges.h
new file mode 100644
--- /dev/null
+++ b/Engine/Plugins/Runtime/MeshModelingToolset/Source/ModelingComponents/Public/Selection/GeometrySelectionChanges.h
@@ -0,0 +1,84 @@
+// Copyright Epic Games, Inc. All Rights Reserved.
+
+#pragma once
+
+#include "InteractiveToolChange.h"
+#include "Selection/GeometrySelection.h"
+#include "Selection/GeometrySelector.h"		// for FGeometryIdentifier
+
+
+/**
+ * Base interface intended to be implemented by FToolCommandChange subclasses that represent 
+ * changes to a FGeometrySelection. The point of the interface is to allow UGeometrySelectionManager
+ * to execute the ApplyChange() and RevertChange() with the FGeometrySelectionEditor it provides.
+ * So the call pattern is that the FToolCommandChange::Apply(), with the UGeometrySelectionManager
+ * as the UObject target, will call UGeometrySelectionManager::ApplyChange(), which then in turn
+ * will call IGeometrySelectionChange::ApplyChange()
+ */
+class MODELINGCOMPONENTS_API IGeometrySelectionChange
+{
+public:
+	virtual ~IGeometrySelectionChange() {}
+
+	virtual FGeometryIdentifier GetIdentifier() const = 0;
+
+	virtual void ApplyChange(FGeometrySelectionEditor* Editor, FGeometrySelectionDelta& ApplyDelta) = 0;
+	virtual void RevertChange(FGeometrySelectionEditor* Editor, FGeometrySelectionDelta& RevertDelta) = 0;
+};
+
+
+/**
+ * FGeometrySelectionDeltaChange stores a Remove-then-Add change in a FGeometrySelection.
+ */
+class MODELINGCOMPONENTS_API FGeometrySelectionDeltaChange : public FToolCommandChange, public IGeometrySelectionChange
+{
+public:
+	FGeometryIdentifier Identifier;
+	FGeometrySelectionDelta Delta;
+
+	/** Makes the change to the object */
+	virtual void Apply(UObject* Object) override;
+
+	/** Reverts change to the object */
+	virtual void Revert(UObject* Object) override;
+
+	/** Describes this change (for debugging) */
+	virtual FString ToString() const override;
+
+	/** Change has expired if the SelectionManager no longer exists */
+	virtual bool HasExpired(UObject* Object) const override;
+
+	virtual FGeometryIdentifier GetIdentifier() const override { return Identifier;	}
+	virtual void ApplyChange(FGeometrySelectionEditor* Editor, FGeometrySelectionDelta& ApplyDelta) override;
+	virtual void RevertChange(FGeometrySelectionEditor* Editor, FGeometrySelectionDelta& RevertDelta) override;
+};
+
+
+
+/**
+ * FGeometrySelectionReplaceChange stores a full replacement of a FGeometrySelection,
+ * ie full copies of the selection set before and after the change
+ */
+class MODELINGCOMPONENTS_API FGeometrySelectionReplaceChange : public FToolCommandChange, public IGeometrySelectionChange
+{
+public:
+	FGeometryIdentifier Identifier;
+	FGeometrySelection Before;
+	FGeometrySelection After;
+
+	/** Makes the change to the object */
+	virtual void Apply(UObject* Object) override;
+
+	/** Reverts change to the object */
+	virtual void Revert(UObject* Object) override;
+
+	/** Describes this change (for debugging) */
+	virtual FString ToString() const override;
+
+	/** Change has expired if the SelectionManager no longer exists */
+	virtual bool HasExpired(UObject* Object) const override;
+
+	virtual FGeometryIdentifier GetIdentifier() const override { return Identifier;	}
+	virtual void ApplyChange(FGeometrySelectionEditor* Editor, FGeometrySelectionDelta& ApplyDelta) override;
+	virtual void RevertChange(FGeometrySelectionEditor* Editor, FGeometrySelectionDelta& RevertDelta) override;
+};
\ No newline at end of file
diff --git a/Engine/Plugins/Runtime/MeshModelingToolset/Source/ModelingComponents/Public/Selection/GeometrySelectionManager.h b/Engine/Plugins/Runtime/MeshModelingToolset/Source/ModelingComponents/Public/Selection/GeometrySelectionManager.h
new file mode 100644
--- /dev/null
+++ b/Engine/Plugins/Runtime/MeshModelingToolset/Source/ModelingComponents/Public/Selection/GeometrySelectionManager.h
@@ -0,0 +1,335 @@
+// Copyright Epic Games, Inc. All Rights Reserved.
+
+#pragma once
+
+#include "CoreMinimal.h"
+#include "Selection/GeometrySelector.h"
+#include "Selection/GeometrySelection.h"
+#include "Selection/GeometrySelectionChanges.h"
+
+#include "GeometrySelectionManager.generated.h"
+
+class IGeometrySelector;
+class UInteractiveToolsContext;
+class IToolsContextRenderAPI;
+class IToolsContextTransactionsAPI;
+class UPersistentMeshSelection;
+
+
+/**
+ * UGeometrySelectionManager provides the infrastructure for "Element Selection", ie 
+ * geometric sub-elements of some geometry object like a Triangle Mesh. The Manager is
+ * designed to work with a relatively vague concept of "element", so it doesn't explicitly
+ * reference triangles/etc, and the selectable-elements and how-elements-are-selected
+ * concepts are provided by abstract-interfaces that allow various implememtations.
+ * 
+ * The "Geometry Objects", eg like a DynamicMeshComponent, Gameplay Volume, etc, are
+ * referred to as "Active Targets" in the Manager. External code provides and updates
+ * the set of Active Targets, eg for example tracking the active Actor Selection in the Editor.
+ * 
+ * For a given Target, a tuple (Selector, Selection, SelectionEditor) is created and maintained.
+ * The FGeometrySelection is ultimately a basic list of integers and does not have any knowledge
+ * of what it is a selection *of*, and is not intended to be directly edited. Instead the
+ * SelectionEditor provides that functionality. This separation allows "selection editing" to
+ * be customized, eg to enforce invariants or constraints that might apply to certain kinds of selections.
+ * 
+ * The IGeometrySelector provides the core implementation of what "selection" means for a given
+ * Target, eg like a mesh Component, or mesh object like a UDynamicMesh. The Selector is
+ * created by a registered Factory, allowing client code to provide custom implementations for
+ * different Target Types. Updates to the Selection are done via the Selector, as well as queries
+ * about (eg) renderable selection geometry. 3D Transforms are also applied via the Selector,
+ * as only it has the knowledge about what can be transformed and how it can be applied.
+ * 
+ * The GeometrySelectionManager provides high-level interfaces for this system, for example
+ * external code (eg such as something that creates a Gizmo for the active selection) only
+ * needs to interact with SelectionManager, calling functions like 
+ * ::BeginTransformation() / ::UpdateTransformation() / ::EndTransformation().
+ * The SelectionManager also handles Transactions/FChanges for the active Targets and Selections. 
+ * 
+ */
+UCLASS()
+class MODELINGCOMPONENTS_API UGeometrySelectionManager : public UObject
+{
+	GENERATED_BODY()
+
+public:
+
+	//
+	// Setup/Teardown
+	//
+
+	virtual void Initialize(UInteractiveToolsContext* ToolsContextIn, IToolsContextTransactionsAPI* TransactionsAPIIn);
+
+	virtual void RegisterSelectorFactory(TUniquePtr<IGeometrySelectorFactory> Factory);
+
+	virtual void Shutdown();
+	virtual bool HasBeenShutDown() const;
+
+	UInteractiveToolsContext* GetToolsContext() const { return ToolsContext; }
+	IToolsContextTransactionsAPI* GetTransactionsAPI() const { return TransactionsAPI; }
+
+	//
+	// Configuration
+	//
+
+	/** EMeshTopologyMode determines what level of mesh element will be selected  */
+	enum class EMeshTopologyMode
+	{
+		None = 0,
+		/** Select mesh triangles, edges, and vertices */
+		Triangle = 1,
+		/** Select mesh polygroups, polygroup-borders, and polygroup-corners */
+		Polygroup = 2
+	};
+
+	virtual void SetSelectionElementType(EGeometryElementType ElementType);
+	virtual EGeometryElementType GetSelectionElementType() const { return SelectionElementType; }
+
+	virtual void SetMeshTopologyMode(EMeshTopologyMode SelectionMode);
+	virtual EMeshTopologyMode GetMeshTopologyMode() const { return MeshTopologyMode; }
+	virtual EGeometryTopologyType GetSelectionTopologyType() const;
+
+	
+
+
+	//
+	// Target Management / Queries
+	// TODO: be able to update active target set w/o losing current selections?
+	//
+
+	/**
+	 * @return true if there are any active selection targets
+	 */
+	bool HasActiveTargets() const;
+	
+	/**
+	 * Empty the active selection target set
+	 * @warning Active selection must be cleared (eg via ClearSelection()) before calling this function
+	 */
+	void ClearActiveTargets();
+
+	/**
+	 * Add a target to the active target set, if a valid IGeometrySelectorFactory can be found
+	 * @return true on success
+	 */
+	bool AddActiveTarget(FGeometryIdentifier Target);
+
+	/**
+	 * Update the current active target set based on DesiredActiveSet, assuming that a valid IGeometrySelectorFactory
+	 * can be found for each identifier. 
+	 * This function will emit a transaction/change if the target set is modified.
+	 * @param WillChangeActiveTargetsCallback this function will be called if the DesiredActiveSet is not the same as the current active set, allowing calling code to (eg) clear the active selection
+	 */
+	void SynchronizeActiveTargets(
+		const TArray<FGeometryIdentifier>& DesiredActiveSet,
+		TFunctionRef<void()> WillChangeActiveTargetsCallback );
+
+	/**
+	 * Test if a World-space ray "hits" the current active target set, which can be used to (eg)
+	 * determine if a higher-level user interaction for selection should "Capture" the click
+	 * @param HitResultOut hit information is returned here
+	 * @return true on hit
+	 */
+	virtual bool RayHitTest(
+		const FRay3d& WorldRay,
+		FInputRayHit& HitResultOut
+	);
+
+
+	//
+	// Selection Updates
+	//
+public:
+	/**
+	 * Clear any active element selections. 
+	 * This function will emit a Transaction for the selection change.
+	 */
+	virtual void ClearSelection();
+
+	/**
+	 * Use the given WorldRay to update the active element selection based on UpdateConfig.
+	 * The intention is that this function is called by higher-level user interaction code after
+	 * RayHitTest() has returned true.
+	 * @param ResultOut information on any element selection modifications is returned here
+	 */
+	virtual void UpdateSelectionViaRaycast(
+		const FRay3d& WorldRay,
+		const FGeometrySelectionUpdateConfig& UpdateConfig,
+		FGeometrySelectionUpdateResult& ResultOut
+	);
+
+
+	DECLARE_MULTICAST_DELEGATE(FModelingSelectionInteraction_SelectionModified);
+	/**
+	 * OnSelectionModified is broadcast if the selection is modified via the above functions.
+	 * There are no arguments.
+	 */
+	FModelingSelectionInteraction_SelectionModified OnSelectionModified;
+
+
+	//
+	// Selection queries
+	//
+public:
+	/** @return true if there is an active element selection */
+	virtual bool HasSelection() const;
+
+	/** @return a world-space bounding box for the active element selection */
+	virtual bool GetSelectionBounds(FGeometrySelectionBounds& BoundsOut) const;
+	/** @return a 3D transformation frame suitable for use with the active element selection */
+	virtual void GetSelectionWorldFrame(UE::Geometry::FFrame3d& SelectionFrame) const;
+
+	/** @return true if there is an active IGeometrySelector target for the given Component and it has a non-empty selection */
+	virtual bool HasSelectionForComponent(UPrimitiveComponent* Component) const;
+	/** 
+	 * Get the active element selection for the given Component, if it exists
+	 * @return true if a valid non-empty selection was returned in SelectionOut
+	 */
+	virtual bool GetSelectionForComponent(UPrimitiveComponent* Component, FGeometrySelection& SelectionOut) const;
+
+
+
+	//
+	// Transformations
+	//
+protected:
+	// Set of existing transformer objects, collected from active IGeometrySelector::InitializeTransformation().
+	// The SelectionManager does not own these objects, they must not be deleted. 
+	// Instead they need to be returned via IGeometrySelector::ShutdownTransformation
+	TArray<IGeometrySelectionTransformer*> ActiveTransformations;
+
+public:
+	/** @return true if SelectionManager is actively transforming element selections (ie during a mouse-drag) */
+	virtual bool IsInActiveTransformation() const { return (ActiveTransformations.Num() > 0); }
+	
+	/**
+	 * Begin a transformation of element selections in active Targets.
+	 * @return true if at least one valid Transformer was initialized, ie the transformation will do something
+	 */
+	virtual bool BeginTransformation();
+
+	/**
+	 * Update the active transformations with the given PositionTransformFunc.
+	 * See IGeometrySelectionTransformer::UpdateTransform for details on this callback
+	 */
+	virtual void UpdateTransformation( TFunctionRef<FVector3d(int32 VertexID, const FVector3d&, const FTransform&)> PositionTransformFunc );
+
+	/**
+	 * End the current active transformation, and emit changes/transactions
+	 */
+	virtual void EndTransformation();
+
+
+	//
+	// Undo/Redo
+	//
+public:
+	virtual void ApplyChange(IGeometrySelectionChange* Change);
+	virtual void RevertChange(IGeometrySelectionChange* Change);
+
+
+	//
+	// Debugging stuff
+	//
+public:
+	/** Print information about the active selection using UE_LOG */
+	virtual void DebugPrintSelection();
+	/** Visualize the active selection using PDI drawing */
+	virtual void DebugRender(IToolsContextRenderAPI* RenderAPI);
+
+	/**
+	 * Convert the active selection to a UPersistentMeshSelection. 
+	 * SelectionManager will hold a reference to UPersistentMeshSelection to prevent it from being garbage collected,
+	 * but only until the next call to GetActiveSingleSelectionConverted_Legacy()
+	 * This function will be removed in future when UPersistentMeshSelection is no longer needed.
+	 */
+	UPersistentMeshSelection* GetActiveSingleSelectionConverted_Legacy(UPrimitiveComponent* ForComponent);
+
+
+protected:
+
+	// current selection mode settings
+
+	EGeometryElementType SelectionElementType = EGeometryElementType::Face;
+	void SetSelectionElementTypeInternal(EGeometryElementType NewElementType);
+
+	EMeshTopologyMode MeshTopologyMode = EMeshTopologyMode::Polygroup;
+	void SetMeshTopologyModeInternal(EMeshTopologyMode NewTopologyMode);
+
+	// ITF references
+
+	UPROPERTY()
+	TObjectPtr<UInteractiveToolsContext> ToolsContext;
+
+	IToolsContextTransactionsAPI* TransactionsAPI;
+
+	// set of registered IGeometrySelector factories
+	TArray<TUniquePtr<IGeometrySelectorFactory>> Factories;
+
+
+
+	// FGeometrySelectionTarget is the set of information tracked for a given "Active Target",
+	// which is (eg) a Mesh Component or other external object that "owns" selectable Geometry.
+	// This includes the IGeometrySelector for that target, the SelectionEditor, and the active Selection
+	struct FGeometrySelectionTarget
+	{
+	public:
+		FGeometryIdentifier TargetIdentifier;		// identifier of target object used to initialize the selection (eg Component/etc)
+		FGeometryIdentifier SelectionIdentifer;		// identifier of object that is being selected-on, eg UDynamicMesh/etc
+
+		TUniquePtr<IGeometrySelector> Selector;					// active Selector
+
+		FGeometrySelection Selection;				// current Selection
+		TUniquePtr<FGeometrySelectionEditor> SelectionEditor;	// active Selection Editor
+
+		FDelegateHandle OnGeometryModifiedHandle;	// hooked up to (eg) UDynamicMesh OnMeshChanged, etc
+	};
+
+	// Set of active Selection Targets updated by SynchronizeActiveTargets / etc
+	TArray<TSharedPtr<FGeometrySelectionTarget>> ActiveTargetReferences;
+
+	// map from external Identifiers to active Selection Targets
+	TMap<FGeometryIdentifier, TSharedPtr<FGeometrySelectionTarget>> ActiveTargetMap;
+
+
+	//
+	// Support for cached FGeometrySelectionTarget / IGeometrySelectors. 
+	// The intention here is to reduce the overhead on selection changes.
+	// Functional, but needs to be smarter.
+	//
+	TMap<FGeometryIdentifier, TSharedPtr<FGeometrySelectionTarget>> TargetCache;
+	void SleepOrShutdownTarget(FGeometrySelectionTarget* Target, bool bForceShutdown);
+	TSharedPtr<FGeometrySelectionTarget> GetCachedTarget(FGeometryIdentifier Identifier, const IGeometrySelectorFactory* UseFactory);
+	void ResetTargetCache();
+	void SetTargetsOnUndoRedo(TArray<FGeometryIdentifier> NewTargets);
+	TArray<FGeometryIdentifier> GetCurrentTargetIdentifiers() const;
+
+	void OnTargetGeometryModified(IGeometrySelector* Selector);
+
+	//
+	// 3D geometry for element selections of each ActiveTarget is cached
+	// to improve rendering performance
+	//
+	TArray<FGeometrySelectionElements> CachedSelectionRenderElements;
+	bool bSelectionRenderCachesDirty = false;
+	void UpdateSelectionRenderCacheOnTargetChange();
+	void RebuildSelectionRenderCaches();
+
+	// various change types need internal access
+
+	friend class FGeometrySelectionManager_SelectionTypeChange;
+	friend class FGeometrySelectionManager_ActiveTargetsChange;
+
+
+	// legacy / to-deprecate/remove
+	
+	// this is to support GetActiveSingleSelectionConverted_Legacy, which will eventually be removed
+	UPROPERTY()
+	TObjectPtr<UPersistentMeshSelection> OldSelection;
+
+};
+
+
+
+
+
diff --git a/Engine/Plugins/Runtime/MeshModelingToolset/Source/ModelingComponents/Public/Selection/GeometrySelectionUtil.h b/Engine/Plugins/Runtime/MeshModelingToolset/Source/ModelingComponents/Public/Selection/GeometrySelectionUtil.h
new file mode 100644
--- /dev/null
+++ b/Engine/Plugins/Runtime/MeshModelingToolset/Source/ModelingComponents/Public/Selection/GeometrySelectionUtil.h
@@ -0,0 +1,180 @@
+// Copyright Epic Games, Inc. All Rights Reserved.
+
+#pragma once
+
+#include "Selection/GeometrySelection.h"
+#include "TriangleTypes.h"
+#include "SegmentTypes.h"
+#include "FrameTypes.h"
+#include "Polygroups/PolygroupSet.h"
+
+PREDECLARE_GEOMETRY(class FDynamicMesh3);
+
+namespace UE
+{
+namespace Geometry
+{
+
+class FColliderMesh;
+class FGroupTopology;
+struct FGroupTopologySelection;
+
+/**
+ * Update a FGeometrySelection (via the Editor) of mesh vertices/edges/triangles
+ * based on a raycast
+ * @param UpdateConfig type of update (add/remove/replace) and other update controls
+ * @param ResultOut selection delta and information on hit/miss are returned here
+ */
+MODELINGCOMPONENTS_API void UpdateTriangleSelectionViaRaycast(
+	const FColliderMesh* ColliderMesh,
+	FGeometrySelectionEditor* Editor,
+	const FRay3d& LocalRay,
+	const FGeometrySelectionUpdateConfig& UpdateConfig,
+	FGeometrySelectionUpdateResult& ResultOut);
+
+/**
+ * Update a FGeometrySelection (via the Editor) of mesh polygroup faces/edges/corners,
+ * based on a raycast
+ * @param UpdateConfig type of update (add/remove/replace) and other update controls
+ * @param ResultOut selection delta and information on hit/miss are returned here
+ */
+MODELINGCOMPONENTS_API void UpdateGroupSelectionViaRaycast(
+	const FColliderMesh* ColliderMesh,
+	const FGroupTopology* GroupTopology,
+	FGeometrySelectionEditor* Editor,
+	const FRay3d& LocalRay,
+	const FGeometrySelectionUpdateConfig& UpdateConfig,
+	FGeometrySelectionUpdateResult& ResultOut);
+
+
+/**
+ * Update a FGeometrySelection (via the Editor)
+ * @param ChangeType type of change to make (add/remove/replace)
+ * @param NewIDs set of new IDs to use to update the selection
+ * @param Delta selection delta will be stored here, if non-nullptr is provided
+ */
+MODELINGCOMPONENTS_API bool UpdateSelectionWithNewElements(
+	FGeometrySelectionEditor* Editor,
+	EGeometrySelectionChangeType ChangeType,
+	const TArray<uint64>& NewIDs,
+	FGeometrySelectionDelta* DeltaOut = nullptr);
+
+
+/**
+ * Call VertexFunc for each selected Mesh element (vertex/edge/tri) in MeshSelection.
+ * ApplyTransform will be applied to Vertex Positions before calling VertexFunc
+ */
+MODELINGCOMPONENTS_API bool EnumerateTriangleSelectionVertices(
+	const FGeometrySelection& MeshSelection,
+	const UE::Geometry::FDynamicMesh3& Mesh,
+	const FTransform& ApplyTransform,
+	TFunctionRef<void(uint64, const FVector3d&)> VertexFunc
+);
+/**
+ * Call VertexFunc for each selected Mesh element (vertex/edge/tri) in the set
+ * of polygroup faces/edges/corners specified by GroupSelection (relative to GroupTopology parameter)
+ * ApplyTransform will be applied to Vertex Positions before calling VertexFunc
+ */
+MODELINGCOMPONENTS_API bool EnumeratePolygroupSelectionVertices(
+	const FGeometrySelection& GroupSelection,
+	const UE::Geometry::FDynamicMesh3& Mesh,
+	const FGroupTopology* GroupTopology,
+	const FTransform& ApplyTransform,
+	TFunctionRef<void(uint64, const FVector3d&)> VertexFunc
+);
+
+
+/**
+ * Call TriangleFunc for each mesh TriangleID included in MeshSelection.
+ * TriangleFunc may be called multiple times for the same TriangleID.
+ * This will forward to EnumerateTriangleSelectionTriangles() or 
+ * EnumeratePolygroupSelectionTriangles() depending on the selection topology type.
+ * If UseGroupSet and MeshSelection is for polygroups, the default Mesh group layer will be used.
+ */
+MODELINGCOMPONENTS_API bool EnumerateSelectionTriangles(
+	const FGeometrySelection& MeshSelection,
+	const UE::Geometry::FDynamicMesh3& Mesh,
+	TFunctionRef<void(int32)> TriangleFunc,
+	const UE::Geometry::FPolygroupSet* UseGroupSet = nullptr
+);
+/**
+ * Call TriangleFunc for each mesh TriangleID included in MeshSelection.
+ * For Edges, both connected edges are included.
+ * For Vertices, all triangles in the vertex one-ring are included.
+ */
+MODELINGCOMPONENTS_API bool EnumerateTriangleSelectionTriangles(
+	const FGeometrySelection& MeshSelection,
+	const UE::Geometry::FDynamicMesh3& Mesh,
+	TFunctionRef<void(int32)> TriangleFunc
+);
+/**
+ * Call TriangleFunc for each mesh TriangleID included in MeshSelection, where MeshSelection has polygroup topology.
+ * For Polygroup Faces, all triangles in the face are included.
+ * For Polygroup Edges, currently all triangles in any group adjacent to the edge
+ * For Polygroup Corners, currently all triangles in any group touching the corner
+ */
+MODELINGCOMPONENTS_API bool EnumeratePolygroupSelectionTriangles(
+	const FGeometrySelection& MeshSelection,
+	const UE::Geometry::FDynamicMesh3& Mesh,
+	const UE::Geometry::FPolygroupSet& GroupSet,
+	TFunctionRef<void(int32)> TriangleFunc
+);
+
+
+/**
+ * Call VertexFunc/EdgeFunc/TriangleFunc for the vertices/edges/triangles identified by MeshSelection.
+ * Since a MeshSelection only stores vertices, edges, or triangles, but not combined, only one
+ * of these functions will be invoked during a call to this function.
+ * This function is useful to collect up geometry that needs to be rendered for a given MeshSelection
+ * @param ApplyTransform if non-null, transform is applied to the 3D geometry
+ */
+MODELINGCOMPONENTS_API bool EnumerateTriangleSelectionElements(
+	const FGeometrySelection& MeshSelection,
+	const UE::Geometry::FDynamicMesh3& Mesh,
+	TFunctionRef<void(int32, const FVector3d&)> VertexFunc,
+	TFunctionRef<void(int32, const FSegment3d&)> EdgeFunc,
+	TFunctionRef<void(int32, const FTriangle3d&)> TriangleFunc,
+	const FTransform* ApplyTransform = nullptr
+);
+/**
+ * Call VertexFunc/EdgeFunc/TriangleFunc for the mesh vertices/edges/triangles identified by MeshSelection,
+ * where MeshSelection has polygroup topology referring to the provided GroupTopology
+ * Since a MeshSelection only stores vertices, edges, or triangles, but not combined, only one
+ * of these functions will be invoked during a call to this function.
+ * This function is useful to collect up geometry that needs to be rendered for a given MeshSelection
+ * @param ApplyTransform if non-null, transform is applied to the 3D geometry
+ */
+MODELINGCOMPONENTS_API bool EnumeratePolygroupSelectionElements(
+	const FGeometrySelection& MeshSelection,
+	const UE::Geometry::FDynamicMesh3& Mesh,
+	const FGroupTopology* GroupTopology,
+	TFunctionRef<void(int32, const FVector3d&)> VertexFunc,
+	TFunctionRef<void(int32, const FSegment3d&)> EdgeFunc,
+	TFunctionRef<void(int32, const FTriangle3d&)> TriangleFunc,
+	const FTransform* ApplyTransform = nullptr
+);
+
+
+/**
+ * Convert a MeshSelection with Polygroup topology type to a FGroupTopologySelection
+ */
+MODELINGCOMPONENTS_API bool ConvertPolygroupSelectionToTopologySelection(
+	const FGeometrySelection& MeshSelection,
+	const UE::Geometry::FDynamicMesh3& Mesh,
+	const FGroupTopology* GroupTopology,
+	FGroupTopologySelection& TopologySelectionOut
+);
+
+
+/**
+ * Compute a 3D Frame suitable for use as a 3D transform gizmo position/orientation
+ * for the given MeshSelection
+ */
+MODELINGCOMPONENTS_API bool GetTriangleSelectionFrame(
+	const FGeometrySelection& MeshSelection,
+	const UE::Geometry::FDynamicMesh3& Mesh,
+	FFrame3d& SelectionFrameOut);
+
+
+}
+}
\ No newline at end of file
diff --git a/Engine/Plugins/Runtime/MeshModelingToolset/Source/ModelingComponents/Public/Selection/GeometrySelector.h b/Engine/Plugins/Runtime/MeshModelingToolset/Source/ModelingComponents/Public/Selection/GeometrySelector.h
new file mode 100644
--- /dev/null
+++ b/Engine/Plugins/Runtime/MeshModelingToolset/Source/ModelingComponents/Public/Selection/GeometrySelector.h
@@ -0,0 +1,348 @@
+// Copyright Epic Games, Inc. All Rights Reserved.
+
+#pragma once
+
+#include "Components/PrimitiveComponent.h"
+#include "Selection/GeometrySelection.h"
+#include "FrameTypes.h"
+#include "InputState.h"
+
+class IToolsContextTransactionsAPI;
+
+
+
+/**
+ * FGeometryIdentifier is used to represent a specific Geometry-containing object.
+ */
+struct MODELINGCOMPONENTS_API FGeometryIdentifier
+{
+	/**
+	 * ETargetType specifies the category of target object
+	 */
+	enum class ETargetType
+	{
+		Unknown = 0,
+		/** TargetObject is a UPrimitiveComponent subclass */
+		PrimitiveComponent = 1,
+		/** TargetObject is a UObject container for a Mesh, like UStaticMesh or UDynamicMesh */
+		MeshContainer = 2,
+
+		UserDefinedBase = 32
+	};
+
+	/**
+	 * EObjectType is the specific UClass/UObject Type of TargetObject
+	 */
+	enum class EObjectType
+	{
+		Unknown = 0,
+		DynamicMeshComponent = 1,
+		DynamicMesh = 2,
+		BrushComponent = 3,
+
+		UserDefinedBase = 64
+	};
+
+	/**
+	 * Target Type specifies the category of TargetObject
+	 */
+	ETargetType TargetType = ETargetType::PrimitiveComponent;
+
+	/**
+	 * Object Type specifies the specific UClass/UObject Type of TargetObject
+	 */
+	EObjectType ObjectType = EObjectType::Unknown;	
+
+	/**
+	 * TargetObject is *optional*, depending on the identifier and situation it may or may not exist
+	 */
+	TWeakObjectPtr<UObject> TargetObject = nullptr;
+
+
+	FGeometryIdentifier() = default;
+	FGeometryIdentifier(const FGeometryIdentifier& Other) = default;
+
+	/**
+	 * @return true if TargetObject is defined and still exists
+	 */
+	bool HasObject() const { return TargetObject.IsValid(); }
+
+
+	bool operator==(const FGeometryIdentifier& Other) const
+	{
+		return Equals(Other);
+	}
+	bool Equals(const FGeometryIdentifier& Other) const 
+	{
+		return TargetType == Other.TargetType && ObjectType == Other.ObjectType && TargetObject == Other.TargetObject;
+	}
+
+	friend uint32 GetTypeHash(const FGeometryIdentifier& Object)
+	{
+		return HashCombine(
+			HashCombine(GetTypeHash(Object.TargetType), GetTypeHash(Object.ObjectType)), 
+			GetTypeHash(Object.TargetObject));
+	}
+
+	/**
+	 * @return TargetObject cast to template ObjectTypeT, or nullptr if it is not defined or not this type
+	 */
+	template<typename ObjectTypeT>
+	ObjectTypeT* GetAsObjectType() const
+	{
+		return Cast<ObjectTypeT>(TargetObject.Get());
+	}
+
+	/**
+	 * @return TargetObject cast to template PrimitiveComponent subclass ComponentTypeT, or nullptr if it is not defined or not this type
+	 */
+	template<typename ComponentTypeT>
+	ComponentTypeT* GetAsComponentType() const
+	{
+		if (ensure(TargetType == ETargetType::PrimitiveComponent) == false)
+		{
+			return nullptr;
+		}
+		return Cast<ComponentTypeT>(TargetObject.Get());
+	}
+
+
+	/**
+	 * @return a FGeometryIdentifier for a given UPrimitiveComponent with the specified ObjectType (may be EObjectType::Unknown)
+	 */
+	static FGeometryIdentifier PrimitiveComponent(UPrimitiveComponent* Component, EObjectType ObjectType)
+	{
+		FGeometryIdentifier Identifier;
+		Identifier.TargetType = ETargetType::PrimitiveComponent;
+		Identifier.ObjectType = ObjectType;
+		Identifier.TargetObject = Component;
+		return Identifier;
+	}
+
+};
+
+
+
+/**
+ * FGeometrySelectionHandle stores a Selection and an Identifier for
+ * the geometry object that the Selection is defined relative to
+ */
+struct MODELINGCOMPONENTS_API FGeometrySelectionHandle
+{
+	FGeometryIdentifier Identifier;
+	const FGeometrySelection* Selection;
+};
+
+
+class IGeometrySelector;
+
+
+/**
+ * IGeometrySelectionTransformer is a transient object that is created by an IGeometrySelector
+ * to provide external code with a way to manipulate/transform an active Selection.
+ * This allows (eg) a Gizmo to be hooked up to an active Selection without either
+ * explicitly knowing about each other. 
+ * 
+ * You should not construct an IGeometrySelectionTransformer implementation directly.
+ * IGeometrySelector::InitializeTransformation() will create a suitable IGeometrySelectionTransformer
+ * implementation if transformation is supported, and IGeometrySelector::ShutdownTransformation() is used
+ * to return/destroy it. Only one active Transformer is supported on a given Selector.
+ * 
+ * The API functions BeginTransform/UpdateTransform/EndTransform are called by
+ * external code with the relevant information, see details below.
+ * 
+ * The current API assumes the transformation will be applied "per vertex", ie that
+ * the concept of "vertices" exists and that the transform is done pointwise.
+ * This may be expanded/revisited in future.
+ */
+class MODELINGCOMPONENTS_API IGeometrySelectionTransformer
+{
+public:
+	virtual ~IGeometrySelectionTransformer() {}
+
+	/** @return the IGeometrySelector that owns this Transformer */
+	virtual IGeometrySelector* GetSelector() const = 0;
+
+	/**
+	 * Start a transform (eg called at beginning of a user interaction with a 3D gizmo).
+	 * The caller provides the active Selection that is intended to be transformed.
+	 */
+	virtual void BeginTransform(const FGeometrySelection& Selection) = 0;
+
+	/**
+	 * Update the active transform (eg called on each mouse update during 3D gizmo usage).
+	 * The caller provides a VertexTransformFunc callback, which will only be called during
+	 * the call to UpdateTransform to compute new 3D positions for 3D points. 
+	 */
+	virtual void UpdateTransform( TFunctionRef<FVector3d(int VertexID, const FVector3d& InitialPosition, const FTransform& WorldTransform)> VertexTransformFunc ) = 0;
+
+	/**
+	 * Finish a transform (eg called on mouse-up during 3D gizmo interaction). 
+	 * The caller provides a TransactionsAPI implementation which the Transformer/Selector can use to emit change events
+	 */
+	virtual void EndTransform(IToolsContextTransactionsAPI* TransactionsAPI) = 0;
+};
+
+
+
+
+/**
+ * IGeometrySelector is a base API definition for, roughly, "an object that knows how to
+ * select elements and/or work with element selections on some other object", where "element"
+ * is some "part" of the object, like a mesh edge/vertex/triangle, or intermediate structure (eg a polygroup).
+ * 
+ * Essentially the purpose of IGeometrySelector is to provide the bridge between "object that can have a selection"
+ * and "the active selection on that object". The Selection is stored as a FGeometrySelection which has 
+ * no reference to a type or target object, and standard target objects (eg a StaticMesh, Volume, etc)
+ * have no inherent knowledge of, or support for, low-level geometry selections. 
+ * So the Selector provides these capabilities, as well as related functionality like 
+ * being able to transform an active selection and emit necessary transactions/etc. 
+ * 
+ * Generally clients will not create a IGeometrySelector instance directly. For a given
+ * implmentation, a IGeometrySelectorFactory is registered with the UGeometrySelectionManager,
+ * and it creates instances based on the active selection.
+ * 
+ * Where possible, the UGeometrySelectionManager will also "sleep" a IGeometrySelector for a given
+ * target object instead of destroying it when the object is deselected, and then "restore"
+ * it if the object is re-selected. This avoids frequent churn in Selectors as the user
+ * clicks between objects.
+ * 
+ * Selector implementations should ideally have a strategy to lazy-construct any large data
+ * structures (eg mesh AABBTrees, etc). Again this is to avoid excessive overhead when the
+ * user (for example) selects a large object but has no intention of doing element-level selection.
+ * 
+ */
+class MODELINGCOMPONENTS_API IGeometrySelector
+{
+public:
+	virtual ~IGeometrySelector() {}
+
+	/**
+	 * Disconnect the Selector from its target object and prepare for deletion.
+	 */
+	virtual void Shutdown() = 0;
+
+
+	/**
+	 * The SelectionManager will call SupportsSleep() to determine if an IGeometrySelector
+	 * needs to be destroyed when the target object is deselected, or if it can be
+	 * re-used (generally more efficient)
+	 */
+	virtual bool SupportsSleep() const { return false; }
+
+	/**
+	 * Temporarily disable the Selector, with the intention of re-enabling it (via Restore()) in
+	 * the future. Sleep should (currently) free any large memory data structures.
+	 * @return true if Sleep was successful, if not the Selector will likely be Shutdown
+	 */
+	virtual bool Sleep() { return false; }
+
+	/**
+	 * Restore a Selector that has had Sleep() called on it. 
+	 * @return true on success
+	 */
+	virtual bool Restore() { return false; }
+
+	/**
+	 * @return a FGeometryIdentifier for this Selector. 
+	 */
+	virtual FGeometryIdentifier GetIdentifier() const = 0;
+
+	/**
+	 * Check for intersection between a world-space Ray with the Selector's target object and return a FInputRayHit result.
+	 * RayHitTest() must return true for other raycast-based functions like UpdateSelectionViaRaycast() to be called.
+	 * @return true on hit, false on miss
+	 */
+	virtual bool RayHitTest(
+		const FRay3d& WorldRay,
+		FInputRayHit& HitResultOut
+	) = 0;
+
+	/**
+	 * Update the active selection based on a world-space Ray. Should only be called if RayHitTest() returns true.
+	 * Uses the provided SelectionEditor and UpdateConfig to do the update
+	 * Information about any changes actually made to the selection via the SelectionEditor are returned in ResultOut
+	 */
+	virtual void UpdateSelectionViaRaycast(	
+		const FRay3d& WorldRay,
+		FGeometrySelectionEditor& SelectionEditor,
+		const FGeometrySelectionUpdateConfig& UpdateConfig,
+		FGeometrySelectionUpdateResult& ResultOut 
+	) = 0;
+
+	/**
+	 * @return the World transform for the Selector's target object
+	 */
+	virtual FTransform GetLocalToWorldTransform() const = 0;
+
+	/**
+	 * Compute a 3D frame for the provided Selection. For example this could simply be a frame aligned to the XYZ axes at the center
+	 * of the selection bounding-box, or something more sophisticated, like aligned to the tangent-space of the selected element(s), etc.
+	 * @param bTransformToWorld if true the resulting SelectionFrameOut will be in world-space, based on GetLocalToWorldTransform()
+	 */
+	virtual void GetSelectionFrame(const FGeometrySelection& Selection, UE::Geometry::FFrame3d& SelectionFrameOut, bool bTransformToWorld) = 0;
+
+	/**
+	 * Accumulate the bounds of the provided Selection in the provided BoundsInOut. BoundsInOut is not cleared.
+	 * @param bTransformToWorld if true each bounded point/element will be transformed to World space before being "contained" in BoundsInOut
+	 */
+	virtual void AccumulateSelectionBounds(const FGeometrySelection& Selection, FGeometrySelectionBounds& BoundsInOut, bool bTransformToWorld) = 0;
+
+	/**
+	 * Accumulate geometric elements (currently 3D triangles, line segments, and points) for the provided Selection in the provided ElementsInOut. 
+	 * ElementsInOut is not cleared.
+	 * @param bTransformToWorld if true each geometric element will be transformed to World space, based on GetLocalToWorldTransform()
+	 */
+	virtual void AccumulateSelectionElements(const FGeometrySelection& Selection, FGeometrySelectionElements& ElementsInOut, bool bTransformToWorld) = 0;
+
+
+	/**
+	 * Create and initialize a IGeometrySelectionTransformer for the provided Selection.
+	 * IGeometrySelector retains ownership of the active Transformer, it should not be
+	 * deleted by the caller, but instead ShutdownTransformation() should be called.
+	 */
+	virtual IGeometrySelectionTransformer* InitializeTransformation(const FGeometrySelection& Selection) = 0;
+
+	/**
+	 * Shutdown an active IGeometrySelectionTransformer. Generally only one Transformer
+	 * may be active at any time, and it is an error to shut it down multiple times.
+	 */
+	virtual void ShutdownTransformation(IGeometrySelectionTransformer* Transformer) = 0;
+
+
+public:
+	DECLARE_MULTICAST_DELEGATE_OneParam(FOnGeometryModified, IGeometrySelector* Selector);
+	/**
+	 * The OnGeometryModified delegate will fire if the geometry of the Selector's Target Object
+	 * is modified. For example if a UDynamicMesh is edited, OnGeometryModified will be broadcast.
+	 */
+	FOnGeometryModified& GetOnGeometryModifed() { return OnGeometryModifed; }
+
+protected:
+	FOnGeometryModified OnGeometryModifed;
+	void NotifyGeometryModified() { OnGeometryModifed.Broadcast(this); }
+
+};
+
+
+/**
+ * Factory for a specific type of IGeometrySelector
+ */
+class MODELINGCOMPONENTS_API IGeometrySelectorFactory
+{
+public:
+	virtual ~IGeometrySelectorFactory() {}
+
+	/**
+	 * @return true if this Factory can build an IGeometrySelector for the provided TargetIdentifier
+	 */
+	virtual bool CanBuildForTarget(FGeometryIdentifier TargetIdentifier) const = 0;
+
+	/**
+	 * @return a new IGeometrySelector for the provided TargetIdentifier
+	 */
+	virtual TUniquePtr<IGeometrySelector> BuildForTarget(FGeometryIdentifier TargetIdentifier) const = 0;
+};
+
+
+
