ModelingComponents: fix mismatched loctext key

diff --git a/Engine/Plugins/Runtime/MeshModelingToolset/Source/ModelingComponents/Private/Selection/GeometrySelectionManager.cpp b/Engine/Plugins/Runtime/MeshModelingToolset/Source/ModelingComponents/Private/Selection/GeometrySelectionManager.cpp
--- a/Engine/Plugins/Runtime/MeshModelingToolset/Source/ModelingComponents/Private/Selection/GeometrySelectionManager.cpp
+++ b/Engine/Plugins/Runtime/MeshModelingToolset/Source/ModelingComponents/Private/Selection/GeometrySelectionManager.cpp
@@ -463,7 +463,7 @@ void UGeometrySelectionManager::ClearSelection()
 			TUniquePtr<FGeometrySelectionDeltaChange> ClearChange = MakeUnique<FGeometrySelectionDeltaChange>();
 			ClearChange->Identifier = Target->SelectionIdentifer;
 			ClearChange->Delta = MoveTemp(ClearDelta);
-			GetTransactionsAPI()->AppendChange(this, MoveTemp(ClearChange), LOCTEXT("ClearSelection", "Selection"));
+			GetTransactionsAPI()->AppendChange(this, MoveTemp(ClearChange), LOCTEXT("ClearSelection", "Clear Selection"));
 		}
 	}
 
