MeshModelingToolSet - DynamicMeshSelector.  Make sure the correct version of shutdown is called.  Fixes a static analysis warning

diff --git a/Engine/Plugins/Runtime/MeshModelingToolset/Source/ModelingComponents/Public/Selection/DynamicMeshSelector.h b/Engine/Plugins/Runtime/MeshModelingToolset/Source/ModelingComponents/Public/Selection/DynamicMeshSelector.h
--- a/Engine/Plugins/Runtime/MeshModelingToolset/Source/ModelingComponents/Public/Selection/DynamicMeshSelector.h
+++ b/Engine/Plugins/Runtime/MeshModelingToolset/Source/ModelingComponents/Public/Selection/DynamicMeshSelector.h
@@ -28,7 +28,7 @@ class MODELINGCOMPONENTS_API FDynamicMeshSelector : public IGeometrySelector
 public:
 	virtual ~FDynamicMeshSelector()
 	{
-		Shutdown();
+		FDynamicMeshSelector::Shutdown();
 	}
 
 	/**
