ModelingComponents: Add UE::Geometry::GetCurrentGeometrySelectionForTarget() util function to query current UGeometrySelectionManager for active Selection. Modify UE::Geometry::GetCurrentToolInputSelection() to return converted selection from UGeometrySelectionManager instead of using the now-disabled UPersistentMeshSelectionManager. Deprecate UPersistentMeshSelectionManager and related functions/classes.

diff --git a/Engine/Plugins/Runtime/MeshModelingToolset/Source/ModelingComponents/Private/Selection/PersistentMeshSelectionManager.cpp b/Engine/Plugins/Runtime/MeshModelingToolset/Source/ModelingComponents/Private/Selection/PersistentMeshSelectionManager.cpp
--- a/Engine/Plugins/Runtime/MeshModelingToolset/Source/ModelingComponents/Private/Selection/PersistentMeshSelectionManager.cpp
+++ b/Engine/Plugins/Runtime/MeshModelingToolset/Source/ModelingComponents/Private/Selection/PersistentMeshSelectionManager.cpp
@@ -10,14 +10,17 @@
 
 using namespace UE::Geometry;
 
-#define LOCTEXT_NAMESPACE "UPersistentMeshSelectionManager"
+#define LOCTEXT_NAMESPACE "UDEPRECATED_PersistentMeshSelectionManager"
 
-void UPersistentMeshSelectionManager::Initialize(TObjectPtr<UInteractiveToolsContext> ToolsContext)
+
+PRAGMA_DISABLE_DEPRECATION_WARNINGS
+
+void UDEPRECATED_PersistentMeshSelectionManager::Initialize(TObjectPtr<UInteractiveToolsContext> ToolsContext)
 {
 	ParentContext = ToolsContext;
 }
 
-void UPersistentMeshSelectionManager::Shutdown()
+void UDEPRECATED_PersistentMeshSelectionManager::Shutdown()
 {
 	if (SelectionDisplay != nullptr)
 	{
@@ -26,17 +29,17 @@ void UPersistentMeshSelectionManager::Shutdown()
 	}
 }
 
-bool UPersistentMeshSelectionManager::HasActiveSelection()
+bool UDEPRECATED_PersistentMeshSelectionManager::HasActiveSelection()
 {
 	return (ActiveSelection != nullptr);
 }
 
-UPersistentMeshSelection* UPersistentMeshSelectionManager::GetActiveSelection()
+UPersistentMeshSelection* UDEPRECATED_PersistentMeshSelectionManager::GetActiveSelection()
 {
 	return ActiveSelection;
 }
 
-void UPersistentMeshSelectionManager::SetNewActiveSelection(UPersistentMeshSelection* Selection)
+void UDEPRECATED_PersistentMeshSelectionManager::SetNewActiveSelection(UPersistentMeshSelection* Selection)
 {
 	TUniquePtr<FPersistentMeshSelectionChange> SelectionChange = MakeUnique<FPersistentMeshSelectionChange>();
 	if (ActiveSelection != nullptr)
@@ -55,14 +58,14 @@ void UPersistentMeshSelectionManager::SetNewActiveSelection(UPersistentMeshSelec
 }
 
 
-void UPersistentMeshSelectionManager::SetNewActiveSelectionInternal(UPersistentMeshSelection* Selection)
+void UDEPRECATED_PersistentMeshSelectionManager::SetNewActiveSelectionInternal(UPersistentMeshSelection* Selection)
 {
 	ActiveSelection = Selection;
 	OnSelectionModified();
 }
 
 
-void UPersistentMeshSelectionManager::ClearActiveSelection()
+void UDEPRECATED_PersistentMeshSelectionManager::ClearActiveSelection()
 {
 	if (ActiveSelection == nullptr)
 	{
@@ -75,7 +78,7 @@ void UPersistentMeshSelectionManager::ClearActiveSelection()
 
 
 
-void UPersistentMeshSelectionManager::OnSelectionModified()
+void UDEPRECATED_PersistentMeshSelectionManager::OnSelectionModified()
 {
 	if (ActiveSelection != nullptr && SelectionDisplay == nullptr)
 	{
@@ -126,10 +129,10 @@ bool UE::Geometry::RegisterPersistentMeshSelectionManager(UInteractiveToolsConte
 {
 	if (ensure(ToolsContext))
 	{
-		UPersistentMeshSelectionManager* Found = ToolsContext->ContextObjectStore->FindContext<UPersistentMeshSelectionManager>();
+		UDEPRECATED_PersistentMeshSelectionManager* Found = ToolsContext->ContextObjectStore->FindContext<UDEPRECATED_PersistentMeshSelectionManager>();
 		if (Found == nullptr)
 		{
-			UPersistentMeshSelectionManager* SelectionManager = NewObject<UPersistentMeshSelectionManager>(ToolsContext->ToolManager);
+			UDEPRECATED_PersistentMeshSelectionManager* SelectionManager = NewObject<UDEPRECATED_PersistentMeshSelectionManager>(ToolsContext->ToolManager);
 			if (ensure(SelectionManager))
 			{
 				SelectionManager->Initialize(ToolsContext);
@@ -152,7 +155,7 @@ bool UE::Geometry::DeregisterPersistentMeshSelectionManager(UInteractiveToolsCon
 {
 	if (ensure(ToolsContext))
 	{
-		UPersistentMeshSelectionManager* Found = ToolsContext->ContextObjectStore->FindContext<UPersistentMeshSelectionManager>();
+		UDEPRECATED_PersistentMeshSelectionManager* Found = ToolsContext->ContextObjectStore->FindContext<UDEPRECATED_PersistentMeshSelectionManager>();
 		if (Found != nullptr)
 		{
 			Found->Shutdown();
@@ -164,11 +167,11 @@ bool UE::Geometry::DeregisterPersistentMeshSelectionManager(UInteractiveToolsCon
 }
 
 
-UPersistentMeshSelectionManager* UE::Geometry::FindPersistentMeshSelectionManager(UInteractiveToolManager* ToolManager)
+UDEPRECATED_PersistentMeshSelectionManager* UE::Geometry::FindPersistentMeshSelectionManager(UInteractiveToolManager* ToolManager)
 {
 	if (ensure(ToolManager))
 	{
-		UPersistentMeshSelectionManager* Found = ToolManager->GetContextObjectStore()->FindContext<UPersistentMeshSelectionManager>();
+		UDEPRECATED_PersistentMeshSelectionManager* Found = ToolManager->GetContextObjectStore()->FindContext<UDEPRECATED_PersistentMeshSelectionManager>();
 		if (Found != nullptr)
 		{
 			return Found;
@@ -182,7 +185,7 @@ UPersistentMeshSelectionManager* UE::Geometry::FindPersistentMeshSelectionManage
 
 void FPersistentMeshSelectionChange::Apply(UObject* Object)
 {
-	UPersistentMeshSelectionManager* SelectionManager = Cast<UPersistentMeshSelectionManager>(Object);
+	UDEPRECATED_PersistentMeshSelectionManager* SelectionManager = Cast<UDEPRECATED_PersistentMeshSelectionManager>(Object);
 	if (SelectionManager)
 	{
 		UPersistentMeshSelection* NewSelection = NewObject<UPersistentMeshSelection>(SelectionManager);
@@ -193,7 +196,7 @@ void FPersistentMeshSelectionChange::Apply(UObject* Object)
 
 void FPersistentMeshSelectionChange::Revert(UObject* Object)
 {
-	UPersistentMeshSelectionManager* SelectionManager = Cast<UPersistentMeshSelectionManager>(Object);
+	UDEPRECATED_PersistentMeshSelectionManager* SelectionManager = Cast<UDEPRECATED_PersistentMeshSelectionManager>(Object);
 	if (SelectionManager)
 	{
 		UPersistentMeshSelection* NewSelection = NewObject<UPersistentMeshSelection>(SelectionManager);
@@ -212,4 +215,8 @@ FString FPersistentMeshSelectionChange::ToString() const
 	return TEXT("PersistentMeshSelectionChange");
 }
 
+
+
+PRAGMA_ENABLE_DEPRECATION_WARNINGS
+
 #undef LOCTEXT_NAMESPACE
\ No newline at end of file
diff --git a/Engine/Plugins/Runtime/MeshModelingToolset/Source/ModelingComponents/Private/Selection/StoredMeshSelectionUtil.cpp b/Engine/Plugins/Runtime/MeshModelingToolset/Source/ModelingComponents/Private/Selection/StoredMeshSelectionUtil.cpp
--- a/Engine/Plugins/Runtime/MeshModelingToolset/Source/ModelingComponents/Private/Selection/StoredMeshSelectionUtil.cpp
+++ b/Engine/Plugins/Runtime/MeshModelingToolset/Source/ModelingComponents/Private/Selection/StoredMeshSelectionUtil.cpp
@@ -6,31 +6,52 @@
 #include "TargetInterfaces/PrimitiveComponentBackedTarget.h"
 #include "InteractiveTool.h"
 #include "InteractiveToolManager.h"
+#include "ContextObjectStore.h"
 #include "GroupTopology.h"
+#include "Selection/GeometrySelectionManager.h"
 
 using namespace UE::Geometry;
 
 
-const UPersistentMeshSelection* UE::Geometry::GetCurrentToolInputSelection(const FToolBuilderState& SceneState, UToolTarget* Target)
+
+
+bool UE::Geometry::GetCurrentGeometrySelectionForTarget(const FToolBuilderState& SceneState, UToolTarget* Target, FGeometrySelection& SelectionOut)
 {
-	UPersistentMeshSelectionManager* SelectionManager = FindPersistentMeshSelectionManager(SceneState.ToolManager);
+	UGeometrySelectionManager* SelectionManager = SceneState.ToolManager->GetContextObjectStore()->FindContext<UGeometrySelectionManager>();
 	if (SelectionManager == nullptr)
 	{
-		return nullptr;
+		return false;
+	}
+	IPrimitiveComponentBackedTarget* TargetInterface = Cast<IPrimitiveComponentBackedTarget>(Target);
+	if (TargetInterface == nullptr)
+	{
+		return false;
 	}
 
-	UPersistentMeshSelection* ActiveSelection = SelectionManager->GetActiveSelection();
-	if (ActiveSelection == nullptr)
+	bool bHaveSelection = SelectionManager->GetSelectionForComponent(TargetInterface->GetOwnerComponent(), SelectionOut);
+	return bHaveSelection;
+}
+
+
+const UPersistentMeshSelection* UE::Geometry::GetCurrentToolInputSelection(const FToolBuilderState& SceneState, UToolTarget* Target)
+{
+	UGeometrySelectionManager* SelectionManager = SceneState.ToolManager->GetContextObjectStore()->FindContext<UGeometrySelectionManager>();
+	if (SelectionManager == nullptr)
 	{
 		return nullptr;
 	}
-
-	// todo: more validation here?
 	IPrimitiveComponentBackedTarget* TargetInterface = Cast<IPrimitiveComponentBackedTarget>(Target);
 	if (TargetInterface == nullptr)
 	{
 		return nullptr;
 	}
+
+	UPersistentMeshSelection* ActiveSelection = SelectionManager->GetActiveSingleSelectionConverted_Legacy(TargetInterface->GetOwnerComponent());
+	if (ActiveSelection == nullptr)
+	{
+		return nullptr;
+	}
+
 	if (TargetInterface->GetOwnerComponent() != ActiveSelection->GetTargetComponent())
 	{
 		return nullptr;
@@ -93,17 +114,20 @@ bool UE::Geometry::GetStoredSelectionAsTriangles(
 
 void UE::Geometry::ClearActiveToolSelection(UInteractiveToolManager* ToolManager)
 {
-	UPersistentMeshSelectionManager* SelectionManager = FindPersistentMeshSelectionManager(ToolManager);
+PRAGMA_DISABLE_DEPRECATION_WARNINGS
+	UDEPRECATED_PersistentMeshSelectionManager* SelectionManager = FindPersistentMeshSelectionManager(ToolManager);
 	if (SelectionManager != nullptr)
 	{
 		SelectionManager->ClearActiveSelection();
 	}
+PRAGMA_ENABLE_DEPRECATION_WARNINGS
 }
 
 
 bool UE::Geometry::SetToolOutputSelection(UInteractiveTool* Tool, UPersistentMeshSelection* Selection)
 {
-	UPersistentMeshSelectionManager* SelectionManager = FindPersistentMeshSelectionManager(Tool->GetToolManager());
+PRAGMA_DISABLE_DEPRECATION_WARNINGS
+	UDEPRECATED_PersistentMeshSelectionManager* SelectionManager = FindPersistentMeshSelectionManager(Tool->GetToolManager());
 	if (SelectionManager == nullptr)
 	{
 		return false;
@@ -119,6 +143,7 @@ bool UE::Geometry::SetToolOutputSelection(UInteractiveTool* Tool, UPersistentMes
 	}
 
 	return true;
+PRAGMA_ENABLE_DEPRECATION_WARNINGS
 }
 
 
diff --git a/Engine/Plugins/Runtime/MeshModelingToolset/Source/ModelingComponents/Public/Selection/PersistentMeshSelectionManager.h b/Engine/Plugins/Runtime/MeshModelingToolset/Source/ModelingComponents/Public/Selection/PersistentMeshSelectionManager.h
--- a/Engine/Plugins/Runtime/MeshModelingToolset/Source/ModelingComponents/Public/Selection/PersistentMeshSelectionManager.h
+++ b/Engine/Plugins/Runtime/MeshModelingToolset/Source/ModelingComponents/Public/Selection/PersistentMeshSelectionManager.h
@@ -26,8 +26,8 @@ class UPreviewGeometry;
  * those functions are preferable to directly using the Manager.
  * 
  */
-UCLASS()
-class MODELINGCOMPONENTS_API UPersistentMeshSelectionManager : public UObject
+UCLASS(Deprecated) 
+class MODELINGCOMPONENTS_API UDEPRECATED_PersistentMeshSelectionManager : public UObject
 {
 	GENERATED_BODY()
 
@@ -73,7 +73,9 @@ protected:
  * FPersistentMeshSelectionChange stores a selection change, ie before-and-after states.
  * UPersistentMeshSelectionManager emits these changes on selection modifications.
  */
-class MODELINGCOMPONENTS_API FPersistentMeshSelectionChange : public FToolCommandChange
+class MODELINGCOMPONENTS_API 
+UE_DEPRECATED(5.1, "UPersistentMeshSelectionManager has been deprecated and FPersistentMeshSelectionChange is no longer used") 
+FPersistentMeshSelectionChange : public FToolCommandChange
 {
 public:
 	FGenericMeshSelection From;
@@ -106,12 +108,14 @@ namespace Geometry
 	 * ToolsContext's ContextObjectStore
 	 * @return true if the ContextObjectStore now has a UPersistentMeshSelectionManager (whether it already existed, or was created)
 	 */
+	UE_DEPRECATED(5.1, "UE::Geometry::RegisterPersistentMeshSelectionManager is deprecated")
 	MODELINGCOMPONENTS_API bool RegisterPersistentMeshSelectionManager(UInteractiveToolsContext* ToolsContext);
 
 	/**
 	 * Remove any existing UPersistentMeshSelectionManager from the ToolsContext's ContextObjectStore
 	 * @return true if the ContextObjectStore no longer has a UPersistentMeshSelectionManager (whether it was removed, or did not exist)
 	 */
+	UE_DEPRECATED(5.1, "UE::Geometry::DeregisterPersistentMeshSelectionManager is deprecated")
 	MODELINGCOMPONENTS_API bool DeregisterPersistentMeshSelectionManager(UInteractiveToolsContext* ToolsContext);
 
 
@@ -119,7 +123,8 @@ namespace Geometry
 	 * Find an existing UPersistentMeshSelectionManager in the ToolsContext's ContextObjectStore
 	 * @return SelectionManager pointer or nullptr if not found
 	 */
-	MODELINGCOMPONENTS_API UPersistentMeshSelectionManager* FindPersistentMeshSelectionManager(UInteractiveToolManager* ToolManager);
+	UE_DEPRECATED(5.1, "UE::Geometry::FindPersistentMeshSelectionManager is deprecated")
+	MODELINGCOMPONENTS_API UDEPRECATED_PersistentMeshSelectionManager* FindPersistentMeshSelectionManager(UInteractiveToolManager* ToolManager);
 
 
 }
diff --git a/Engine/Plugins/Runtime/MeshModelingToolset/Source/ModelingComponents/Public/Selection/StoredMeshSelectionUtil.h b/Engine/Plugins/Runtime/MeshModelingToolset/Source/ModelingComponents/Public/Selection/StoredMeshSelectionUtil.h
--- a/Engine/Plugins/Runtime/MeshModelingToolset/Source/ModelingComponents/Public/Selection/StoredMeshSelectionUtil.h
+++ b/Engine/Plugins/Runtime/MeshModelingToolset/Source/ModelingComponents/Public/Selection/StoredMeshSelectionUtil.h
@@ -10,6 +10,7 @@ class UPersistentMeshSelection;
 class UToolTarget;
 class UInteractiveTool;
 class UInteractiveToolManager;
+struct FGeometrySelection;
 PREDECLARE_USE_GEOMETRY_CLASS(FGroupTopology);
 PREDECLARE_USE_GEOMETRY_CLASS(FCompactMaps);
 
@@ -27,6 +28,14 @@ namespace Geometry
 	 */
 	MODELINGCOMPONENTS_API const UPersistentMeshSelection* GetCurrentToolInputSelection(const FToolBuilderState& SceneState, UToolTarget* Target);
 
+	/**
+	* Given a FToolBuilderState, find geometry selection for the given ToolTarget
+	* @param SelectionOut non-empty selection will be returned here
+	* @return true if a non-empty selection was found
+	*/
+	MODELINGCOMPONENTS_API bool GetCurrentGeometrySelectionForTarget(const FToolBuilderState& SceneState, UToolTarget* Target,
+		FGeometrySelection& SelectionOut);
+
 	/**
 	 * Convert the given Seletion to a list of Triangles of the specified Mesh. 
 	 * This will attempt to convert the different selection types represented by UPersistentMeshSelection
