ModelingMode: add Delete command for current Mesh Selection

ToolsFramework:
add UInteractiveCommand and UInteractiveCommandArguments types, base classes for 'interactive commands' which are intended to be used for atomic actions, ie stateless, no user interaction, etc
ModelingComponents:
add UGeometrySelectionEditCommand/Arguments, this is a UInteractiveCommand that uses a FGeometrySelection as an argument
add UGeometrySelectionManager::CanExecuteSelectionCommand() and ExecuteSelectionCommand()

MeshModelingTools:
add UDeleteGeometrySelectionCommand, deletes selected mesh geometry by converting input selection to list of triangles

ModelingToolsEditorMode:
add Delete command to ModelingMode and UI. ModelingToolsEditorMode currently is keeping active command set alive via new UModelingToolsEditorMode::ModelingModeCommands member.

diff --git a/Engine/Plugins/Editor/ModelingToolsEditorMode/Source/ModelingToolsEditorMode/Private/ModelingToolsEditorMode.cpp b/Engine/Plugins/Editor/ModelingToolsEditorMode/Source/ModelingToolsEditorMode/Private/ModelingToolsEditorMode.cpp
--- a/Engine/Plugins/Editor/ModelingToolsEditorMode/Source/ModelingToolsEditorMode/Private/ModelingToolsEditorMode.cpp
+++ b/Engine/Plugins/Editor/ModelingToolsEditorMode/Source/ModelingToolsEditorMode/Private/ModelingToolsEditorMode.cpp
@@ -101,6 +101,10 @@
 #include "Tools/GenerateStaticMeshLODAssetTool.h"
 #include "Tools/LODManagerTool.h"
 
+// commands
+#include "Commands/DeleteGeometrySelectionCommand.h"
+
+
 #include "EditorModeManager.h"
 #include "UnrealWidget.h"
 
@@ -733,6 +737,22 @@ void UModelingToolsEditorMode::Enter()
 	}
 
 
+	// this function registers and tracks an active UGeometrySelectionEditCommand and it's associated UICommand
+	auto RegisterSelectionCommand = [&](UGeometrySelectionEditCommand* Command, TSharedPtr<FUICommandInfo> UICommand)
+	{
+		ModelingModeCommands.Add(Command);
+		const TSharedRef<FUICommandList>& CommandList = Toolkit->GetToolkitCommands();
+		CommandList->MapAction(UICommand,
+			FExecuteAction::CreateUObject(GetSelectionManager(), &UGeometrySelectionManager::ExecuteSelectionCommand, Command),
+			FCanExecuteAction::CreateUObject(GetSelectionManager(), &UGeometrySelectionManager::CanExecuteSelectionCommand, Command),
+			FIsActionChecked(),
+			EUIActionRepeatMode::RepeatDisabled);
+	};
+
+	// create and register InteractiveCommands for mesh selections
+	RegisterSelectionCommand(NewObject<UDeleteGeometrySelectionCommand>(), ToolManagerCommands.BeginSelectionAction_Delete);
+
+
 	// register extensions
 	TArray<IModelingModeToolExtension*> Extensions = IModularFeatures::Get().GetModularFeatureImplementations<IModelingModeToolExtension>(
 		IModelingModeToolExtension::GetModularFeatureName());
diff --git a/Engine/Plugins/Editor/ModelingToolsEditorMode/Source/ModelingToolsEditorMode/Private/ModelingToolsEditorModeToolkit.cpp b/Engine/Plugins/Editor/ModelingToolsEditorMode/Source/ModelingToolsEditorMode/Private/ModelingToolsEditorModeToolkit.cpp
--- a/Engine/Plugins/Editor/ModelingToolsEditorMode/Source/ModelingToolsEditorMode/Private/ModelingToolsEditorModeToolkit.cpp
+++ b/Engine/Plugins/Editor/ModelingToolsEditorMode/Source/ModelingToolsEditorMode/Private/ModelingToolsEditorModeToolkit.cpp
@@ -625,7 +625,7 @@ void FModelingToolsEditorModeToolkit::BuildToolPalette(FName PaletteIndex, class
 	}
 	else if (PaletteIndex == SelectionActionsTabName)
 	{
-		//ToolbarBuilder.AddToolBarButton(Commands.BeginSelectionAction_Delete);
+		ToolbarBuilder.AddToolBarButton(Commands.BeginSelectionAction_Delete);
 		//ToolbarBuilder.AddToolBarButton(Commands.BeginSelectionAction_Extrude);
 	}
 	else if (PaletteIndex == PrimitiveTabName)
diff --git a/Engine/Plugins/Editor/ModelingToolsEditorMode/Source/ModelingToolsEditorMode/Private/ModelingToolsManagerActions.cpp b/Engine/Plugins/Editor/ModelingToolsEditorMode/Source/ModelingToolsEditorMode/Private/ModelingToolsManagerActions.cpp
--- a/Engine/Plugins/Editor/ModelingToolsEditorMode/Source/ModelingToolsEditorMode/Private/ModelingToolsManagerActions.cpp
+++ b/Engine/Plugins/Editor/ModelingToolsEditorMode/Source/ModelingToolsEditorMode/Private/ModelingToolsManagerActions.cpp
@@ -182,6 +182,7 @@ void FModelingToolsManagerCommands::RegisterCommands()
 	REGISTER_MODELING_TOOL_COMMAND(BeginSelectionAction_ToEdgeType, "Edge", "Set Selection Type to Edges");
 	REGISTER_MODELING_TOOL_COMMAND(BeginSelectionAction_ToFaceType, "Face", "Set Selection Type to Faces");
 
+	REGISTER_MODELING_TOOL_COMMAND(BeginSelectionAction_Delete, "Delete", "Delete Selection");
 
 	UI_COMMAND(AcceptActiveTool, "Accept", "Accept the active tool", EUserInterfaceActionType::Button, FInputChord());
 	UI_COMMAND(CancelActiveTool, "Cancel", "Cancel the active tool", EUserInterfaceActionType::Button, FInputChord());
diff --git a/Engine/Plugins/Editor/ModelingToolsEditorMode/Source/ModelingToolsEditorMode/Public/ModelingToolsEditorMode.h b/Engine/Plugins/Editor/ModelingToolsEditorMode/Source/ModelingToolsEditorMode/Public/ModelingToolsEditorMode.h
--- a/Engine/Plugins/Editor/ModelingToolsEditorMode/Source/ModelingToolsEditorMode/Public/ModelingToolsEditorMode.h
+++ b/Engine/Plugins/Editor/ModelingToolsEditorMode/Source/ModelingToolsEditorMode/Public/ModelingToolsEditorMode.h
@@ -19,6 +19,7 @@ class FLevelObjectsObserver;
 class UModelingSceneSnappingManager;
 class UModelingSelectionInteraction;
 class UGeometrySelectionManager;
+class UInteractiveCommand;
 
 UCLASS(Transient)
 class UModelingToolsEditorMode : public UBaseLegacyWidgetEdMode
@@ -113,6 +114,14 @@ protected:
 	void ConfigureRealTimeViewportsOverride(bool bEnable);
 
 
+	// UInteractiveCommand support. Currently implemented by creating instances of
+	// commands on mode startup and holding onto them. This perhaps should be revisited,
+	// command instances could probably be created as needed...
+
+	UPROPERTY()
+	TArray<TObjectPtr<UInteractiveCommand>> ModelingModeCommands;
+
+
 	// analytics tracking
 	static FDateTime LastModeStartTimestamp;
 	static FDateTime LastToolStartTimestamp;
diff --git a/Engine/Plugins/Editor/ModelingToolsEditorMode/Source/ModelingToolsEditorMode/Public/ModelingToolsManagerActions.h b/Engine/Plugins/Editor/ModelingToolsEditorMode/Source/ModelingToolsEditorMode/Public/ModelingToolsManagerActions.h
--- a/Engine/Plugins/Editor/ModelingToolsEditorMode/Source/ModelingToolsEditorMode/Public/ModelingToolsManagerActions.h
+++ b/Engine/Plugins/Editor/ModelingToolsEditorMode/Source/ModelingToolsEditorMode/Public/ModelingToolsManagerActions.h
@@ -152,6 +152,7 @@ public:
 	TSharedPtr<FUICommandInfo> BeginSelectionAction_ToTriangleType;
 	TSharedPtr<FUICommandInfo> BeginSelectionAction_ToPolygroupType;
 
+	TSharedPtr<FUICommandInfo> BeginSelectionAction_Delete;
 
 	//
 	// Accept/Cancel/Complete commands are used to end the active Tool via ToolManager
diff --git a/Engine/Plugins/Runtime/MeshModelingToolset/Source/MeshModelingTools/Private/Commands/DeleteGeometrySelectionCommand.cpp b/Engine/Plugins/Runtime/MeshModelingToolset/Source/MeshModelingTools/Private/Commands/DeleteGeometrySelectionCommand.cpp
new file mode 100644
--- /dev/null
+++ b/Engine/Plugins/Runtime/MeshModelingToolset/Source/MeshModelingTools/Private/Commands/DeleteGeometrySelectionCommand.cpp
@@ -0,0 +1,89 @@
+// Copyright Epic Games, Inc. All Rights Reserved.
+
+#include "Commands/DeleteGeometrySelectionCommand.h"
+#include "ToolContextInterfaces.h"
+#include "UDynamicMesh.h"
+#include "DynamicMesh/DynamicMesh3.h"
+#include "DynamicMesh/DynamicMeshChangeTracker.h"
+#include "DynamicMeshEditor.h"
+#include "Changes/MeshChange.h"
+#include "Selection/GeometrySelectionUtil.h"
+
+using namespace UE::Geometry;
+
+
+#define LOCTEXT_NAMESPACE "UDeleteGeometrySelectionCommand"
+
+
+FText UDeleteGeometrySelectionCommand::GetCommandShortString() const
+{
+	return LOCTEXT("ShortString", "Delete Selection");
+}
+
+
+bool UDeleteGeometrySelectionCommand::CanExecuteCommandForSelection(UGeometrySelectionEditCommandArguments* SelectionArgs)
+{
+	return SelectionArgs->IsMatchingType(FGeometryIdentifier::ETargetType::MeshContainer, FGeometryIdentifier::EObjectType::DynamicMesh);
+}
+
+
+void UDeleteGeometrySelectionCommand::ExecuteCommandForSelection(UGeometrySelectionEditCommandArguments* SelectionArgs)
+{
+	// should have been verified by CanExecute
+	check(SelectionArgs->IsMatchingType(FGeometryIdentifier::ETargetType::MeshContainer, FGeometryIdentifier::EObjectType::DynamicMesh));
+	
+	// collect up all our inputs
+	UDynamicMesh* MeshObject = SelectionArgs->SelectionHandle.Identifier.GetAsObjectType<UDynamicMesh>();
+	check(MeshObject != nullptr);
+	const FGeometrySelection* Selection = SelectionArgs->SelectionHandle.Selection;
+	if (Selection->Selection.Num() == 0)
+	{
+		return;
+	}
+
+	bool bTrackChanges = SelectionArgs->HasTransactionsAPI();
+	TUniquePtr<FDynamicMeshChange> DynamicMeshChange;	// only initialized if bTrackChanges == true
+
+	// apply the delete operation
+	MeshObject->EditMesh([&](FDynamicMesh3& EditMesh)
+	{
+		// build list of triangles from whatever the selection contains
+		TSet<int32> TriangleList;
+		//UE::Geometry::FPolygroupSet UsePolygroupSet = ...;		// need to support this eventually
+		UE::Geometry::EnumerateSelectionTriangles(*Selection, EditMesh,
+			[&](int32 TriangleID) { TriangleList.Add(TriangleID); });
+
+		// mark triangles for change
+		FDynamicMeshChangeTracker ChangeTracker(&EditMesh);
+		if (bTrackChanges)
+		{
+			ChangeTracker.BeginChange();
+			ChangeTracker.SaveTriangles(TriangleList, true);
+		}
+
+		// actually delete them
+		FDynamicMeshEditor Editor(&EditMesh);
+		Editor.RemoveTriangles(TriangleList.Array(), true);
+
+		// extract the change record
+		if (bTrackChanges)
+		{
+			DynamicMeshChange = ChangeTracker.EndChange();
+		}
+
+	}, EDynamicMeshChangeType::GeneralEdit, EDynamicMeshAttributeChangeFlags::Unknown, false);
+
+
+	// emit change 
+	if ( bTrackChanges && DynamicMeshChange.IsValid() )
+	{
+		SelectionArgs->GetTransactionsAPI()->BeginUndoTransaction(GetCommandShortString());
+		SelectionArgs->GetTransactionsAPI()->AppendChange(MeshObject, 
+			MakeUnique<FMeshChange>(MoveTemp(DynamicMeshChange)), GetCommandShortString());
+		SelectionArgs->GetTransactionsAPI()->EndUndoTransaction();
+	}
+
+}
+
+
+#undef LOCTEXT_NAMESPACE
diff --git a/Engine/Plugins/Runtime/MeshModelingToolset/Source/MeshModelingTools/Public/Commands/DeleteGeometrySelectionCommand.h b/Engine/Plugins/Runtime/MeshModelingToolset/Source/MeshModelingTools/Public/Commands/DeleteGeometrySelectionCommand.h
new file mode 100644
--- /dev/null
+++ b/Engine/Plugins/Runtime/MeshModelingToolset/Source/MeshModelingTools/Public/Commands/DeleteGeometrySelectionCommand.h
@@ -0,0 +1,24 @@
+// Copyright Epic Games, Inc. All Rights Reserved.
+
+#pragma once
+
+#include "Selection/SelectionEditInteractiveCommand.h"
+#include "DeleteGeometrySelectionCommand.generated.h"
+
+
+/**
+ * UDeleteGeometrySelectionCommand deletes the geometric elements identified by the Selection.
+ * Currently only supports mesh selections (Triangle and Polygroup types)
+ * Deletes selected faces, or faces connected to selected edges, or faces connected to selected vertices.
+ */
+UCLASS()
+class MESHMODELINGTOOLS_API UDeleteGeometrySelectionCommand : public UGeometrySelectionEditCommand
+{
+	GENERATED_BODY()
+public:
+
+	virtual FText GetCommandShortString() const override;
+
+	virtual bool CanExecuteCommandForSelection(UGeometrySelectionEditCommandArguments* Arguments) override;
+	virtual void ExecuteCommandForSelection(UGeometrySelectionEditCommandArguments* Arguments) override;
+};
\ No newline at end of file
diff --git a/Engine/Plugins/Runtime/MeshModelingToolset/Source/ModelingComponents/Private/Selection/GeometrySelectionManager.cpp b/Engine/Plugins/Runtime/MeshModelingToolset/Source/ModelingComponents/Private/Selection/GeometrySelectionManager.cpp
--- a/Engine/Plugins/Runtime/MeshModelingToolset/Source/ModelingComponents/Private/Selection/GeometrySelectionManager.cpp
+++ b/Engine/Plugins/Runtime/MeshModelingToolset/Source/ModelingComponents/Private/Selection/GeometrySelectionManager.cpp
@@ -2,6 +2,7 @@
 
 #include "Selection/GeometrySelectionManager.h"
 #include "Selection/DynamicMeshSelector.h"
+#include "Selection/SelectionEditInteractiveCommand.h"
 #include "ToolContextInterfaces.h"
 #include "ToolDataVisualizer.h"
 
@@ -653,6 +654,99 @@ void UGeometrySelectionManager::EndTransformation()
 
 
 
+
+bool UGeometrySelectionManager::CanExecuteSelectionCommand(UGeometrySelectionEditCommand* Command)
+{
+	if (SelectionArguments == nullptr)
+	{
+		SelectionArguments = NewObject<UGeometrySelectionEditCommandArguments>();
+	}
+
+	bool bCanExecute = true;
+	bool bHaveSelections = false;
+	ProcessActiveSelections([&](FGeometrySelectionHandle Handle)
+	{
+		SelectionArguments->SelectionHandle = Handle;
+		SelectionArguments->SetTransactionsAPI(TransactionsAPI);
+		bCanExecute = bCanExecute && Command->CanExecuteCommand(SelectionArguments);
+		bHaveSelections = true;
+	});
+
+	return bHaveSelections && bCanExecute;
+}
+
+void UGeometrySelectionManager::ExecuteSelectionCommand(UGeometrySelectionEditCommand* Command)
+{
+	if (SelectionArguments == nullptr)
+	{
+		SelectionArguments = NewObject<UGeometrySelectionEditCommandArguments>();
+	}
+
+	// open transaction to wrap the entire set of Commands and selection changes
+	FText CommandText = Command->GetCommandShortString();
+	GetTransactionsAPI()->BeginUndoTransaction(CommandText);
+
+	for (TSharedPtr<FGeometrySelectionTarget> Target : ActiveTargetReferences)
+	{
+		if (Target->Selection.IsEmpty()) continue;
+
+		// When initially executing the command, we do not clear the selection, because we pass it to the command.
+		// However, when we later *undo* any changes emitted by the command, we need to restore the selection aftewards.
+		// So we emit a clearing change here, so that undo un-clears. 
+		// When we later Redo, it is also necessary to Clear as otherwise an invalid Selection might hang around.
+		// Note that this must happen *before* the Command. The Command will not be re-executed, only its emitted Changes,
+		// so it will not be holding onto the active Selection on Redo later
+		// (if that becomes necessary, this sequence of changes will need to become more complicated....)
+		TUniquePtr<FGeometrySelectionReplaceChange> ClearChange = MakeUnique<FGeometrySelectionReplaceChange>();
+		ClearChange->Identifier = Target->Selector->GetIdentifier();
+		ClearChange->Before = Target->Selection;
+		ClearChange->After.InitializeTypes(ClearChange->Before);
+		GetTransactionsAPI()->AppendChange(this, MoveTemp(ClearChange), LOCTEXT("ClearSelection", "Clear Selection"));
+
+		// q: we could clear the selection here, and pass the Handle a copy. Perhaps safer?
+
+		SelectionArguments->SelectionHandle = FGeometrySelectionHandle{ Target->Selector->GetIdentifier(), &Target->Selection };
+		SelectionArguments->SetTransactionsAPI(TransactionsAPI);
+		Command->ExecuteCommand(SelectionArguments);
+
+		// clear selection after executing command
+		FGeometrySelectionDelta ClearDelta;
+		Target->SelectionEditor->ClearSelection(ClearDelta);
+
+		// 
+		// if commands could emit new selections, this is where we would set them up and emit changes!
+	}
+
+	GetTransactionsAPI()->EndUndoTransaction();
+
+	// assume this is true for now
+	bSelectionRenderCachesDirty = true;
+	OnSelectionModified.Broadcast();
+}
+
+
+void UGeometrySelectionManager::ProcessActiveSelections(TFunctionRef<void(FGeometrySelectionHandle)> ProcessFunc)
+{
+	for (TSharedPtr<FGeometrySelectionTarget> Target : ActiveTargetReferences)
+	{
+		if (Target->Selection.IsEmpty() == false)
+		{
+			FGeometrySelectionHandle Handle;
+			Handle.Selection = & Target->Selection;
+			Handle.Identifier = Target->Selector->GetIdentifier();
+			ProcessFunc(Handle);
+		}
+	}
+}
+
+
+
+
+
+
+
+
+
 void UGeometrySelectionManager::ApplyChange(IGeometrySelectionChange* Change)
 {
 	// We should not get here because selection changes should have been expired.
diff --git a/Engine/Plugins/Runtime/MeshModelingToolset/Source/ModelingComponents/Public/Selection/GeometrySelectionManager.h b/Engine/Plugins/Runtime/MeshModelingToolset/Source/ModelingComponents/Public/Selection/GeometrySelectionManager.h
--- a/Engine/Plugins/Runtime/MeshModelingToolset/Source/ModelingComponents/Public/Selection/GeometrySelectionManager.h
+++ b/Engine/Plugins/Runtime/MeshModelingToolset/Source/ModelingComponents/Public/Selection/GeometrySelectionManager.h
@@ -14,6 +14,8 @@ class UInteractiveToolsContext;
 class IToolsContextRenderAPI;
 class IToolsContextTransactionsAPI;
 class UPersistentMeshSelection;
+class UGeometrySelectionEditCommand;
+class UGeometrySelectionEditCommandArguments;
 
 
 /**
@@ -220,6 +222,32 @@ public:
 	virtual void EndTransformation();
 
 
+	//
+	// Command Execution
+	//
+public:
+	/**
+	 * @return true if Command->CanExecuteCommand() returns true for *all* the current Selections
+	 */
+	bool CanExecuteSelectionCommand(UGeometrySelectionEditCommand* Command);
+
+	/**
+	 * execute the selection command for *all* the current selections
+	 */
+	void ExecuteSelectionCommand(UGeometrySelectionEditCommand* Command);
+
+
+protected:
+	// This is set to current selection during CanExecuteSelectionCommand/ExecuteSelectionCommand, to keep the UObject alive
+	// Not expected to be used outside that context
+	UPROPERTY()
+	TObjectPtr<UGeometrySelectionEditCommandArguments> SelectionArguments;
+
+	// apply ProcessFunc to active selections via handles, perhaps should be public?
+	void ProcessActiveSelections(TFunctionRef<void(FGeometrySelectionHandle)> ProcessFunc);
+
+
+
 	//
 	// Undo/Redo
 	//
diff --git a/Engine/Plugins/Runtime/MeshModelingToolset/Source/ModelingComponents/Public/Selection/SelectionEditInteractiveCommand.h b/Engine/Plugins/Runtime/MeshModelingToolset/Source/ModelingComponents/Public/Selection/SelectionEditInteractiveCommand.h
new file mode 100644
--- /dev/null
+++ b/Engine/Plugins/Runtime/MeshModelingToolset/Source/ModelingComponents/Public/Selection/SelectionEditInteractiveCommand.h
@@ -0,0 +1,78 @@
+// Copyright Epic Games, Inc. All Rights Reserved.
+
+#pragma once
+
+#include "InteractiveCommand.h"
+#include "Selection/GeometrySelector.h"    // for FGeometrySelectionHandle
+#include "SelectionEditInteractiveCommand.generated.h"
+
+/**
+ * Arguments for a UGeometrySelectionEditCommand
+ */
+UCLASS()
+class MODELINGCOMPONENTS_API UGeometrySelectionEditCommandArguments : public UInteractiveCommandArguments
+{
+	GENERATED_BODY()
+public:
+	FGeometrySelectionHandle SelectionHandle;
+
+	bool IsSelectionEmpty() const
+	{
+		return SelectionHandle.Selection == nullptr || SelectionHandle.Selection->IsEmpty();
+	}
+
+	bool IsMatchingType(FGeometryIdentifier::ETargetType TargetType, FGeometryIdentifier::EObjectType EngineType) const
+	{
+		return SelectionHandle.Identifier.TargetType == TargetType
+			|| SelectionHandle.Identifier.ObjectType == EngineType;
+	}
+};
+
+
+
+/**
+ * UGeometrySelectionEditCommand is a command that edits geometry based on a selection.
+ * Requires a UGeometrySelectionEditCommandArguments
+ */
+UCLASS(Abstract)
+class MODELINGCOMPONENTS_API UGeometrySelectionEditCommand : public UInteractiveCommand
+{
+	GENERATED_BODY()
+public:
+	virtual bool AllowEmptySelection() const { return false; }
+
+	virtual bool CanExecuteCommandForSelection(UGeometrySelectionEditCommandArguments* SelectionArgs)
+	{
+		return false;
+	}
+	
+	virtual void ExecuteCommandForSelection(UGeometrySelectionEditCommandArguments* Arguments)
+	{
+	}
+
+
+public:
+
+	// UInteractiveCommand API
+
+	virtual bool CanExecuteCommand(UInteractiveCommandArguments* Arguments) override
+	{
+		UGeometrySelectionEditCommandArguments* SelectionArgs = Cast<UGeometrySelectionEditCommandArguments>(Arguments);
+		if  (SelectionArgs && (SelectionArgs->IsSelectionEmpty() == false || AllowEmptySelection()) )
+		{
+			return CanExecuteCommandForSelection(SelectionArgs);
+		}
+		return false;
+	}
+
+	virtual void ExecuteCommand(UInteractiveCommandArguments* Arguments) override
+	{
+		UGeometrySelectionEditCommandArguments* SelectionArgs = Cast<UGeometrySelectionEditCommandArguments>(Arguments);
+		if ( SelectionArgs && (SelectionArgs->IsSelectionEmpty() == false || AllowEmptySelection()) )
+		{
+			ExecuteCommandForSelection(SelectionArgs);
+		}
+	}
+
+
+};
diff --git a/Engine/Source/Runtime/InteractiveToolsFramework/Public/InteractiveCommand.h b/Engine/Source/Runtime/InteractiveToolsFramework/Public/InteractiveCommand.h
new file mode 100644
--- /dev/null
+++ b/Engine/Source/Runtime/InteractiveToolsFramework/Public/InteractiveCommand.h
@@ -0,0 +1,68 @@
+// Copyright Epic Games, Inc. All Rights Reserved.
+
+#pragma once
+
+#include "ToolContextInterfaces.h"
+#include "InteractiveCommand.generated.h"
+
+class IToolsContextTransactionsAPI;
+
+/**
+ * UInteractiveCommandArguments are arguments passed to a UInteractiveCommand.
+ * Subclasses of UInteractiveCommand will generally be paired with subclasses of UInteractiveCommandArguments.
+ * 
+ * The base UInteractiveCommandArguments provides support for transactions via 
+ * a IToolsContextTransactionsAPI
+ */
+UCLASS()
+class INTERACTIVETOOLSFRAMEWORK_API UInteractiveCommandArguments : public UObject
+{
+	GENERATED_BODY()
+
+public:
+	virtual void SetTransactionsAPI(IToolsContextTransactionsAPI* TransactionsAPIIn) { TransactionsAPI = TransactionsAPIIn;	}
+	virtual bool HasTransactionsAPI() const { return TransactionsAPI != nullptr; }
+	virtual IToolsContextTransactionsAPI* GetTransactionsAPI() const { return TransactionsAPI; }
+
+protected:
+	IToolsContextTransactionsAPI* TransactionsAPI = nullptr;
+
+};
+
+
+/**
+ * A UInteractiveCommand is an atomic action that can be executed via some user interaction.
+ * For example clicking a button that deletes an active selection can be considered an Interactive Command.
+ * This differs from an Interactive Tool in that there is no ongoing user interaction once the
+ * command has been initiated. 
+ */
+UCLASS(Abstract)
+class INTERACTIVETOOLSFRAMEWORK_API UInteractiveCommand : public UObject
+{
+	GENERATED_BODY()
+
+public:
+
+	/**
+	 * @return a short text string that can be used for the command in (eg) Editor transaction/undo toasts
+	 */
+	virtual FText GetCommandShortString() const 
+	{ 
+		return FText(); 
+	}
+
+	/**
+	 * @return true if it is safe to call ExecuteCommand() with the given Arguments
+	 */
+	virtual bool CanExecuteCommand(UInteractiveCommandArguments* Arguments) 
+	{ 
+		return false;
+	}
+
+	/**
+	 * Execute the command with the given Arguments
+	 */
+	virtual void ExecuteCommand(UInteractiveCommandArguments* Arguments)
+	{
+	}
+};
\ No newline at end of file
