add file that was missed in CL 21117964

diff --git a/Engine/Plugins/Runtime/MeshModelingToolset/Source/ModelingComponents/Public/Selection/StoredMeshSelectionUtil.h b/Engine/Plugins/Runtime/MeshModelingToolset/Source/ModelingComponents/Public/Selection/StoredMeshSelectionUtil.h
--- a/Engine/Plugins/Runtime/MeshModelingToolset/Source/ModelingComponents/Public/Selection/StoredMeshSelectionUtil.h
+++ b/Engine/Plugins/Runtime/MeshModelingToolset/Source/ModelingComponents/Public/Selection/StoredMeshSelectionUtil.h
@@ -4,13 +4,13 @@
 
 #include "CoreMinimal.h"
 #include "DynamicMesh/DynamicMesh3.h"
+#include "Selections/GeometrySelection.h"
 
 struct FToolBuilderState;
 class UPersistentMeshSelection;
 class UToolTarget;
 class UInteractiveTool;
 class UInteractiveToolManager;
-struct FGeometrySelection;
 PREDECLARE_USE_GEOMETRY_CLASS(FGroupTopology);
 PREDECLARE_USE_GEOMETRY_CLASS(FCompactMaps);
 
