Render Capture Baking: Fix requirement that the tool is run more than once in order to get the correct result by caching the needed material shaders before rendering

This fixes a problem that occurs in both the Render Capture Baking tool and the Approximate Actors tool

#rnx

diff --git a/Engine/Plugins/Runtime/MeshModelingToolset/Source/ModelingComponents/Private/Scene/WorldRenderCapture.cpp b/Engine/Plugins/Runtime/MeshModelingToolset/Source/ModelingComponents/Private/Scene/WorldRenderCapture.cpp
--- a/Engine/Plugins/Runtime/MeshModelingToolset/Source/ModelingComponents/Private/Scene/WorldRenderCapture.cpp
+++ b/Engine/Plugins/Runtime/MeshModelingToolset/Source/ModelingComponents/Private/Scene/WorldRenderCapture.cpp
@@ -5,6 +5,8 @@
 
 #include "Engine/Canvas.h"
 #include "Engine/TextureRenderTarget2D.h"
+#include "Materials/Material.h"
+#include "BufferVisualizationData.h"
 #include "GameFramework/WorldSettings.h"
 #include "LegacyScreenPercentageDriver.h"
 #include "EngineModule.h"
@@ -106,8 +108,28 @@ FRenderCaptureConfig GetDefaultRenderCaptureConfig(ERenderCaptureType CaptureTyp
 	return Config;
 }
 
+} // end namespace Geometry
+
+namespace Internal
+{
+
+void FlushRHIThreadToUpdateTextureRenderTargetReference()
+{
+	// Flush RHI thread after creating texture render target to make sure that RHIUpdateTextureReference is executed
+	// before doing any rendering with it This makes sure that
+	//    Value->TextureReference.TextureReferenceRHI->GetReferencedTexture()
+	// is valid so that
+	//    UniformExpressionSet::FillUniformBuffer
+	// properly uses the texture for rendering, instead of using a fallback texture
+	ENQUEUE_RENDER_COMMAND(FlushRHIThreadToUpdateTextureRenderTargetReference)(
+	[](FRHICommandListImmediate& RHICmdList)
+	{
+		RHICmdList.ImmediateFlush(EImmediateFlushType::FlushRHIThread);
+	});
 }
-}
+
+} // end namespace Internal
+} // end namespace UE
 
 
 
@@ -246,6 +268,10 @@ UTextureRenderTarget2D* FWorldRenderCapture::GetRenderTexture(bool bLinear)
 
 	if ( *WhichTexture != nullptr )
 	{
+		(*WhichTexture)->UpdateResourceImmediate(true);
+
+		UE::Internal::FlushRHIThreadToUpdateTextureRenderTargetReference();
+
 		return *WhichTexture;
 	}
 
@@ -260,6 +286,9 @@ UTextureRenderTarget2D* FWorldRenderCapture::GetRenderTexture(bool bLinear)
 		(*WhichTexture)->ClearColor = FLinearColor::Transparent;
 		(*WhichTexture)->TargetGamma = (bLinear) ? 1.0f : 2.2f;
 		(*WhichTexture)->InitCustomFormat(Width, Height, PF_FloatRGBA, false);
+		(*WhichTexture)->UpdateResourceImmediate(true);
+		
+		UE::Internal::FlushRHIThreadToUpdateTextureRenderTargetReference();
 	}
 
 	return *WhichTexture;
@@ -295,17 +324,7 @@ UTextureRenderTarget2D* FWorldRenderCapture::GetDepthRenderTexture()
 		DepthRenderTexture->InitAutoFormat(Width, Height);
 		DepthRenderTexture->UpdateResourceImmediate(true);
 
-		// Flush RHI thread after creating texture render target to make sure that RHIUpdateTextureReference is executed
-		// before doing any rendering with it This makes sure that
-		//    Value->TextureReference.TextureReferenceRHI->GetReferencedTexture()
-		// is valid so that
-		//    UniformExpressionSet::FillUniformBuffer
-		// properly uses the texture for rendering, instead of using a fallback texture
-		ENQUEUE_RENDER_COMMAND(FlushRHIThreadToUpdateTextureRenderTargetReference)(
-		[](FRHICommandListImmediate& RHICmdList)
-		{
-			RHICmdList.ImmediateFlush(EImmediateFlushType::FlushRHIThread);
-		});
+		UE::Internal::FlushRHIThreadToUpdateTextureRenderTargetReference();
 	}
 
 	return DepthRenderTexture;
@@ -403,6 +422,37 @@ static bool ReadPixelDataToImage(TUniquePtr<FImagePixelData>& PixelData, FImageA
 	}
 }
 
+// This function is used to cache the material shaders used for world render capture (mostly buffer visualization
+// materials, but there are also other ones). If we don't do this, and if the shaders also haven't been otherwise
+// compiled (e.g., by switching to the needed buffer visualization mode in the viewport), then the render capture will
+// use a fallback material the first time around and give incorrect results (which caused #jira UE-146097). Note that
+// although buffer visualization materials are initialized during engine/editor startup, we do not cache needed shaders
+// at that point since doing so only compiles them for the default feature level, GMaxRHIFeatureLevel. The user may
+// change the preview feature level after editor startup, or may have done so in a previous session (the setting is
+// serialized/reused in subsequent sessions via Engine/Saved/Config/WindowsEditor/EditorPerProjectUserSettings.ini),
+// so in order to make the render capture code work in that case we do this caching right here. The CacheShaders
+// function compiles all feature levels returned by GetFeatureLevelsToCompileForAllMaterials (note that
+// SetFeatureLevelToCompile is called when the feature level is switched).
+void CacheShadersForMaterial(UMaterialInterface* MaterialInterface)
+{
+	if (ensure(MaterialInterface))
+	{
+		UMaterial* Material = MaterialInterface->GetMaterial();
+		if (ensure(Material))
+		{
+			// Mark the material as a special engine material so that FMaterial::IsRequiredComplete returns true and we
+			// compile the material, we restore the previous value to minimize surprise for other code
+			bool OldState = Material->bUsedAsSpecialEngineMaterial;
+			Material->bUsedAsSpecialEngineMaterial = true;
+			Material->CacheShaders(EMaterialShaderPrecompileMode::Synchronous);
+			ensureMsgf(Material->IsComplete(),
+				TEXT("Caching shaders for material %s did not complete"),
+				*MaterialInterface->GetName());
+			Material->bUsedAsSpecialEngineMaterial = OldState;
+		}
+	}
+}
+
 } // end namespace Internal
 } // end namespace UE
 
@@ -547,6 +597,7 @@ bool FWorldRenderCapture::CaptureMRSFromPosition(
 	FCanvas Canvas = FCanvas(RenderTargetResource, nullptr, this->World, ERHIFeatureLevel::SM5, FCanvas::CDM_DeferDrawing, 1.0f);
 	Canvas.Clear(FLinearColor::Transparent);
 
+	UE::Internal::CacheShadersForMaterial(PostProcessMaterial);
 	UE::Internal::PerformSceneRender(&Canvas, &ViewFamily);
 
 	// Cache the view/projection matricies we used to render the scene
@@ -662,6 +713,8 @@ bool FWorldRenderCapture::CaptureEmissiveFromPosition(
 	FCanvas Canvas(RenderTargetResource, NULL, FGameTime::GetTimeSinceAppStart(), World->Scene->GetFeatureLevel());
 	Canvas.Clear(FLinearColor::Transparent);
 
+	UMaterialInterface* MaterialInterface = GetBufferVisualizationData().GetMaterial(NewView->CurrentBufferVisualizationMode);
+	UE::Internal::CacheShadersForMaterial(MaterialInterface);
 	UE::Internal::PerformSceneRender(&Canvas, &ViewFamily);
 
 	// Cache the view/projection matricies we used to render the scene
@@ -781,6 +834,7 @@ bool FWorldRenderCapture::CaptureDeviceDepthFromPosition(
 	FCanvas Canvas(RenderTargetResource, NULL, FGameTime::GetTimeSinceAppStart(), World->Scene->GetFeatureLevel());
 	Canvas.Clear(FLinearColor::Transparent);
 
+	// Unlike other capture types, we don't need to cache any shaders before we capture the device depth render
 	UE::Internal::PerformSceneRender(&Canvas, &ViewFamily);
 
 	// Cache the view/projection matricies we used to render the scene
@@ -910,6 +964,8 @@ namespace Internal
 		FCanvas Canvas(RenderTargetResource, NULL, FGameTime::GetTimeSinceAppStart(), Scene->GetFeatureLevel());
 		Canvas.Clear(FLinearColor::Transparent);
 
+		UMaterialInterface* MaterialInterface = GetBufferVisualizationData().GetMaterial(NewView->CurrentBufferVisualizationMode);
+		UE::Internal::CacheShadersForMaterial(MaterialInterface);
 		UE::Internal::PerformSceneRender(&Canvas, &ViewFamily);
 
 		// Cache the view/projection matricies we used to render the scene
