MeshModelingToolset: Optimize FMeshWireframeSceneProxy constructor

diff --git a/Engine/Plugins/Runtime/MeshModelingToolset/Source/ModelingComponents/Private/Drawing/MeshWireframeComponent.cpp b/Engine/Plugins/Runtime/MeshModelingToolset/Source/ModelingComponents/Private/Drawing/MeshWireframeComponent.cpp
--- a/Engine/Plugins/Runtime/MeshModelingToolset/Source/ModelingComponents/Private/Drawing/MeshWireframeComponent.cpp
+++ b/Engine/Plugins/Runtime/MeshModelingToolset/Source/ModelingComponents/Private/Drawing/MeshWireframeComponent.cpp
@@ -43,57 +43,74 @@ public:
 		MaterialRelevance(Component->GetMaterialRelevance(GetScene().GetFeatureLevel())),
 		VertexFactory(GetScene().GetFeatureLevel(), "FPointSetSceneProxy")
 	{
+		TRACE_CPUPROFILER_EVENT_SCOPE(FMeshWireframeSceneProxy);
+
 		if (!ensure(WireSource && WireSource->IsValid())) return;
 
 		if (WireSource->GetEdgeCount() <= 0) return;
 
+		// Create local copies to avoid repeated indirect look-up through the Component pointer in loops. 
+		const bool bEnableWireframe = Component->bEnableWireframe;
+		const bool bEnableBoundaryEdges = Component->bEnableBoundaryEdges;
+		const bool bEnableUVSeams = Component->bEnableUVSeams;
+		const bool bEnableNormalSeams = Component->bEnableNormalSeams;
+		const bool bEnableColorSeams = Component->bEnableColorSeams;
+
+		// Setup edge batches to process in parallel.
+		const int32 NumSourceEdges = WireSource->GetMaxEdgeIndex();
+		constexpr int32 BatchSize = 16384;
+		const int32 NumBatches = FMath::DivideAndRoundUp<int32>(NumSourceEdges, BatchSize);
+		
+		struct EdgeBatch
+		{
+			int32 Offset;
+			TArray<UE::Geometry::FIndex3i> Edges;
+		};
+		
+		TArray<EdgeBatch> VisibleEdgeBatches;
+		VisibleEdgeBatches.SetNum(NumBatches);
+
 		// count visible edges and remap so we can build in parallel below
-		int32 MaxEdgeIndex = WireSource->GetMaxEdgeIndex();
-		CurrentEdgeSet.Reserve(WireSource->GetEdgeCount());
-		for (int32 li = 0; li < MaxEdgeIndex; ++li)
+		int32 NumVisibleEdges = 0;
+		ParallelFor(NumBatches, [&](int32 BatchIndex)
 		{
-			if (WireSource->IsEdge(li) == false) continue;
+			TArray<UE::Geometry::FIndex3i>& VisibleEdgeBatch = VisibleEdgeBatches[BatchIndex].Edges;
+			const int32 CurrentBatchSize = (BatchIndex < NumBatches - 1) || (NumSourceEdges % BatchSize == 0) ? BatchSize : NumSourceEdges % BatchSize;
+			VisibleEdgeBatch.Reserve(CurrentBatchSize);
+			
+			for (int32 EdgeIndex = BatchIndex * BatchSize, BatchEnd = BatchIndex * BatchSize + CurrentBatchSize; EdgeIndex < BatchEnd; ++EdgeIndex)
+			{
+				if (WireSource->IsEdge(EdgeIndex) == false) continue;
 
-			int32 VertIndexA, VertIndexB;
-			IMeshWireframeSource::EMeshEdgeType EdgeType;
-			WireSource->GetEdge(li, VertIndexA, VertIndexB, EdgeType);
+				int32 VertIndexA, VertIndexB;
+				IMeshWireframeSource::EMeshEdgeType EdgeType;
+				WireSource->GetEdge(EdgeIndex, VertIndexA, VertIndexB, EdgeType);
 
-			bool bEdgeIsVisible = false;
-			if (Component->bEnableWireframe)
-			{
-				bEdgeIsVisible = true;
-			}
-			else if (((int)EdgeType & (int)IMeshWireframeSource::EMeshEdgeType::MeshBoundary) != 0 && Component->bEnableBoundaryEdges)
-			{
-				bEdgeIsVisible = true;
-			}
-			else if (((int)EdgeType & (int)IMeshWireframeSource::EMeshEdgeType::UVSeam) != 0 && Component->bEnableUVSeams)
-			{
-				bEdgeIsVisible = true;
-			}
-			else if (((int)EdgeType & (int)IMeshWireframeSource::EMeshEdgeType::NormalSeam) != 0 && Component->bEnableNormalSeams)
-			{
-				bEdgeIsVisible = true;
-			}
-			else if (((int)EdgeType & (int)IMeshWireframeSource::EMeshEdgeType::ColorSeam) != 0 && Component->bEnableColorSeams)
-			{
-				bEdgeIsVisible = true;
+				if (bEnableWireframe ||
+					(bEnableBoundaryEdges && (static_cast<int>(EdgeType) & static_cast<int>(IMeshWireframeSource::EMeshEdgeType::MeshBoundary)) != 0) ||
+					(bEnableUVSeams       && (static_cast<int>(EdgeType) & static_cast<int>(IMeshWireframeSource::EMeshEdgeType::UVSeam      )) != 0) ||
+					(bEnableNormalSeams   && (static_cast<int>(EdgeType) & static_cast<int>(IMeshWireframeSource::EMeshEdgeType::NormalSeam  )) != 0) ||
+					(bEnableColorSeams    && (static_cast<int>(EdgeType) & static_cast<int>(IMeshWireframeSource::EMeshEdgeType::ColorSeam   )) != 0))
+				{
+					VisibleEdgeBatch.Add(UE::Geometry::FIndex3i(VertIndexA, VertIndexB, static_cast<int>(EdgeType)));
+				}
 			}
+		});
 
-			if (bEdgeIsVisible)
-			{
-				CurrentEdgeSet.Add(FIndex4i(li, VertIndexA, VertIndexB, (int)EdgeType));
-			}
-		};
-		int32 NumEdges = CurrentEdgeSet.Num();
-		if (NumEdges == 0)
+		for (int32 BatchIndex = 0; BatchIndex < NumBatches; ++BatchIndex)
+		{
+			VisibleEdgeBatches[BatchIndex].Offset = NumVisibleEdges;
+			NumVisibleEdges += VisibleEdgeBatches[BatchIndex].Edges.Num();
+		}
+
+		if (NumVisibleEdges == 0)
 		{
 			return;
 		}
 
-		const int32 NumLineVertices = NumEdges * 4;
-		const int32 NumLineIndices = NumEdges * 6;
-		const int32 NumTextureCoordinates = 1;
+		const int32 NumLineVertices = NumVisibleEdges * 4;
+		const int32 NumLineIndices = NumVisibleEdges * 6;
+		constexpr int32 NumTextureCoordinates = 1;
 
 		VertexBuffers.PositionVertexBuffer.Init(NumLineVertices);
 		VertexBuffers.StaticMeshVertexBuffer.Init(NumLineVertices, NumTextureCoordinates);
@@ -105,7 +122,7 @@ public:
 		MeshBatchData.MinVertexIndex = 0;
 		MeshBatchData.MaxVertexIndex = 0 + NumLineVertices - 1;
 		MeshBatchData.StartIndex = 0;
-		MeshBatchData.NumPrimitives = NumEdges * 2;
+		MeshBatchData.NumPrimitives = NumVisibleEdges * 2;
 		if (Component->GetMaterial(0) != nullptr)
 		{
 			MeshBatchData.MaterialProxy = Component->GetMaterial(0)->GetRenderProxy();
@@ -114,94 +131,100 @@ public:
 		{
 			MeshBatchData.MaterialProxy = UMaterial::GetDefaultMaterial(MD_Surface)->GetRenderProxy();
 		}
-		
-		FColor RegularEdgeColor = FLinearColor::FromSRGBColor(Component->WireframeColor).ToFColor(false);
-		float RegularEdgeThickness = Component->ThicknessScale * Component->WireframeThickness;
-		FColor BoundaryEdgeColor = FLinearColor::FromSRGBColor(Component->BoundaryEdgeColor).ToFColor(false);
-		float BoundaryEdgeThickness = Component->ThicknessScale * Component->BoundaryEdgeThickness;
-		FColor UVSeamColor = FLinearColor::FromSRGBColor(Component->UVSeamColor).ToFColor(false);
-		float UVSeamThickness = Component->ThicknessScale * Component->UVSeamThickness;
-		FColor NormalSeamColor = FLinearColor::FromSRGBColor(Component->NormalSeamColor).ToFColor(false);
-		float NormalSeamThickness = Component->ThicknessScale * Component->NormalSeamThickness;
-		FColor ColorSeamColor = FLinearColor::FromSRGBColor(Component->ColorSeamColor).ToFColor(false);
-		float ColorSeamThickness = Component->ThicknessScale * Component->ColorSeamThickness;
-
-		float LineDepthBias = Component->LineDepthBias * Component->LineDepthBiasSizeScale;
+
+		const FColor RegularEdgeColor = FLinearColor::FromSRGBColor(Component->WireframeColor).ToFColor(false);
+		const float RegularEdgeThickness = Component->ThicknessScale * Component->WireframeThickness;
+		const FColor BoundaryEdgeColor = FLinearColor::FromSRGBColor(Component->BoundaryEdgeColor).ToFColor(false);
+		const float BoundaryEdgeThickness = Component->ThicknessScale * Component->BoundaryEdgeThickness;
+		const FColor UVSeamColor = FLinearColor::FromSRGBColor(Component->UVSeamColor).ToFColor(false);
+		const float UVSeamThickness = Component->ThicknessScale * Component->UVSeamThickness;
+		const FColor NormalSeamColor = FLinearColor::FromSRGBColor(Component->NormalSeamColor).ToFColor(false);
+		const float NormalSeamThickness = Component->ThicknessScale * Component->NormalSeamThickness;
+		const FColor ColorSeamColor = FLinearColor::FromSRGBColor(Component->ColorSeamColor).ToFColor(false);
+		const float ColorSeamThickness = Component->ThicknessScale * Component->ColorSeamThickness;
+
+		const float LineDepthBias = Component->LineDepthBias * Component->LineDepthBiasSizeScale;
 
 		// Initialize lines.
 		// Lines are represented as two tris of zero thickness. The UV's stored at vertices are actually (lineThickness, depthBias), 
 		// which the material unpacks and uses to thicken the polygons and set the pixel depth bias.
-		ParallelFor(NumEdges, [&](int32 idx)
+		ParallelFor(NumBatches, [&](int32 BatchIndex)
 		{
-			int32 VertexBufferIndex = idx * 4;
-			int32 IndexBufferIndex = idx * 6;
+			const int32 Offset = VisibleEdgeBatches[BatchIndex].Offset;
+			const TArray<UE::Geometry::FIndex3i>& VisibleEdgeBatch = VisibleEdgeBatches[BatchIndex].Edges;
 
-			FIndex4i EdgeInfo = CurrentEdgeSet[idx];
-			IMeshWireframeSource::EMeshEdgeType EdgeType = (IMeshWireframeSource::EMeshEdgeType)EdgeInfo.D;
+			for (int32 VisibleEdgeIndex = 0, Num = VisibleEdgeBatch.Num(); VisibleEdgeIndex < Num; ++VisibleEdgeIndex)
+			{
+				const int32 VertexBufferIndex = (Offset + VisibleEdgeIndex) * 4;
+				const int32 IndexBufferIndex = (Offset + VisibleEdgeIndex) * 6;
 
-			float UseThickness = RegularEdgeThickness;
-			FColor UseColor = RegularEdgeColor;
+				const UE::Geometry::FIndex3i EdgeInfo = VisibleEdgeBatch[VisibleEdgeIndex];
+				IMeshWireframeSource::EMeshEdgeType EdgeType = static_cast<IMeshWireframeSource::EMeshEdgeType>(EdgeInfo.C);
 
-			bool bIsRegularEdge = (EdgeType == IMeshWireframeSource::EMeshEdgeType::Regular);
-			bool bIsBoundaryEdge = (((int)EdgeType & (int)IMeshWireframeSource::EMeshEdgeType::MeshBoundary) != 0);
-			if (!bIsRegularEdge)
-			{
-				if (bIsBoundaryEdge && Component->bEnableBoundaryEdges)
-				{
-					UseThickness = BoundaryEdgeThickness;
-					UseColor = BoundaryEdgeColor;
-				}
-				else if (((int)EdgeType & (int)IMeshWireframeSource::EMeshEdgeType::UVSeam) != 0 && Component->bEnableUVSeams)
-				{
-					UseThickness = (bIsBoundaryEdge) ? BoundaryEdgeThickness : UVSeamThickness;
-					UseColor = UVSeamColor;
-				}
-				else if (((int)EdgeType & (int)IMeshWireframeSource::EMeshEdgeType::NormalSeam) != 0 && Component->bEnableNormalSeams)
-				{
-					UseThickness = (bIsBoundaryEdge) ? BoundaryEdgeThickness : NormalSeamThickness;
-					UseColor = NormalSeamColor;
-				}
-				else if (((int)EdgeType & (int)IMeshWireframeSource::EMeshEdgeType::ColorSeam) != 0 && Component->bEnableColorSeams)
+				float UseThickness = RegularEdgeThickness;
+				FColor UseColor = RegularEdgeColor;
+
+				if (EdgeType != IMeshWireframeSource::EMeshEdgeType::Regular)
 				{
-					UseThickness = (bIsBoundaryEdge) ? BoundaryEdgeThickness : ColorSeamThickness;
-					UseColor = ColorSeamColor;
+					const bool bIsBoundaryEdge = ((static_cast<int>(EdgeType) & static_cast<int>(IMeshWireframeSource::EMeshEdgeType::MeshBoundary)) != 0);
+
+					if (bEnableBoundaryEdges && bIsBoundaryEdge)
+					{
+						UseThickness = BoundaryEdgeThickness;
+						UseColor = BoundaryEdgeColor;
+					}
+					else if (bEnableUVSeams && (static_cast<int>(EdgeType) & static_cast<int>(IMeshWireframeSource::EMeshEdgeType::UVSeam)) != 0)
+					{
+						UseThickness = (bIsBoundaryEdge) ? BoundaryEdgeThickness : UVSeamThickness;
+						UseColor = UVSeamColor;
+					}
+					else if (bEnableNormalSeams && (static_cast<int>(EdgeType) & static_cast<int>(IMeshWireframeSource::EMeshEdgeType::NormalSeam)) != 0)
+					{
+						UseThickness = (bIsBoundaryEdge) ? BoundaryEdgeThickness : NormalSeamThickness;
+						UseColor = NormalSeamColor;
+					}
+					else if (bEnableColorSeams && (static_cast<int>(EdgeType) & static_cast<int>(IMeshWireframeSource::EMeshEdgeType::ColorSeam)) != 0)
+					{
+						UseThickness = (bIsBoundaryEdge) ? BoundaryEdgeThickness : ColorSeamThickness;
+						UseColor = ColorSeamColor;
+					}
 				}
-			}
 
-			const FVector A = WireSource->GetVertex(EdgeInfo.B);
-			const FVector B = WireSource->GetVertex(EdgeInfo.C);
-			const FVector LineDirection = (B - A).GetSafeNormal();
-			const FVector2f UV(UseThickness, LineDepthBias);
-
-			VertexBuffers.PositionVertexBuffer.VertexPosition(VertexBufferIndex + 0) = (FVector3f)A;
-			VertexBuffers.PositionVertexBuffer.VertexPosition(VertexBufferIndex + 1) = (FVector3f)B;
-			VertexBuffers.PositionVertexBuffer.VertexPosition(VertexBufferIndex + 2) = (FVector3f)B;
-			VertexBuffers.PositionVertexBuffer.VertexPosition(VertexBufferIndex + 3) = (FVector3f)A;
-
-			VertexBuffers.StaticMeshVertexBuffer.SetVertexTangents(VertexBufferIndex + 0, FVector3f::ZeroVector, FVector3f::ZeroVector, (FVector3f)-LineDirection);
-			VertexBuffers.StaticMeshVertexBuffer.SetVertexTangents(VertexBufferIndex + 1, FVector3f::ZeroVector, FVector3f::ZeroVector, (FVector3f)-LineDirection);
-			VertexBuffers.StaticMeshVertexBuffer.SetVertexTangents(VertexBufferIndex + 2, FVector3f::ZeroVector, FVector3f::ZeroVector, (FVector3f)LineDirection);
-			VertexBuffers.StaticMeshVertexBuffer.SetVertexTangents(VertexBufferIndex + 3, FVector3f::ZeroVector, FVector3f::ZeroVector, (FVector3f)LineDirection);
-
-			VertexBuffers.StaticMeshVertexBuffer.SetVertexUV(VertexBufferIndex + 0, 0, UV);
-			VertexBuffers.StaticMeshVertexBuffer.SetVertexUV(VertexBufferIndex + 1, 0, UV);
-			VertexBuffers.StaticMeshVertexBuffer.SetVertexUV(VertexBufferIndex + 2, 0, UV);
-			VertexBuffers.StaticMeshVertexBuffer.SetVertexUV(VertexBufferIndex + 3, 0, UV);
-
-			// The color stored in the vertices actually gets interpreted as a linear color by the material,
-			// whereas it is more convenient for the user of the MeshWireframe to specify colors as sRGB. So we actually
-			// have to convert it back to linear. The ToFColor(false) call just scales back into 0-255 space.
-			VertexBuffers.ColorVertexBuffer.VertexColor(VertexBufferIndex + 0) = UseColor;
-			VertexBuffers.ColorVertexBuffer.VertexColor(VertexBufferIndex + 1) = UseColor;
-			VertexBuffers.ColorVertexBuffer.VertexColor(VertexBufferIndex + 2) = UseColor;
-			VertexBuffers.ColorVertexBuffer.VertexColor(VertexBufferIndex + 3) = UseColor;
-
-			IndexBuffer.Indices[IndexBufferIndex + 0] = VertexBufferIndex + 0;
-			IndexBuffer.Indices[IndexBufferIndex + 1] = VertexBufferIndex + 1;
-			IndexBuffer.Indices[IndexBufferIndex + 2] = VertexBufferIndex + 2;
-			IndexBuffer.Indices[IndexBufferIndex + 3] = VertexBufferIndex + 2;
-			IndexBuffer.Indices[IndexBufferIndex + 4] = VertexBufferIndex + 3;
-			IndexBuffer.Indices[IndexBufferIndex + 5] = VertexBufferIndex + 0;
+				const FVector A = WireSource->GetVertex(EdgeInfo.A);
+				const FVector B = WireSource->GetVertex(EdgeInfo.B);
+				const FVector LineDirection = (B - A).GetSafeNormal();
+				const FVector2f UV(UseThickness, LineDepthBias);
+
+				VertexBuffers.PositionVertexBuffer.VertexPosition(VertexBufferIndex + 0) = static_cast<FVector3f>(A);
+				VertexBuffers.PositionVertexBuffer.VertexPosition(VertexBufferIndex + 1) = static_cast<FVector3f>(B);
+				VertexBuffers.PositionVertexBuffer.VertexPosition(VertexBufferIndex + 2) = static_cast<FVector3f>(B);
+				VertexBuffers.PositionVertexBuffer.VertexPosition(VertexBufferIndex + 3) = static_cast<FVector3f>(A);
+
+				VertexBuffers.StaticMeshVertexBuffer.SetVertexTangents(VertexBufferIndex + 0, FVector3f::ZeroVector, FVector3f::ZeroVector, static_cast<FVector3f>(-LineDirection));
+				VertexBuffers.StaticMeshVertexBuffer.SetVertexTangents(VertexBufferIndex + 1, FVector3f::ZeroVector, FVector3f::ZeroVector, static_cast<FVector3f>(-LineDirection));
+				VertexBuffers.StaticMeshVertexBuffer.SetVertexTangents(VertexBufferIndex + 2, FVector3f::ZeroVector, FVector3f::ZeroVector, static_cast<FVector3f>(LineDirection));
+				VertexBuffers.StaticMeshVertexBuffer.SetVertexTangents(VertexBufferIndex + 3, FVector3f::ZeroVector, FVector3f::ZeroVector, static_cast<FVector3f>(LineDirection));
+
+				VertexBuffers.StaticMeshVertexBuffer.SetVertexUV(VertexBufferIndex + 0, 0, UV);
+				VertexBuffers.StaticMeshVertexBuffer.SetVertexUV(VertexBufferIndex + 1, 0, UV);
+				VertexBuffers.StaticMeshVertexBuffer.SetVertexUV(VertexBufferIndex + 2, 0, UV);
+				VertexBuffers.StaticMeshVertexBuffer.SetVertexUV(VertexBufferIndex + 3, 0, UV);
+
+				// The color stored in the vertices actually gets interpreted as a linear color by the material,
+				// whereas it is more convenient for the user of the MeshWireframe to specify colors as sRGB. So we actually
+				// have to convert it back to linear. The ToFColor(false) call just scales back into 0-255 space.
+				VertexBuffers.ColorVertexBuffer.VertexColor(VertexBufferIndex + 0) = UseColor;
+				VertexBuffers.ColorVertexBuffer.VertexColor(VertexBufferIndex + 1) = UseColor;
+				VertexBuffers.ColorVertexBuffer.VertexColor(VertexBufferIndex + 2) = UseColor;
+				VertexBuffers.ColorVertexBuffer.VertexColor(VertexBufferIndex + 3) = UseColor;
+
+				IndexBuffer.Indices[IndexBufferIndex + 0] = VertexBufferIndex + 0;
+				IndexBuffer.Indices[IndexBufferIndex + 1] = VertexBufferIndex + 1;
+				IndexBuffer.Indices[IndexBufferIndex + 2] = VertexBufferIndex + 2;
+				IndexBuffer.Indices[IndexBufferIndex + 3] = VertexBufferIndex + 2;
+				IndexBuffer.Indices[IndexBufferIndex + 4] = VertexBufferIndex + 3;
+				IndexBuffer.Indices[IndexBufferIndex + 5] = VertexBufferIndex + 0;
+			}
 		});
 
 		ENQUEUE_RENDER_COMMAND(MeshWireframeVertexBuffersInit)(
@@ -223,7 +246,7 @@ public:
 		});
 	}
 
-	virtual ~FMeshWireframeSceneProxy()
+	virtual ~FMeshWireframeSceneProxy() override
 	{
 		VertexBuffers.PositionVertexBuffer.ReleaseResource();
 		VertexBuffers.StaticMeshVertexBuffer.ReleaseResource();
@@ -304,8 +327,6 @@ private:
 	FLocalVertexFactory VertexFactory;
 	FStaticMeshVertexBuffers VertexBuffers;
 	FDynamicMeshIndexBuffer32 IndexBuffer;
-
-	TArray<FIndex4i> CurrentEdgeSet;
 };
 
 
