ModelingComponents: refactor conversion code from ExtractCollisionGeometryTool into static functions ConvertSimpleCollisionToMeshes() and ConvertComplexCollisionToMeshes()
GeometryScript: add CopyCollisionMeshesFromObject function, extracts meshed collision shapes or complex collision from various object types as a UDynamicMesh

diff --git a/Engine/Plugins/Experimental/GeometryScripting/Source/GeometryScriptingCore/Private/SceneUtilityFunctions.cpp b/Engine/Plugins/Experimental/GeometryScripting/Source/GeometryScriptingCore/Private/SceneUtilityFunctions.cpp
--- a/Engine/Plugins/Experimental/GeometryScripting/Source/GeometryScriptingCore/Private/SceneUtilityFunctions.cpp
+++ b/Engine/Plugins/Experimental/GeometryScripting/Source/GeometryScriptingCore/Private/SceneUtilityFunctions.cpp
@@ -8,10 +8,14 @@
 #include "DynamicMesh/MeshTransforms.h"
 #include "DynamicMesh/MeshNormals.h"
 
+#include "Engine/StaticMesh.h"
+#include "Engine/SkeletalMesh.h"
 #include "Components/StaticMeshComponent.h"
+#include "Components/SkeletalMeshComponent.h"
 #include "Components/BrushComponent.h"
 #include "Components/DynamicMeshComponent.h"
 #include "ConversionUtils/VolumeToDynamicMesh.h"
+#include "Physics/ComponentCollisionUtil.h"
 
 using namespace UE::Geometry;
 
@@ -136,4 +140,102 @@ void UGeometryScriptLibrary_SceneUtilityFunctions::SetComponentMaterialList(
 
 
 
+UDynamicMesh* UGeometryScriptLibrary_SceneUtilityFunctions::CopyCollisionMeshesFromObject(
+	UObject* FromObject,
+	UDynamicMesh* ToDynamicMesh,
+	bool bTransformToWorld,
+	FTransform& LocalToWorld,
+	TEnumAsByte<EGeometryScriptOutcomePins>& Outcome,
+	bool bUseComplexCollision,
+	int SphereResolution,
+	UGeometryScriptDebug* Debug)
+{
+	Outcome = EGeometryScriptOutcomePins::Failure;
+
+	FDynamicMesh3 AccumulatedMesh;
+	ToDynamicMesh->ProcessMesh([&](const FDynamicMesh3& ReadMesh)
+	{
+		AccumulatedMesh.EnableMatchingAttributes(ReadMesh, true);
+	});
+
+	if (bUseComplexCollision)
+	{
+		// rewrite object type if necessary
+		if (UStaticMeshComponent* StaticMeshComp = Cast<UStaticMeshComponent>(FromObject))
+		{
+			FromObject = StaticMeshComp->GetStaticMesh();
+		}
+
+		// find the Complex Collision mesh interface
+		IInterface_CollisionDataProvider* CollisionProvider = nullptr;
+		if (UStaticMesh* StaticMesh = Cast<UStaticMesh>(FromObject))
+		{
+			LocalToWorld = FTransform::Identity;
+			CollisionProvider = (IInterface_CollisionDataProvider*)StaticMesh;
+		}
+		else if ( Cast<UDynamicMeshComponent>(FromObject) != nullptr )
+		{
+			UPrimitiveComponent* AnyComponent = Cast<UPrimitiveComponent>(FromObject);
+			LocalToWorld = AnyComponent->GetComponentTransform();
+			CollisionProvider = (IInterface_CollisionDataProvider*)AnyComponent;
+		}
+
+		if (CollisionProvider == nullptr)
+		{
+			UE::Geometry::AppendError(Debug, EGeometryScriptErrorType::InvalidInputs, LOCTEXT("CopyCollisionMeshesFromObject_NoCollisionProvider", "CopyCollisionMeshesFromObject: Complex Collision Provider is not available for this object type"));
+			return ToDynamicMesh;
+		}
+
+		FTransformSequence3d Transforms;
+		if (bTransformToWorld)
+		{
+			Transforms.Append(LocalToWorld);
+		}
+
+		// generate mesh
+		bool bFoundMeshErrors = false;
+		UE::Geometry::ConvertComplexCollisionToMeshes(CollisionProvider, AccumulatedMesh, Transforms, bFoundMeshErrors, true, true);
+	}
+	else
+	{
+		const UBodySetup* BodySetup = nullptr;
+		if (UPrimitiveComponent* AnyComponent = Cast<UPrimitiveComponent>(FromObject))
+		{
+			LocalToWorld = AnyComponent->GetComponentTransform();
+			BodySetup = AnyComponent->GetBodySetup();
+		}
+		else if (const UStaticMesh* StaticMesh = Cast<UStaticMesh>(FromObject))
+		{
+			LocalToWorld = FTransform::Identity;
+			BodySetup = StaticMesh->GetBodySetup();
+		}
+
+		if (BodySetup == nullptr)
+		{
+			UE::Geometry::AppendWarning(Debug, EGeometryScriptErrorType::InvalidInputs, LOCTEXT("CopyCollisionMeshesFromObject_NoBodySetup", "CopyCollisionMeshesFromObject: BodySetup is null or Object type is not supported"));
+			return ToDynamicMesh;
+		}
+
+		FTransformSequence3d Transforms;
+		if (bTransformToWorld)
+		{
+			Transforms.Append(LocalToWorld);
+		}
+
+		UE::Geometry::ConvertSimpleCollisionToMeshes(BodySetup->AggGeom, AccumulatedMesh, Transforms, SphereResolution, true, true);
+	}
+
+	ToDynamicMesh->EditMesh([&](FDynamicMesh3& EditMesh)
+	{
+		EditMesh = MoveTemp(AccumulatedMesh);
+	});
+
+	Outcome = EGeometryScriptOutcomePins::Success;
+
+	return ToDynamicMesh;
+}
+
+
+
+
 #undef LOCTEXT_NAMESPACE
\ No newline at end of file
diff --git a/Engine/Plugins/Experimental/GeometryScripting/Source/GeometryScriptingCore/Public/GeometryScript/SceneUtilityFunctions.h b/Engine/Plugins/Experimental/GeometryScripting/Source/GeometryScriptingCore/Public/GeometryScript/SceneUtilityFunctions.h
--- a/Engine/Plugins/Experimental/GeometryScripting/Source/GeometryScriptingCore/Public/GeometryScript/SceneUtilityFunctions.h
+++ b/Engine/Plugins/Experimental/GeometryScripting/Source/GeometryScriptingCore/Public/GeometryScript/SceneUtilityFunctions.h
@@ -34,12 +34,22 @@ class GEOMETRYSCRIPTINGCORE_API UGeometryScriptLibrary_SceneUtilityFunctions : p
 	GENERATED_BODY()
 public:
 
+	/**
+	 * Create a new UDynamicMeshPool object. 
+	 * Caller needs to create a UProperty reference to the returned object, or it will be garbage-collected.
+	 */
 	UFUNCTION(BlueprintCallable, Category = "GeometryScript|Utility")
 	static UPARAM(DisplayName = "Dynamic Mesh Pool") UDynamicMeshPool*
 	CreateDynamicMeshPool();
 
 
-
+	/**
+	 * Copy the mesh from a given Component to a Dynamic Mesh.
+	 * StaticMeshComponent, DynamicMeshCompnent, and BrushComponent are supported.
+	 * This function offers minimal control over the copying, if more control is needed for Static Meshes, use CopyMeshFromStaticMesh.
+	 * @param bTransformToWorld if true, output mesh is in World space
+	 * @param LocalToWorld local-to-world transform is returned here (whether mesh is in local or world space)
+	 */
 	UFUNCTION(BlueprintCallable, Category = "GeometryScript|Scene", meta = (ExpandEnumAsExecs = "Outcome"))
 	static UPARAM(DisplayName = "Dynamic Mesh") UDynamicMesh*
 	CopyMeshFromComponent(
@@ -52,6 +62,10 @@ public:
 		UGeometryScriptDebug* Debug = nullptr);
 
 
+	/**
+	 * Configure the Material set on a PrimitiveComponent, by repeatedly calling SetMaterial.
+	 * This is a simple utility function and it's behavior will depend on the specifics of the Component.
+	 */
 	UFUNCTION(BlueprintCallable, Category = "GeometryScript|Scene")
 	static void
 	SetComponentMaterialList(
@@ -59,4 +73,26 @@ public:
 		const TArray<UMaterialInterface*>& MaterialList,
 		UGeometryScriptDebug* Debug = nullptr);
 
+
+	/**
+	 * Extract the Collision Geometry from FromObject and copy/approximate it with meshes stored in ToDynamicMesh.
+	 * For Simple Collision, FromObject can be a StaticMesh Asset or any PrimitiveComponent
+	 * For Complex Collision, FromObject can be a StaticMesh Asset, StaticMeshComponent, or DynamicMeshComponent
+	 * @param bTransformToWorld if true, output mesh is in World space
+	 * @param LocalToWorld local-to-world transform is returned here (whether mesh is in local or world space)
+	 * @param bUseComplexCollision if true, complex collision is extracted, otherwise Simple collision shapes are meshed
+	 * @param SphereResolution determines tessellation density of sphere and capsule shapes
+	 */
+	UFUNCTION(BlueprintCallable, Category = "GeometryScript|Scene", meta = (ExpandEnumAsExecs = "Outcome"))
+	static UPARAM(DisplayName = "Dynamic Mesh") UDynamicMesh*
+	CopyCollisionMeshesFromObject(
+		UObject* FromObject,
+		UDynamicMesh* ToDynamicMesh,
+		bool bTransformToWorld,
+		FTransform& LocalToWorld,
+		TEnumAsByte<EGeometryScriptOutcomePins>& Outcome,
+		bool bUseComplexCollision = false,
+		int SphereResolution = 16,
+		UGeometryScriptDebug* Debug = nullptr);
+
 };
\ No newline at end of file
diff --git a/Engine/Plugins/Experimental/MeshModelingToolsetExp/Source/MeshModelingToolsExp/Private/Physics/ExtractCollisionGeometryTool.cpp b/Engine/Plugins/Experimental/MeshModelingToolsetExp/Source/MeshModelingToolsExp/Private/Physics/ExtractCollisionGeometryTool.cpp
--- a/Engine/Plugins/Experimental/MeshModelingToolsetExp/Source/MeshModelingToolsExp/Private/Physics/ExtractCollisionGeometryTool.cpp
+++ b/Engine/Plugins/Experimental/MeshModelingToolsetExp/Source/MeshModelingToolsExp/Private/Physics/ExtractCollisionGeometryTool.cpp
@@ -21,6 +21,7 @@
 
 #include "Physics/PhysicsDataCollection.h"
 #include "Physics/CollisionGeometryVisualization.h"
+#include "Physics/ComponentCollisionUtil.h"
 
 // physics data
 #include "PhysicsEngine/BodySetup.h"
@@ -246,97 +247,18 @@ void UExtractCollisionGeometryTool::RecalculateMesh_Simple()
 
 	CurrentMeshParts.Reset();
 
-	FDynamicMeshEditor Editor(&CurrentMesh);
-
 	const FKAggregateGeom& AggGeom = PhysicsInfo->AggGeom;
-
-	for (const FKSphereElem& Sphere : AggGeom.SphereElems)
-	{
-		FSphereGenerator SphereGen;
-		SphereGen.Radius = Sphere.Radius;
-		SphereGen.NumPhi = SphereGen.NumTheta = SphereResolution;
-		SphereGen.bPolygroupPerQuad = false;
-		SphereGen.Generate();
-		FDynamicMesh3 SphereMesh(&SphereGen);
-
-		MeshTransforms::Translate(SphereMesh, FVector3d(Sphere.Center));
-
-		FMeshIndexMappings Mappings;
-		Editor.AppendMesh(&SphereMesh, Mappings);
-		CurrentMeshParts.Add(MakeShared<FDynamicMesh3>(MoveTemp(SphereMesh)));
-	}
-
-
-	for (const FKBoxElem& Box : AggGeom.BoxElems)
-	{
-		FMinimalBoxMeshGenerator BoxGen;
-		BoxGen.Box = UE::Geometry::FOrientedBox3d(
-			FFrame3d(FVector3d(Box.Center), FQuaterniond(Box.Rotation.Quaternion())),
-			0.5*FVector3d(Box.X, Box.Y, Box.Z));
-		BoxGen.Generate();
-		FDynamicMesh3 BoxMesh(&BoxGen);
-
-		// transform not applied because it is just the Center/Rotation
-
-		FMeshIndexMappings Mappings;
-		Editor.AppendMesh(&BoxMesh, Mappings);
-		CurrentMeshParts.Add(MakeShared<FDynamicMesh3>(MoveTemp(BoxMesh)));
-	}
-
-
-	for (const FKSphylElem& Capsule: AggGeom.SphylElems)
-	{
-		FCapsuleGenerator CapsuleGen;
-		CapsuleGen.Radius = Capsule.Radius;
-		CapsuleGen.SegmentLength = Capsule.Length;
-		CapsuleGen.NumHemisphereArcSteps = SphereResolution/4+1;
-		CapsuleGen.NumCircleSteps = SphereResolution;
-		CapsuleGen.bPolygroupPerQuad = false;
-		CapsuleGen.Generate();
-		FDynamicMesh3 CapsuleMesh(&CapsuleGen);
-
-		MeshTransforms::Translate(CapsuleMesh, FVector3d(0,0,-0.5*Capsule.Length) );
-
-		FTransformSRT3d Transform(Capsule.GetTransform());
-		MeshTransforms::ApplyTransform(CapsuleMesh, Transform);
-
-		FMeshIndexMappings Mappings;
-		Editor.AppendMesh(&CapsuleMesh, Mappings);
-		CurrentMeshParts.Add(MakeShared<FDynamicMesh3>(MoveTemp(CapsuleMesh)));
-	}
-
-
-	for (const FKConvexElem& Convex : AggGeom.ConvexElems)
-	{
-		FTransformSRT3d ElemTransform(Convex.GetTransform());
-		FDynamicMesh3 ConvexMesh(EMeshComponents::None);
-		int32 NumVertices = Convex.VertexData.Num();
-		for (int32 k = 0; k < NumVertices; ++k)
-		{
-			ConvexMesh.AppendVertex( FVector3d(Convex.VertexData[k]) );
-		}
-		int32 NumTriangles = Convex.IndexData.Num() / 3;
-		for (int32 k = 0; k < NumTriangles; ++k)
-		{
-			ConvexMesh.AppendTriangle(Convex.IndexData[3*k], Convex.IndexData[3*k+1], Convex.IndexData[3*k+2]);
-		}
-
-		ConvexMesh.ReverseOrientation();
-		ConvexMesh.EnableTriangleGroups(0);
-		ConvexMesh.EnableAttributes();
-		FDynamicMeshUVEditor UVEditor(&ConvexMesh, 0, true);
-		UVEditor.SetPerTriangleUVs();
-		FMeshIndexMappings Mappings;
-		Editor.AppendMesh(&ConvexMesh, Mappings);
-		CurrentMeshParts.Add(MakeShared<FDynamicMesh3>(MoveTemp(ConvexMesh)));
-	}
+	UE::Geometry::ConvertSimpleCollisionToMeshes(AggGeom, CurrentMesh,
+		FTransformSequence3d(), SphereResolution, true, true,
+		[&](int32 ElemType, const FDynamicMesh3& ElemMesh) {
+			CurrentMeshParts.Add(MakeShared<FDynamicMesh3>(ElemMesh));
+		});
 
 	for ( int32 k = 0; k < CurrentMeshParts.Num(); ++k)
 	{
 		FDynamicMesh3& MeshPart = *CurrentMeshParts[k];
 		FMeshNormals::InitializeMeshToPerTriangleNormals(&MeshPart);
 	}
-	FMeshNormals::InitializeMeshToPerTriangleNormals(&CurrentMesh);
 
 	PreviewMesh->UpdatePreview(&CurrentMesh);
 
@@ -357,43 +279,10 @@ void UExtractCollisionGeometryTool::RecalculateMesh_Complex()
 	bool bMeshErrors = false;
 
 	IInterface_CollisionDataProvider* CollisionProvider = UE::ToolTarget::GetPhysicsCollisionDataProvider(Target);
-	if (CollisionProvider && CollisionProvider->ContainsPhysicsTriMeshData(true))
+	if (CollisionProvider)
 	{
-		FTriMeshCollisionData CollisionData;
-		if (CollisionProvider->GetPhysicsTriMeshData(&CollisionData, true))
-		{
-			TArray<int32> VertexIDMap;
-			for (int32 k = 0; k < CollisionData.Vertices.Num(); ++k)
-			{
-				int32 vid = CurrentMesh.AppendVertex((FVector)CollisionData.Vertices[k]);
-				VertexIDMap.Add(vid);
-			}
-			for (const FTriIndices& TriIndices : CollisionData.Indices)
-			{
-				FIndex3i Triangle(TriIndices.v0, TriIndices.v1, TriIndices.v2);
-				int32 tid = CurrentMesh.AppendTriangle(Triangle);
-				if (tid < 0)
-				{
-					bMeshErrors = true;
-				}
-			}
-
-			if (Settings->bWeldEdges)
-			{
-				FMergeCoincidentMeshEdges Weld(&CurrentMesh);
-				Weld.OnlyUniquePairs = true;
-				Weld.Apply();
-				Weld.OnlyUniquePairs = false;
-				Weld.Apply();
-			}
-
-			if (!CollisionData.bFlipNormals)		// collision mesh has reversed orientation
-			{
-				CurrentMesh.ReverseOrientation(false);
-			}
-
-			FMeshNormals::InitializeMeshToPerTriangleNormals(&CurrentMesh);
-		}
+		FTransformSequence3d Transform;
+		UE::Geometry::ConvertComplexCollisionToMeshes(CollisionProvider, CurrentMesh, FTransformSequence3d(), bMeshErrors, true, true);
 	}
 
 	PreviewMesh->UpdatePreview(&CurrentMesh);
diff --git a/Engine/Plugins/Runtime/MeshModelingToolset/Source/ModelingComponents/Private/Physics/ComponentCollisionUtil.cpp b/Engine/Plugins/Runtime/MeshModelingToolset/Source/ModelingComponents/Private/Physics/ComponentCollisionUtil.cpp
--- a/Engine/Plugins/Runtime/MeshModelingToolset/Source/ModelingComponents/Private/Physics/ComponentCollisionUtil.cpp
+++ b/Engine/Plugins/Runtime/MeshModelingToolset/Source/ModelingComponents/Private/Physics/ComponentCollisionUtil.cpp
@@ -4,10 +4,17 @@
 
 #include "DynamicMesh/DynamicMesh3.h"
 #include "ShapeApproximation/SimpleShapeSet3.h"
+#include "DynamicMeshEditor.h"
+#include "DynamicMesh/MeshNormals.h"
+#include "Generators/SphereGenerator.h"
+#include "Generators/MinimalBoxMeshGenerator.h"
+#include "Generators/CapsuleGenerator.h"
+#include "DynamicMesh/MeshTransforms.h"
+#include "Parameterization/DynamicMeshUVEditor.h"
+#include "DynamicMesh/Operations/MergeCoincidentMeshEdges.h"
 
 #include "Physics/PhysicsDataCollection.h"
 
-
 #include "Components/StaticMeshComponent.h"
 #include "Engine/StaticMesh.h"
 #include "Engine/Classes/PhysicsEngine/BodySetup.h"
@@ -178,3 +185,202 @@ bool UE::Geometry::AppendSimpleCollision(
 }
 
 
+void UE::Geometry::ConvertSimpleCollisionToMeshes(
+	const FKAggregateGeom& AggGeom,
+	FDynamicMesh3& MeshOut,
+	const FTransformSequence3d& TransformSeqeuence,
+	int32 SphereResolution,
+	bool bSetToPerTriangleNormals,
+	bool bInitializeConvexUVs,
+	TFunction<void(int, const FDynamicMesh3&)> PerElementMeshCallback )
+{
+	FDynamicMeshEditor Editor(&MeshOut);
+
+	for (const FKSphereElem& Sphere : AggGeom.SphereElems)
+	{
+		FSphereGenerator SphereGen;
+		SphereGen.Radius = Sphere.Radius;
+		SphereGen.NumPhi = SphereGen.NumTheta = SphereResolution;
+		SphereGen.bPolygroupPerQuad = false;
+		SphereGen.Generate();
+		FDynamicMesh3 SphereMesh(&SphereGen);
+
+		MeshTransforms::Translate(SphereMesh, FVector3d(Sphere.Center));
+
+		FMeshIndexMappings Mappings;
+		Editor.AppendMesh(&SphereMesh, Mappings,
+			[&](int32 vid, const FVector3d& P) { return TransformSeqeuence.TransformPosition(P); },
+			[&](int32 vid, const FVector3d& N) { return TransformSeqeuence.TransformNormal(N); });
+
+		if (PerElementMeshCallback)
+		{
+			PerElementMeshCallback(0, SphereMesh);
+		}
+	}
+
+	for (const FKBoxElem& Box : AggGeom.BoxElems)
+	{
+		FMinimalBoxMeshGenerator BoxGen;
+		BoxGen.Box = UE::Geometry::FOrientedBox3d(
+			FFrame3d(FVector3d(Box.Center), FQuaterniond(Box.Rotation.Quaternion())),
+			0.5*FVector3d(Box.X, Box.Y, Box.Z));
+		BoxGen.Generate();
+		FDynamicMesh3 BoxMesh(&BoxGen);
+
+		// transform not applied because it is just the Center/Rotation
+
+		FMeshIndexMappings Mappings;
+		Editor.AppendMesh(&BoxMesh, Mappings,
+			[&](int32 vid, const FVector3d& P) { return TransformSeqeuence.TransformPosition(P); },
+			[&](int32 vid, const FVector3d& N) { return TransformSeqeuence.TransformNormal(N); });	
+
+		if (PerElementMeshCallback)
+		{
+			PerElementMeshCallback(1, BoxMesh);
+		}
+	}
+
+
+	for (const FKSphylElem& Capsule: AggGeom.SphylElems)
+	{
+		FCapsuleGenerator CapsuleGen;
+		CapsuleGen.Radius = Capsule.Radius;
+		CapsuleGen.SegmentLength = Capsule.Length;
+		CapsuleGen.NumHemisphereArcSteps = SphereResolution/4+1;
+		CapsuleGen.NumCircleSteps = SphereResolution;
+		CapsuleGen.bPolygroupPerQuad = false;
+		CapsuleGen.Generate();
+		FDynamicMesh3 CapsuleMesh(&CapsuleGen);
+
+		MeshTransforms::Translate(CapsuleMesh, FVector3d(0,0,-0.5*Capsule.Length) );
+
+		FTransformSRT3d Transform(Capsule.GetTransform());
+		MeshTransforms::ApplyTransform(CapsuleMesh, Transform);
+
+		FMeshIndexMappings Mappings;
+		Editor.AppendMesh(&CapsuleMesh, Mappings,
+			[&](int32 vid, const FVector3d& P) { return TransformSeqeuence.TransformPosition(P); },
+			[&](int32 vid, const FVector3d& N) { return TransformSeqeuence.TransformNormal(N); });	
+
+		if (PerElementMeshCallback)
+		{
+			PerElementMeshCallback(2, CapsuleMesh);
+		}
+	}
+
+
+	for (const FKConvexElem& Convex : AggGeom.ConvexElems)
+	{
+		FTransformSRT3d ElemTransform(Convex.GetTransform());
+		FDynamicMesh3 ConvexMesh(EMeshComponents::None);
+		int32 NumVertices = Convex.VertexData.Num();
+		for (int32 k = 0; k < NumVertices; ++k)
+		{
+			ConvexMesh.AppendVertex( FVector3d(Convex.VertexData[k]) );
+		}
+		int32 NumTriangles = Convex.IndexData.Num() / 3;
+		for (int32 k = 0; k < NumTriangles; ++k)
+		{
+			ConvexMesh.AppendTriangle(Convex.IndexData[3*k], Convex.IndexData[3*k+1], Convex.IndexData[3*k+2]);
+		}
+
+		ConvexMesh.ReverseOrientation();
+		ConvexMesh.EnableTriangleGroups(0);
+		ConvexMesh.EnableAttributes();
+		if (bInitializeConvexUVs)
+		{
+			FDynamicMeshUVEditor UVEditor(&ConvexMesh, 0, true);
+			UVEditor.SetPerTriangleUVs();
+		}
+
+		FMeshIndexMappings Mappings;
+		Editor.AppendMesh(&ConvexMesh, Mappings,
+			[&](int32 vid, const FVector3d& P) { return TransformSeqeuence.TransformPosition(P); },
+			[&](int32 vid, const FVector3d& N) { return TransformSeqeuence.TransformNormal(N); });	
+
+		if (PerElementMeshCallback)
+		{
+			PerElementMeshCallback(3, ConvexMesh);
+		}
+	}
+
+	if (MeshOut.HasAttributes() && MeshOut.Attributes()->PrimaryNormals() != nullptr)
+	{
+		if (bSetToPerTriangleNormals)
+		{
+			FMeshNormals::InitializeMeshToPerTriangleNormals(&MeshOut);
+		}
+		else
+		{
+			FMeshNormals::InitializeOverlayToPerVertexNormals(MeshOut.Attributes()->PrimaryNormals(), false);
+		}
+	}
+}
+
+
+
+bool UE::Geometry::ConvertComplexCollisionToMeshes(
+	IInterface_CollisionDataProvider* CollisionProvider,
+	UE::Geometry::FDynamicMesh3& MeshOut,
+	const FTransformSequence3d& TransformSeqeuence,
+	bool& bFoundMeshErrorsOut,
+	bool bWeldEdges,
+	bool bSetToPerTriangleNormals)
+{
+	bFoundMeshErrorsOut = false;
+	if (CollisionProvider && CollisionProvider->ContainsPhysicsTriMeshData(true))
+	{
+		FTriMeshCollisionData CollisionData;
+		if (CollisionProvider->GetPhysicsTriMeshData(&CollisionData, true))
+		{
+			TArray<int32> VertexIDMap;
+			for (int32 k = 0; k < CollisionData.Vertices.Num(); ++k)
+			{
+				int32 VertexID = MeshOut.AppendVertex((FVector)CollisionData.Vertices[k]);
+				VertexIDMap.Add(VertexID);
+			}
+			for (const FTriIndices& TriIndices : CollisionData.Indices)
+			{
+				FIndex3i Triangle(TriIndices.v0, TriIndices.v1, TriIndices.v2);
+				int32 TriangleID = MeshOut.AppendTriangle(Triangle);
+				if (TriangleID < 0)
+				{
+					bFoundMeshErrorsOut = true;
+					// create new vertices for this triangle
+					int32 A = MeshOut.AppendVertex(MeshOut.GetVertex(TriIndices.v0));
+					int32 B = MeshOut.AppendVertex(MeshOut.GetVertex(TriIndices.v1));
+					int32 C = MeshOut.AppendVertex(MeshOut.GetVertex(TriIndices.v2));
+					MeshOut.AppendTriangle(FIndex3i(A,B,C));
+				}
+			}
+
+			if (bWeldEdges)
+			{
+				FMergeCoincidentMeshEdges Weld(&MeshOut);
+				Weld.OnlyUniquePairs = true;
+				Weld.Apply();
+				Weld.OnlyUniquePairs = false;
+				Weld.Apply();
+			}
+
+			if (!CollisionData.bFlipNormals)		// collision mesh has reversed orientation
+			{
+				MeshOut.ReverseOrientation(false);
+			}
+		}
+	}
+
+	if (MeshOut.HasAttributes() && MeshOut.Attributes()->PrimaryNormals() != nullptr)
+	{
+		if (bSetToPerTriangleNormals)
+		{
+			FMeshNormals::InitializeMeshToPerTriangleNormals(&MeshOut);
+		}
+		else
+		{
+			FMeshNormals::InitializeOverlayToPerVertexNormals(MeshOut.Attributes()->PrimaryNormals(), false);
+		}
+	}
+
+	return (MeshOut.TriangleCount() > 0);
+}
\ No newline at end of file
diff --git a/Engine/Plugins/Runtime/MeshModelingToolset/Source/ModelingComponents/Public/Physics/ComponentCollisionUtil.h b/Engine/Plugins/Runtime/MeshModelingToolset/Source/ModelingComponents/Public/Physics/ComponentCollisionUtil.h
--- a/Engine/Plugins/Runtime/MeshModelingToolset/Source/ModelingComponents/Public/Physics/ComponentCollisionUtil.h
+++ b/Engine/Plugins/Runtime/MeshModelingToolset/Source/ModelingComponents/Public/Physics/ComponentCollisionUtil.h
@@ -5,10 +5,13 @@
 #include "CoreMinimal.h"
 #include "Components/PrimitiveComponent.h"
 #include "TransformTypes.h"
+#include "TransformSequence.h"
+#include "DynamicMesh/DynamicMesh3.h"
 
 class UBodySetup;
 class UStaticMesh;
 struct FKAggregateGeom;
+class IInterface_CollisionDataProvider;
 
 namespace UE
 {
@@ -85,7 +88,36 @@ MODELINGCOMPONENTS_API void UpdateSimpleCollision(
 	UStaticMesh* StaticMesh = nullptr,
 	FComponentCollisionSettings CollisionSettings = FComponentCollisionSettings());
 
+/**
+ * Extract the simple collision geometry from AggGeom as meshes (ie spheres and capsules are tessellated) and
+ * accumulate into MeshOut. TransformSequence is applied to the meshes as they are accumulated.
+ * @param bSetToPerTriangleNormals if true and the mesh has a normals attribute overlay,  the mesh is set to face normals, otherwise averaged vertex normals
+ * @param bInitializeConvexUVs if true convex hulls have their UVs initialized to per-face planar projections, otherwise no UVs are set
+ * @param PerElementMeshCallback if provided, called with each element mesh before transforming/appending to MeshOut. The int parameter is the shape type, 0=Sphere, 1=Box, 2=Capsule, 3=Convex
+ */
+MODELINGCOMPONENTS_API void ConvertSimpleCollisionToMeshes(
+	const FKAggregateGeom& AggGeom,
+	UE::Geometry::FDynamicMesh3& MeshOut,
+	const FTransformSequence3d& TransformSeqeuence,
+	int32 SphereResolution = 16,
+	bool bSetToPerTriangleNormals = false,
+	bool bInitializeConvexUVs = false,
+	TFunction<void(int, const FDynamicMesh3&)> PerElementMeshCallback = nullptr );
 
+/**
+ * Extract the complex collision geometry from CollisionDataProvider as a dynamic mesh.
+ * TransformSequence is applied to the mesh.
+ * @param bFoundMeshErrorsOut will be returned as true if potential non-manifold geometry was detected during construction (in that case some triangles will be disconnected)
+ * @param bWeldEdges if true the mesh edges are welded
+ * @param bSetToPerTriangleNormals if true and the mesh has a normals attribute overlay, the mesh is set to face normals, otherwise averaged vertex normals
+ */
+MODELINGCOMPONENTS_API bool ConvertComplexCollisionToMeshes(
+	IInterface_CollisionDataProvider* CollisionDataProvider,
+	UE::Geometry::FDynamicMesh3& MeshOut,
+	const FTransformSequence3d& TransformSeqeuence,
+	bool& bFoundMeshErrorsOut,
+	bool bWeldEdges = true,
+	bool bSetToPerTriangleNormals = false);
 
 }  // end namespace Geometry
 }  // end namespace UE
\ No newline at end of file
