ModelingTools: disable the OutputType property panel on CutMeshWithMeshTool, as it is not currently supported.

diff --git a/Engine/Plugins/Runtime/MeshModelingToolset/Source/MeshModelingTools/Private/CutMeshWithMeshTool.cpp b/Engine/Plugins/Runtime/MeshModelingToolset/Source/MeshModelingTools/Private/CutMeshWithMeshTool.cpp
--- a/Engine/Plugins/Runtime/MeshModelingToolset/Source/MeshModelingTools/Private/CutMeshWithMeshTool.cpp
+++ b/Engine/Plugins/Runtime/MeshModelingToolset/Source/MeshModelingTools/Private/CutMeshWithMeshTool.cpp
@@ -70,7 +70,7 @@ void UCutMeshWithMeshTool::ConvertInputsAndSetPreviewMaterials(bool bSetPreviewM
 	// disable output options
 	// (this property set is not registered yet in SetupProperties() above)
 	SetToolPropertySourceEnabled(HandleSourcesProperties, false);
-
+	SetToolPropertySourceEnabled(OutputTypeProperties, false);
 
 	FComponentMaterialSet AllMaterialSet;
 	TArray<TArray<int>> MaterialRemap; MaterialRemap.SetNum(Targets.Num());
@@ -377,6 +377,7 @@ void UCutMeshWithMeshTool::OnShutdown(EToolShutdownType ShutdownType)
 			NewMeshObjectParams.BaseName = UseBaseName;
 			NewMeshObjectParams.Materials = GetOutputMaterials();
 			NewMeshObjectParams.SetMesh(&IntersectionMesh);
+			// note: CutMeshWithMeshTool does not support converting types currently
 			UE::ToolTarget::ConfigureCreateMeshObjectParams(Targets[0], NewMeshObjectParams);
 			FCreateMeshObjectResult Result = UE::Modeling::CreateMeshObject(GetToolManager(), MoveTemp(NewMeshObjectParams));
 			if (Result.IsOK() && Result.NewActor != nullptr)
diff --git a/Engine/Plugins/Runtime/MeshModelingToolset/Source/MeshModelingTools/Public/CutMeshWithMeshTool.h b/Engine/Plugins/Runtime/MeshModelingToolset/Source/MeshModelingTools/Public/CutMeshWithMeshTool.h
--- a/Engine/Plugins/Runtime/MeshModelingToolset/Source/MeshModelingTools/Public/CutMeshWithMeshTool.h
+++ b/Engine/Plugins/Runtime/MeshModelingToolset/Source/MeshModelingTools/Public/CutMeshWithMeshTool.h
@@ -23,11 +23,11 @@ class MESHMODELINGTOOLS_API UCutMeshWithMeshToolProperties : public UInteractive
 	GENERATED_BODY()
 public:
 	/** Try to fill holes created by the Boolean operation, e.g. due to numerical errors */
-	UPROPERTY(EditAnywhere, Category = Boolean, AdvancedDisplay)
+	UPROPERTY(EditAnywhere, Category = Boolean)
 	bool bTryFixHoles = false;
 
 	/** Try to collapse extra edges created by the Boolean operation */
-	UPROPERTY(EditAnywhere, Category = Boolean, AdvancedDisplay)
+	UPROPERTY(EditAnywhere, Category = Boolean)
 	bool bTryCollapseEdges = true;
 
 	/** Threshold to determine whether a triangle in one mesh is inside or outside of the other */
