MeshModelingTools: Remove redundant bSnapToGrid property

diff --git a/Engine/Plugins/Runtime/MeshModelingToolset/Source/MeshModelingTools/Private/AddPrimitiveTool.cpp b/Engine/Plugins/Runtime/MeshModelingToolset/Source/MeshModelingTools/Private/AddPrimitiveTool.cpp
--- a/Engine/Plugins/Runtime/MeshModelingToolset/Source/MeshModelingTools/Private/AddPrimitiveTool.cpp
+++ b/Engine/Plugins/Runtime/MeshModelingToolset/Source/MeshModelingTools/Private/AddPrimitiveTool.cpp
@@ -289,21 +289,18 @@ void UAddPrimitiveTool::UpdatePreviewPosition(const FInputDeviceRay& DeviceClick
 		}
 	}
 
-	// Snap to grid if applicable
-	if (ShapeSettings->bSnapToGrid)
-	{
-		USceneSnappingManager* SnapManager = USceneSnappingManager::Find(GetToolManager());
-		if (SnapManager)
+	// Snap to grid
+	USceneSnappingManager* SnapManager = USceneSnappingManager::Find(GetToolManager());
+	if (SnapManager)
+	{
+		FSceneSnapQueryRequest Request;
+		Request.RequestType = ESceneSnapQueryType::Position;
+		Request.TargetTypes = ESceneSnapQueryTargetType::Grid;
+		Request.Position = (FVector)ShapeFrame.Origin;
+		TArray<FSceneSnapQueryResult> Results;
+		if (SnapManager->ExecuteSceneSnapQuery(Request, Results))
 		{
-			FSceneSnapQueryRequest Request;
-			Request.RequestType = ESceneSnapQueryType::Position;
-			Request.TargetTypes = ESceneSnapQueryTargetType::Grid;
-			Request.Position = (FVector)ShapeFrame.Origin;
-			TArray<FSceneSnapQueryResult> Results;
-			if (SnapManager->ExecuteSceneSnapQuery(Request, Results))
-			{
-				ShapeFrame.Origin = (FVector3d)Results[0].Position;
-			}
+			ShapeFrame.Origin = (FVector3d)Results[0].Position;
 		}
 	}
 
diff --git a/Engine/Plugins/Runtime/MeshModelingToolset/Source/MeshModelingTools/Public/AddPrimitiveTool.h b/Engine/Plugins/Runtime/MeshModelingToolset/Source/MeshModelingTools/Public/AddPrimitiveTool.h
--- a/Engine/Plugins/Runtime/MeshModelingToolset/Source/MeshModelingTools/Public/AddPrimitiveTool.h
+++ b/Engine/Plugins/Runtime/MeshModelingToolset/Source/MeshModelingTools/Public/AddPrimitiveTool.h
@@ -88,10 +88,6 @@ public:
 	UPROPERTY(EditAnywhere, Category = Positioning)
 	EMakeMeshPlacementType TargetSurface = EMakeMeshPlacementType::OnScene;
 
-	/** If true, the shape pivot is snapped to the grid. This is only relevant if the coordinate system is set to world space. */
-	UPROPERTY(EditAnywhere, Category = Positioning)
-	bool bSnapToGrid = true;
-
 	/** Location of pivot within the shape */
 	UPROPERTY(EditAnywhere, Category = Positioning, meta = (ProceduralShapeSetting))
 	EMakeMeshPivotLocation PivotLocation = EMakeMeshPivotLocation::Base;
