When applying MeshTransforms' transform and transforminverse, generally pass bReverseOrientationIfNeeded=true

(+ for code that was passing false then manually applying the same logic, remove the extra reverse orientation logic)

diff --git a/Engine/Plugins/Experimental/GeometryScripting/Source/GeometryScriptingCore/Private/MeshBakeFunctions.cpp b/Engine/Plugins/Experimental/GeometryScripting/Source/GeometryScriptingCore/Private/MeshBakeFunctions.cpp
--- a/Engine/Plugins/Experimental/GeometryScripting/Source/GeometryScriptingCore/Private/MeshBakeFunctions.cpp
+++ b/Engine/Plugins/Experimental/GeometryScripting/Source/GeometryScriptingCore/Private/MeshBakeFunctions.cpp
@@ -495,7 +495,7 @@ namespace GeometryScriptBakeLocals
 			// TODO: Remove this once we have support for transforming rays in the core bake loop
 			SourceMeshCopy = *SourceMeshOriginal;
 			const FTransformSRT3d SourceToWorld = SourceTransform;
-			MeshTransforms::ApplyTransform(SourceMeshCopy, SourceToWorld);
+			MeshTransforms::ApplyTransform(SourceMeshCopy, SourceToWorld, true);
 			const FTransformSRT3d TargetToWorld = TargetTransform;
 			MeshTransforms::ApplyTransformInverse(SourceMeshCopy, TargetToWorld, true);
 			SourceMeshToUse = &SourceMeshCopy;
diff --git a/Engine/Plugins/Experimental/MeshModelingToolsetExp/Source/MeshModelingToolsExp/Private/BakeMeshAttributeMapsTool.cpp b/Engine/Plugins/Experimental/MeshModelingToolsetExp/Source/MeshModelingToolsExp/Private/BakeMeshAttributeMapsTool.cpp
--- a/Engine/Plugins/Experimental/MeshModelingToolsetExp/Source/MeshModelingToolsExp/Private/BakeMeshAttributeMapsTool.cpp
+++ b/Engine/Plugins/Experimental/MeshModelingToolsetExp/Source/MeshModelingToolsExp/Private/BakeMeshAttributeMapsTool.cpp
@@ -548,9 +548,9 @@ void UBakeMeshAttributeMapsTool::UpdateDetailMesh()
 	if (InputMeshSettings->bProjectionInWorldSpace && bIsBakeToSelf == false)
 	{
 		const FTransformSRT3d DetailToWorld = UE::ToolTarget::GetLocalToWorldTransform(DetailTarget);
-		MeshTransforms::ApplyTransform(*DetailMesh, DetailToWorld);
+		MeshTransforms::ApplyTransform(*DetailMesh, DetailToWorld, true);
 		const FTransformSRT3d WorldToBase = UE::ToolTarget::GetLocalToWorldTransform(Targets[0]);
-		MeshTransforms::ApplyTransformInverse(*DetailMesh, WorldToBase);
+		MeshTransforms::ApplyTransformInverse(*DetailMesh, WorldToBase, true);
 	}
 
 	DetailSpatial = MakeShared<FDynamicMeshAABBTree3, ESPMode::ThreadSafe>();
diff --git a/Engine/Plugins/Experimental/MeshModelingToolsetExp/Source/MeshModelingToolsExp/Private/BakeMeshAttributeVertexTool.cpp b/Engine/Plugins/Experimental/MeshModelingToolsetExp/Source/MeshModelingToolsExp/Private/BakeMeshAttributeVertexTool.cpp
--- a/Engine/Plugins/Experimental/MeshModelingToolsetExp/Source/MeshModelingToolsExp/Private/BakeMeshAttributeVertexTool.cpp
+++ b/Engine/Plugins/Experimental/MeshModelingToolsetExp/Source/MeshModelingToolsExp/Private/BakeMeshAttributeVertexTool.cpp
@@ -486,9 +486,9 @@ void UBakeMeshAttributeVertexTool::UpdateDetailMesh()
 	if (InputMeshSettings->bProjectionInWorldSpace && bIsBakeToSelf == false)
 	{
 		const FTransformSRT3d DetailToWorld(DetailComponent->GetWorldTransform());
-		MeshTransforms::ApplyTransform(*DetailMesh, DetailToWorld);
+		MeshTransforms::ApplyTransform(*DetailMesh, DetailToWorld, true);
 		const FTransformSRT3d WorldToBase(TargetComponent->GetWorldTransform());
-		MeshTransforms::ApplyTransformInverse(*DetailMesh, WorldToBase);
+		MeshTransforms::ApplyTransformInverse(*DetailMesh, WorldToBase, true);
 	}
 
 	DetailSpatial = MakeShared<FDynamicMeshAABBTree3, ESPMode::ThreadSafe>();
diff --git a/Engine/Plugins/Experimental/MeshModelingToolsetExp/Source/MeshModelingToolsExp/Private/BakeRenderCaptureTool.cpp b/Engine/Plugins/Experimental/MeshModelingToolsetExp/Source/MeshModelingToolsExp/Private/BakeRenderCaptureTool.cpp
--- a/Engine/Plugins/Experimental/MeshModelingToolsetExp/Source/MeshModelingToolsExp/Private/BakeRenderCaptureTool.cpp
+++ b/Engine/Plugins/Experimental/MeshModelingToolsetExp/Source/MeshModelingToolsExp/Private/BakeRenderCaptureTool.cpp
@@ -633,7 +633,7 @@ void UBakeRenderCaptureTool::Setup()
 		TargetMeshTangents->CopyTriVertexTangents(Mesh);
 		
 		// FMeshSceneAdapter operates in world space, so ensure our mesh transformed to world.
-		MeshTransforms::ApplyTransform(TargetMesh, BaseToWorld);
+		MeshTransforms::ApplyTransform(TargetMesh, BaseToWorld, true);
 		TargetSpatial.SetMesh(&TargetMesh, true);
 	});
 
diff --git a/Engine/Plugins/Experimental/MeshModelingToolsetExp/Source/MeshModelingToolsExp/Private/Physics/SetCollisionGeometryTool.cpp b/Engine/Plugins/Experimental/MeshModelingToolsetExp/Source/MeshModelingToolsExp/Private/Physics/SetCollisionGeometryTool.cpp
--- a/Engine/Plugins/Experimental/MeshModelingToolsetExp/Source/MeshModelingToolsExp/Private/Physics/SetCollisionGeometryTool.cpp
+++ b/Engine/Plugins/Experimental/MeshModelingToolsetExp/Source/MeshModelingToolsExp/Private/Physics/SetCollisionGeometryTool.cpp
@@ -616,8 +616,8 @@ void USetCollisionGeometryTool::PrecomputeInputMeshes()
 		if (Settings->bUseWorldSpace)
 		{
 			FTransformSRT3d ToWorld(UE::ToolTarget::GetLocalToWorldTransform(Targets[k]));
-			MeshTransforms::ApplyTransform(SourceMesh, ToWorld);
-			MeshTransforms::ApplyTransformInverse(SourceMesh, TargetTransform);
+			MeshTransforms::ApplyTransform(SourceMesh, ToWorld, true);
+			MeshTransforms::ApplyTransformInverse(SourceMesh, TargetTransform, true);
 		}
 		SourceMesh.DiscardAttributes();
 		InputMeshes[k] = MakeShared<FDynamicMesh3, ESPMode::ThreadSafe>(MoveTemp(SourceMesh));
diff --git a/Engine/Plugins/Experimental/PlanarCutPlugin/Source/PlanarCut/Private/GeometryMeshConversion.cpp b/Engine/Plugins/Experimental/PlanarCutPlugin/Source/PlanarCut/Private/GeometryMeshConversion.cpp
--- a/Engine/Plugins/Experimental/PlanarCutPlugin/Source/PlanarCut/Private/GeometryMeshConversion.cpp
+++ b/Engine/Plugins/Experimental/PlanarCutPlugin/Source/PlanarCut/Private/GeometryMeshConversion.cpp
@@ -1556,7 +1556,7 @@ FCellMeshes::FCellMeshes(int32 NumUVLayersIn, const FDynamicMesh3& SingleCutter,
 
 	if (Transform.IsSet())
 	{
-		MeshTransforms::ApplyTransform(CellMeshes[0].AugMesh, FTransformSRT3d(Transform.GetValue()));
+		MeshTransforms::ApplyTransform(CellMeshes[0].AugMesh, FTransformSRT3d(Transform.GetValue()), true);
 	}
 
 	// Mesh should already be augmented
diff --git a/Engine/Plugins/Runtime/GeometryProcessing/Source/DynamicMesh/Private/Operations/MeshBoolean.cpp b/Engine/Plugins/Runtime/GeometryProcessing/Source/DynamicMesh/Private/Operations/MeshBoolean.cpp
--- a/Engine/Plugins/Runtime/GeometryProcessing/Source/DynamicMesh/Private/Operations/MeshBoolean.cpp
+++ b/Engine/Plugins/Runtime/GeometryProcessing/Source/DynamicMesh/Private/Operations/MeshBoolean.cpp
@@ -224,11 +224,7 @@ bool FMeshBoolean::Compute()
 		FTransformSRT3d CenteredTransform = Transforms[MeshIdx];
 		CenteredTransform.SetTranslation(ScaleFactor*(CenteredTransform.GetTranslation() - CombinedAABB.Center()));
 		CenteredTransform.SetScale(ScaleFactor*CenteredTransform.GetScale());
-		MeshTransforms::ApplyTransform(*CutMesh[MeshIdx], CenteredTransform);
-		if (CenteredTransform.GetDeterminant() < 0)
-		{
-			CutMesh[MeshIdx]->ReverseOrientation(false);
-		}
+		MeshTransforms::ApplyTransform(*CutMesh[MeshIdx], CenteredTransform, true);
 	}
 	ResultTransform = FTransformSRT3d(CombinedAABB.Center());
 	ResultTransform.SetScale(FVector3d::One() * (1.0 / ScaleFactor));
diff --git a/Engine/Plugins/Runtime/GeometryProcessing/Source/DynamicMesh/Private/ShapeApproximation/SimpleShapeSet3.cpp b/Engine/Plugins/Runtime/GeometryProcessing/Source/DynamicMesh/Private/ShapeApproximation/SimpleShapeSet3.cpp
--- a/Engine/Plugins/Runtime/GeometryProcessing/Source/DynamicMesh/Private/ShapeApproximation/SimpleShapeSet3.cpp
+++ b/Engine/Plugins/Runtime/GeometryProcessing/Source/DynamicMesh/Private/ShapeApproximation/SimpleShapeSet3.cpp
@@ -537,7 +537,7 @@ void FSimpleShapeSet3d::Append(const FSimpleShapeSet3d& OtherShapeSet, const FTr
 	for (const FConvexShape3d& ConvexShape : OtherShapeSet.Convexes)
 	{
 		Convexes.Add(ConvexShape);
-		MeshTransforms::ApplyTransform(Convexes.Last().Mesh, Transform);
+		MeshTransforms::ApplyTransform(Convexes.Last().Mesh, Transform, true);
 	}
 
 	LevelSets.Reserve(LevelSets.Num() + OtherShapeSet.LevelSets.Num());
@@ -576,7 +576,7 @@ void FSimpleShapeSet3d::Append(const FSimpleShapeSet3d& OtherShapeSet, const TAr
 		Convexes.Add(ConvexShape);
 		for (const FTransform3d& XForm : TransformSequence)
 		{
-			MeshTransforms::ApplyTransform(Convexes.Last().Mesh, XForm);
+			MeshTransforms::ApplyTransform(Convexes.Last().Mesh, XForm, true);
 		}
 	}
 
@@ -653,7 +653,7 @@ void FSimpleShapeSet3d::ApplyTransform(const FTransform3d& Transform)
 
 	for (FConvexShape3d& ConvexShape : Convexes)
 	{
-		MeshTransforms::ApplyTransform(ConvexShape.Mesh, Transform);
+		MeshTransforms::ApplyTransform(ConvexShape.Mesh, Transform, true);
 	}
 
 	for (FLevelSetShape3d& LevelSetShape : LevelSets)
diff --git a/Engine/Plugins/Runtime/MeshModelingToolset/Source/MeshModelingTools/Private/CombineMeshesTool.cpp b/Engine/Plugins/Runtime/MeshModelingToolset/Source/MeshModelingTools/Private/CombineMeshesTool.cpp
--- a/Engine/Plugins/Runtime/MeshModelingToolset/Source/MeshModelingTools/Private/CombineMeshesTool.cpp
+++ b/Engine/Plugins/Runtime/MeshModelingToolset/Source/MeshModelingTools/Private/CombineMeshesTool.cpp
@@ -397,16 +397,8 @@ void UCombineMeshesTool::UpdateExistingAsset()
 			if (ComponentIdx != SkipIndex)
 			{
 				FTransform3d ComponentToWorld = (FTransform3d)UE::ToolTarget::GetLocalToWorldTransform(Targets[ComponentIdx]);
-				MeshTransforms::ApplyTransform(ComponentDMesh, ComponentToWorld);
-				if (ComponentToWorld.GetDeterminant() < 0)
-				{
-					ComponentDMesh.ReverseOrientation(true);
-				}
-				MeshTransforms::ApplyTransformInverse(ComponentDMesh, TargetToWorld);
-				if (TargetToWorld.GetDeterminant() < 0)
-				{
-					ComponentDMesh.ReverseOrientation(true);
-				}
+				MeshTransforms::ApplyTransform(ComponentDMesh, ComponentToWorld, true);
+				MeshTransforms::ApplyTransformInverse(ComponentDMesh, TargetToWorld, true);
 				Transforms[0] = ComponentToWorld;
 				if (TargetToWorld.GetRotation().IsIdentity() || TargetToWorld.GetScale3D().IsUniform())
 				{
diff --git a/Engine/Plugins/Runtime/MeshModelingToolset/Source/MeshModelingTools/Private/CutMeshWithMeshTool.cpp b/Engine/Plugins/Runtime/MeshModelingToolset/Source/MeshModelingTools/Private/CutMeshWithMeshTool.cpp
--- a/Engine/Plugins/Runtime/MeshModelingToolset/Source/MeshModelingTools/Private/CutMeshWithMeshTool.cpp
+++ b/Engine/Plugins/Runtime/MeshModelingToolset/Source/MeshModelingTools/Private/CutMeshWithMeshTool.cpp
@@ -353,8 +353,8 @@ void UCutMeshWithMeshTool::OnShutdown(EToolShutdownType ShutdownType)
 		{
 			if (OpResult.Mesh->TriangleCount() > 0)
 			{
-				MeshTransforms::ApplyTransform(*OpResult.Mesh, OpResult.Transform);
-				MeshTransforms::ApplyTransformInverse(*OpResult.Mesh, TargetToWorld);
+				MeshTransforms::ApplyTransform(*OpResult.Mesh, OpResult.Transform, true);
+				MeshTransforms::ApplyTransformInverse(*OpResult.Mesh, TargetToWorld, true);
 				UE::ToolTarget::CommitMeshDescriptionUpdateViaDynamicMesh(Targets[0], *OpResult.Mesh, true);
 				Cast<IMaterialProvider>(Targets[0])->CommitMaterialSetUpdate(MaterialSet, true);
 			}
@@ -364,8 +364,8 @@ void UCutMeshWithMeshTool::OnShutdown(EToolShutdownType ShutdownType)
 		// create intersection asset
 		if ( IntersectionMesh.TriangleCount() > 0)
 		{
-			MeshTransforms::ApplyTransform(IntersectionMesh, OpResult.Transform);
-			MeshTransforms::ApplyTransformInverse(IntersectionMesh, TargetToWorld);
+			MeshTransforms::ApplyTransform(IntersectionMesh, OpResult.Transform, true);
+			MeshTransforms::ApplyTransformInverse(IntersectionMesh, TargetToWorld, true);
 			FTransform3d NewTransform = TargetToWorld;
 
 			FString CurName = UE::Modeling::GetComponentAssetBaseName(UE::ToolTarget::GetTargetComponent(Targets[0]));
diff --git a/Engine/Plugins/Runtime/MeshModelingToolset/Source/MeshModelingTools/Private/ToolActivities/PolyEditExtrudeActivity.cpp b/Engine/Plugins/Runtime/MeshModelingToolset/Source/MeshModelingTools/Private/ToolActivities/PolyEditExtrudeActivity.cpp
--- a/Engine/Plugins/Runtime/MeshModelingToolset/Source/MeshModelingTools/Private/ToolActivities/PolyEditExtrudeActivity.cpp
+++ b/Engine/Plugins/Runtime/MeshModelingToolset/Source/MeshModelingTools/Private/ToolActivities/PolyEditExtrudeActivity.cpp
@@ -394,7 +394,7 @@ void UPolyEditExtrudeActivity::ReinitializeExtrudeHeightMechanic()
 	// the extrude measurement line.
 	ExtrudeHitTargetMesh = *PatchMesh;
 	MeshTransforms::ApplyTransform(ExtrudeHitTargetMesh, 
-		(FTransform3d)ActivityContext->Preview->PreviewMesh->GetTransform());
+		(FTransform3d)ActivityContext->Preview->PreviewMesh->GetTransform(), true);
 
 	double Length = 99999.0;
 	FVector3d ExtrudeDirection = GetExtrudeDirection();
diff --git a/Engine/Plugins/Runtime/MeshModelingToolset/Source/ModelingComponents/Private/BaseTools/BaseCreateFromSelectedTool.cpp b/Engine/Plugins/Runtime/MeshModelingToolset/Source/ModelingComponents/Private/BaseTools/BaseCreateFromSelectedTool.cpp
--- a/Engine/Plugins/Runtime/MeshModelingToolset/Source/ModelingComponents/Private/BaseTools/BaseCreateFromSelectedTool.cpp
+++ b/Engine/Plugins/Runtime/MeshModelingToolset/Source/ModelingComponents/Private/BaseTools/BaseCreateFromSelectedTool.cpp
@@ -188,7 +188,7 @@ void UBaseCreateFromSelectedTool::GenerateAsset(const FDynamicMeshOpResult& OpRe
 	if (Targets.Num() == 1) // in the single-selection case, shove the result back into the original component space
 	{
 		FTransform3d FromSourceComponentSpace = UE::ToolTarget::GetLocalToWorldTransform(Targets[0]);
-		MeshTransforms::ApplyTransformInverse(*OpResult.Mesh, FromSourceComponentSpace);
+		MeshTransforms::ApplyTransformInverse(*OpResult.Mesh, FromSourceComponentSpace, true);
 		NewTransform = UE::ToolTarget::GetLocalToWorldTransform(Targets[0]);
 	}
 	else // in the multi-selection case, center the pivot for the combined result
@@ -197,7 +197,7 @@ void UBaseCreateFromSelectedTool::GenerateAsset(const FDynamicMeshOpResult& OpRe
 		double Rescale = OpResult.Transform.GetScale().X;
 		FTransform3d LocalTransform(-Center * Rescale);
 		LocalTransform.SetScale3D(FVector3d(Rescale, Rescale, Rescale));
-		MeshTransforms::ApplyTransform(*OpResult.Mesh, LocalTransform);
+		MeshTransforms::ApplyTransform(*OpResult.Mesh, LocalTransform, true);
 		NewTransform = OpResult.Transform;
 		NewTransform.SetScale3D(FVector3d::One());
 		NewTransform.SetTranslation(NewTransform.GetTranslation() + NewTransform.TransformVector(Center * Rescale));
@@ -240,8 +240,8 @@ void UBaseCreateFromSelectedTool::UpdateAsset(const FDynamicMeshOpResult& Result
 	FTransform3d TargetToWorld = UE::ToolTarget::GetLocalToWorldTransform(UpdateTarget);
 
 	FTransform3d ResultTransform = Result.Transform;
-	MeshTransforms::ApplyTransform(*Result.Mesh, ResultTransform);
-	MeshTransforms::ApplyTransformInverse(*Result.Mesh, TargetToWorld);
+	MeshTransforms::ApplyTransform(*Result.Mesh, ResultTransform, true);
+	MeshTransforms::ApplyTransformInverse(*Result.Mesh, TargetToWorld, true);
 
 	UE::ToolTarget::CommitMeshDescriptionUpdateViaDynamicMesh(UpdateTarget, *Result.Mesh, true);
 
diff --git a/Engine/Plugins/Runtime/MeshModelingToolset/Source/ModelingComponents/Private/Components/OctreeDynamicMeshComponent.cpp b/Engine/Plugins/Runtime/MeshModelingToolset/Source/ModelingComponents/Private/Components/OctreeDynamicMeshComponent.cpp
--- a/Engine/Plugins/Runtime/MeshModelingToolset/Source/ModelingComponents/Private/Components/OctreeDynamicMeshComponent.cpp
+++ b/Engine/Plugins/Runtime/MeshModelingToolset/Source/ModelingComponents/Private/Components/OctreeDynamicMeshComponent.cpp
@@ -79,11 +79,11 @@ void UOctreeDynamicMeshComponent::ApplyTransform(const FTransform3d& Transform,
 {
 	if (bInvert)
 	{
-		MeshTransforms::ApplyTransformInverse(*GetMesh(), Transform);
+		MeshTransforms::ApplyTransformInverse(*GetMesh(), Transform, true);
 	}
 	else
 	{
-		MeshTransforms::ApplyTransform(*GetMesh(), Transform);
+		MeshTransforms::ApplyTransform(*GetMesh(), Transform, true);
 	}
 
 	if (GetCurrentSceneProxy() != nullptr)
diff --git a/Engine/Plugins/Runtime/MeshModelingToolset/Source/ModelingComponents/Private/ConversionUtils/VolumeToDynamicMesh.cpp b/Engine/Plugins/Runtime/MeshModelingToolset/Source/ModelingComponents/Private/ConversionUtils/VolumeToDynamicMesh.cpp
--- a/Engine/Plugins/Runtime/MeshModelingToolset/Source/ModelingComponents/Private/ConversionUtils/VolumeToDynamicMesh.cpp
+++ b/Engine/Plugins/Runtime/MeshModelingToolset/Source/ModelingComponents/Private/ConversionUtils/VolumeToDynamicMesh.cpp
@@ -39,7 +39,7 @@ void VolumeToDynamicMesh(AVolume* Volume, FDynamicMesh3& Mesh,
 	if (Options.bInWorldSpace)
 	{
 		FTransform3d XForm(Volume->GetTransform());
-		MeshTransforms::ApplyTransform(Mesh, XForm);
+		MeshTransforms::ApplyTransform(Mesh, XForm, true);
 	}
 }
 
@@ -57,7 +57,7 @@ void BrushComponentToDynamicMesh(UBrushComponent* Component, UE::Geometry::FDyna
 	if (Options.bInWorldSpace)
 	{
 		FTransform3d XForm(Component->GetComponentTransform());
-		MeshTransforms::ApplyTransform(Mesh, XForm);
+		MeshTransforms::ApplyTransform(Mesh, XForm, true);
 	}
 }
 
diff --git a/Engine/Plugins/Runtime/MeshModelingToolset/Source/ModelingComponents/Private/Drawing/PolyEditPreviewMesh.cpp b/Engine/Plugins/Runtime/MeshModelingToolset/Source/ModelingComponents/Private/Drawing/PolyEditPreviewMesh.cpp
--- a/Engine/Plugins/Runtime/MeshModelingToolset/Source/ModelingComponents/Private/Drawing/PolyEditPreviewMesh.cpp
+++ b/Engine/Plugins/Runtime/MeshModelingToolset/Source/ModelingComponents/Private/Drawing/PolyEditPreviewMesh.cpp
@@ -28,7 +28,7 @@ void UPolyEditPreviewMesh::InitializeExtrudeType(
 	if (bHaveMeshTransform)
 	{
 		MeshTransform = *MeshTransformIn;
-		MeshTransforms::ApplyTransform(EditPatch, MeshTransform);
+		MeshTransforms::ApplyTransform(EditPatch, MeshTransform, true);
 	}
 	//FMeshNormals::QuickComputeVertexNormals(EditPatch);
 
@@ -99,7 +99,7 @@ void UPolyEditPreviewMesh::InitializeExtrudeType(FDynamicMesh3&& BaseMesh,
 	if (bHaveMeshTransform)
 	{
 		MeshTransform = *MeshTransformIn;
-		MeshTransforms::ApplyTransform(InitialEditPatch, MeshTransform);
+		MeshTransforms::ApplyTransform(InitialEditPatch, MeshTransform, true);
 	}
 
 	// save copy of initial patch
@@ -254,7 +254,7 @@ void UPolyEditPreviewMesh::InitializeInsetType(const FDynamicMesh3* SourceMesh,
 	if (bHaveMeshTransform)
 	{
 		MeshTransform = *MeshTransformIn;
-		MeshTransforms::ApplyTransform(EditPatch, MeshTransform);
+		MeshTransforms::ApplyTransform(EditPatch, MeshTransform, true);
 	}
 
 	// save copy of initial patch
@@ -310,7 +310,7 @@ void UPolyEditPreviewMesh::InitializeStaticType(const FDynamicMesh3* SourceMesh,
 	if (bHaveMeshTransform)
 	{
 		MeshTransform = *MeshTransformIn;
-		MeshTransforms::ApplyTransform(EditPatch, MeshTransform);
+		MeshTransforms::ApplyTransform(EditPatch, MeshTransform, true);
 	}
 
 	// save copy of initial patch
diff --git a/Engine/Plugins/Runtime/MeshModelingToolset/Source/ModelingOperators/Private/CompositionOps/CubeGridBooleanOp.cpp b/Engine/Plugins/Runtime/MeshModelingToolset/Source/ModelingOperators/Private/CompositionOps/CubeGridBooleanOp.cpp
--- a/Engine/Plugins/Runtime/MeshModelingToolset/Source/ModelingOperators/Private/CompositionOps/CubeGridBooleanOp.cpp
+++ b/Engine/Plugins/Runtime/MeshModelingToolset/Source/ModelingOperators/Private/CompositionOps/CubeGridBooleanOp.cpp
@@ -548,7 +548,7 @@ void FCubeGridBooleanOp::CalculateResult(FProgressCancel* Progress)
 		FVector3d Center = ResultMesh->GetBounds().Center();
 		ResultTransform = FTransformSRT3d(Center);
 	}
-	MeshTransforms::ApplyTransformInverse(*ResultMesh, InputTransform);
+	MeshTransforms::ApplyTransformInverse(*ResultMesh, InputTransform, true);
 
 	// TODO: Is it worth trying to fix holes like BooleanMeshesOp tries to?
 
diff --git a/Engine/Plugins/Runtime/MeshModelingToolset/Source/ModelingOperators/Private/CompositionOps/VoxelBlendMeshesOp.cpp b/Engine/Plugins/Runtime/MeshModelingToolset/Source/ModelingOperators/Private/CompositionOps/VoxelBlendMeshesOp.cpp
--- a/Engine/Plugins/Runtime/MeshModelingToolset/Source/ModelingOperators/Private/CompositionOps/VoxelBlendMeshesOp.cpp
+++ b/Engine/Plugins/Runtime/MeshModelingToolset/Source/ModelingOperators/Private/CompositionOps/VoxelBlendMeshesOp.cpp
@@ -44,11 +44,7 @@ void FVoxelBlendMeshesOp::CalculateResult(FProgressCancel* Progress)
 			continue;
 		}
 		
-		if (Transforms[MeshIdx].GetDeterminant() < 0)
-		{
-			TransformedMeshes[MeshIdx].ReverseOrientation(false);
-		}
-		MeshTransforms::ApplyTransform(TransformedMeshes[MeshIdx], (FTransformSRT3d)Transforms[MeshIdx]);
+		MeshTransforms::ApplyTransform(TransformedMeshes[MeshIdx], (FTransformSRT3d)Transforms[MeshIdx], true);
 
 		if (bVoxWrap)
 		{
