Handle mesh inversions due to transforms in ComponentCollisionUtil + add missing transform calls

diff --git a/Engine/Plugins/Experimental/GeometryScripting/Source/GeometryScriptingCore/Private/SceneUtilityFunctions.cpp b/Engine/Plugins/Experimental/GeometryScripting/Source/GeometryScriptingCore/Private/SceneUtilityFunctions.cpp
--- a/Engine/Plugins/Experimental/GeometryScripting/Source/GeometryScriptingCore/Private/SceneUtilityFunctions.cpp
+++ b/Engine/Plugins/Experimental/GeometryScripting/Source/GeometryScriptingCore/Private/SceneUtilityFunctions.cpp
@@ -160,24 +160,23 @@ UDynamicMesh* UGeometryScriptLibrary_SceneUtilityFunctions::CopyCollisionMeshesF
 
 	if (bUseComplexCollision)
 	{
-		// rewrite object type if necessary
+		// find the Complex Collision mesh interface
+		IInterface_CollisionDataProvider* CollisionProvider = nullptr;
 		if (UStaticMeshComponent* StaticMeshComp = Cast<UStaticMeshComponent>(FromObject))
 		{
-			FromObject = StaticMeshComp->GetStaticMesh();
+			LocalToWorld = StaticMeshComp->GetComponentTransform();
+			UStaticMesh* StaticMesh = StaticMeshComp->GetStaticMesh();
+			CollisionProvider = (IInterface_CollisionDataProvider*)StaticMesh;
 		}
-
-		// find the Complex Collision mesh interface
-		IInterface_CollisionDataProvider* CollisionProvider = nullptr;
-		if (UStaticMesh* StaticMesh = Cast<UStaticMesh>(FromObject))
+		else if (UStaticMesh* StaticMesh = Cast<UStaticMesh>(FromObject))
 		{
 			LocalToWorld = FTransform::Identity;
 			CollisionProvider = (IInterface_CollisionDataProvider*)StaticMesh;
 		}
-		else if ( Cast<UDynamicMeshComponent>(FromObject) != nullptr )
+		else if (UPrimitiveComponent* DynamicMeshComp = Cast<UDynamicMeshComponent>(FromObject))
 		{
-			UPrimitiveComponent* AnyComponent = Cast<UPrimitiveComponent>(FromObject);
-			LocalToWorld = AnyComponent->GetComponentTransform();
-			CollisionProvider = (IInterface_CollisionDataProvider*)AnyComponent;
+			LocalToWorld = DynamicMeshComp->GetComponentTransform();
+			CollisionProvider = (IInterface_CollisionDataProvider*)DynamicMeshComp;
 		}
 
 		if (CollisionProvider == nullptr)
diff --git a/Engine/Plugins/Runtime/MeshModelingToolset/Source/ModelingComponents/Private/Physics/ComponentCollisionUtil.cpp b/Engine/Plugins/Runtime/MeshModelingToolset/Source/ModelingComponents/Private/Physics/ComponentCollisionUtil.cpp
--- a/Engine/Plugins/Runtime/MeshModelingToolset/Source/ModelingComponents/Private/Physics/ComponentCollisionUtil.cpp
+++ b/Engine/Plugins/Runtime/MeshModelingToolset/Source/ModelingComponents/Private/Physics/ComponentCollisionUtil.cpp
@@ -196,6 +196,8 @@ void UE::Geometry::ConvertSimpleCollisionToMeshes(
 {
 	FDynamicMeshEditor Editor(&MeshOut);
 
+	bool bTransformInverts = TransformSeqeuence.WillInvert();
+
 	for (const FKSphereElem& Sphere : AggGeom.SphereElems)
 	{
 		FSphereGenerator SphereGen;
@@ -207,6 +209,11 @@ void UE::Geometry::ConvertSimpleCollisionToMeshes(
 
 		MeshTransforms::Translate(SphereMesh, FVector3d(Sphere.Center));
 
+		if (bTransformInverts)
+		{
+			SphereMesh.ReverseOrientation(false);
+		}
+
 		FMeshIndexMappings Mappings;
 		Editor.AppendMesh(&SphereMesh, Mappings,
 			[&](int32 vid, const FVector3d& P) { return TransformSeqeuence.TransformPosition(P); },
@@ -229,6 +236,11 @@ void UE::Geometry::ConvertSimpleCollisionToMeshes(
 
 		// transform not applied because it is just the Center/Rotation
 
+		if (bTransformInverts)
+		{
+			BoxMesh.ReverseOrientation(false);
+		}
+
 		FMeshIndexMappings Mappings;
 		Editor.AppendMesh(&BoxMesh, Mappings,
 			[&](int32 vid, const FVector3d& P) { return TransformSeqeuence.TransformPosition(P); },
@@ -255,7 +267,11 @@ void UE::Geometry::ConvertSimpleCollisionToMeshes(
 		MeshTransforms::Translate(CapsuleMesh, FVector3d(0,0,-0.5*Capsule.Length) );
 
 		FTransformSRT3d Transform(Capsule.GetTransform());
-		MeshTransforms::ApplyTransform(CapsuleMesh, Transform);
+		MeshTransforms::ApplyTransform(CapsuleMesh, Transform, false);
+		if (Transform.GetDeterminant() < 0 != bTransformInverts)
+		{
+			CapsuleMesh.ReverseOrientation(false);
+		}
 
 		FMeshIndexMappings Mappings;
 		Editor.AppendMesh(&CapsuleMesh, Mappings,
@@ -276,15 +292,20 @@ void UE::Geometry::ConvertSimpleCollisionToMeshes(
 		int32 NumVertices = Convex.VertexData.Num();
 		for (int32 k = 0; k < NumVertices; ++k)
 		{
-			ConvexMesh.AppendVertex( FVector3d(Convex.VertexData[k]) );
+			ConvexMesh.AppendVertex(ElemTransform.TransformPosition(FVector3d(Convex.VertexData[k])) );
 		}
 		int32 NumTriangles = Convex.IndexData.Num() / 3;
 		for (int32 k = 0; k < NumTriangles; ++k)
 		{
 			ConvexMesh.AppendTriangle(Convex.IndexData[3*k], Convex.IndexData[3*k+1], Convex.IndexData[3*k+2]);
 		}
-
-		ConvexMesh.ReverseOrientation();
+		
+		// Note we need to reverse the orientation if no transform inverts, or if both invert,
+		// because ConvexMesh has reversed-orientation triangles normally
+		if (ElemTransform.GetDeterminant() < 0 == bTransformInverts)
+		{
+			ConvexMesh.ReverseOrientation();
+		}
 		ConvexMesh.EnableTriangleGroups(0);
 		ConvexMesh.EnableAttributes();
 		if (bInitializeConvexUVs)
@@ -333,10 +354,11 @@ bool UE::Geometry::ConvertComplexCollisionToMeshes(
 		FTriMeshCollisionData CollisionData;
 		if (CollisionProvider->GetPhysicsTriMeshData(&CollisionData, true))
 		{
+			bool bWillInvert = TransformSeqeuence.WillInvert();
 			TArray<int32> VertexIDMap;
 			for (int32 k = 0; k < CollisionData.Vertices.Num(); ++k)
 			{
-				int32 VertexID = MeshOut.AppendVertex((FVector)CollisionData.Vertices[k]);
+				int32 VertexID = MeshOut.AppendVertex(TransformSeqeuence.TransformPosition((FVector)CollisionData.Vertices[k]));
 				VertexIDMap.Add(VertexID);
 			}
 			for (const FTriIndices& TriIndices : CollisionData.Indices)
@@ -347,9 +369,9 @@ bool UE::Geometry::ConvertComplexCollisionToMeshes(
 				{
 					bFoundMeshErrorsOut = true;
 					// create new vertices for this triangle
-					int32 A = MeshOut.AppendVertex(MeshOut.GetVertex(TriIndices.v0));
-					int32 B = MeshOut.AppendVertex(MeshOut.GetVertex(TriIndices.v1));
-					int32 C = MeshOut.AppendVertex(MeshOut.GetVertex(TriIndices.v2));
+					int32 A = MeshOut.AppendVertex(TransformSeqeuence.TransformPosition(MeshOut.GetVertex(TriIndices.v0)));
+					int32 B = MeshOut.AppendVertex(TransformSeqeuence.TransformPosition(MeshOut.GetVertex(TriIndices.v1)));
+					int32 C = MeshOut.AppendVertex(TransformSeqeuence.TransformPosition(MeshOut.GetVertex(TriIndices.v2)));
 					MeshOut.AppendTriangle(FIndex3i(A,B,C));
 				}
 			}
@@ -363,7 +385,7 @@ bool UE::Geometry::ConvertComplexCollisionToMeshes(
 				Weld.Apply();
 			}
 
-			if (!CollisionData.bFlipNormals)		// collision mesh has reversed orientation
+			if (!CollisionData.bFlipNormals != bWillInvert)		// collision mesh has reversed orientation
 			{
 				MeshOut.ReverseOrientation(false);
 			}
diff --git a/Engine/Source/Runtime/GeometryCore/Public/TransformSequence.h b/Engine/Source/Runtime/GeometryCore/Public/TransformSequence.h
--- a/Engine/Source/Runtime/GeometryCore/Public/TransformSequence.h
+++ b/Engine/Source/Runtime/GeometryCore/Public/TransformSequence.h
@@ -78,6 +78,19 @@ public:
 		return FinalScale;
 	}
 
+	/**
+	 * @return Whether the sequence will invert a shape (by negative scaling) when applied
+	 */
+	bool WillInvert() const
+	{
+		RealType Det = 1;
+		for (const TTransformSRT3<RealType>& Transform : Transforms)
+		{
+			Det *= Transform.GetDeterminant();
+		}
+		return Det < 0;
+	}
+
 	/**
 	 * Set scales of all transforms to (1,1,1)
 	 */
