add F-to-focus support to 'add primitive' tools

diff --git a/Engine/Plugins/Runtime/MeshModelingToolset/Source/MeshModelingTools/Private/AddPrimitiveTool.cpp b/Engine/Plugins/Runtime/MeshModelingToolset/Source/MeshModelingTools/Private/AddPrimitiveTool.cpp
--- a/Engine/Plugins/Runtime/MeshModelingToolset/Source/MeshModelingTools/Private/AddPrimitiveTool.cpp
+++ b/Engine/Plugins/Runtime/MeshModelingToolset/Source/MeshModelingTools/Private/AddPrimitiveTool.cpp
@@ -362,6 +362,45 @@ void UAddPrimitiveTool::UpdatePreviewMesh() const
 	checkSlow(CalculateTangentsSuccessful);
 }
 
+
+bool UAddPrimitiveTool::SupportsWorldSpaceFocusBox()
+{
+	return !PreviewMesh.IsNull();
+}
+
+FBox UAddPrimitiveTool::GetWorldSpaceFocusBox()
+{
+	if (PreviewMesh)
+	{
+		if (UPrimitiveComponent* Component = PreviewMesh->GetRootComponent())
+		{
+			return Component->Bounds.GetBox();
+		}
+	}
+	return FBox();
+}
+
+bool UAddPrimitiveTool::SupportsWorldSpaceFocusPoint()
+{
+	return !PreviewMesh.IsNull();
+
+}
+
+bool UAddPrimitiveTool::GetWorldSpaceFocusPoint(const FRay& WorldRay, FVector& PointOut)
+{
+	if (PreviewMesh)
+	{
+		FHitResult HitResult;
+		if (PreviewMesh->FindRayIntersection(WorldRay, HitResult))
+		{
+			PointOut = HitResult.ImpactPoint;
+			return true;
+		}
+	}
+	return false;
+}
+
+
 FInputRayHit UAddPrimitiveTool::IsHitByClick(const FInputDeviceRay& ClickPos)
 {
 	FInputRayHit Result(0);
diff --git a/Engine/Plugins/Runtime/MeshModelingToolset/Source/MeshModelingTools/Public/AddPrimitiveTool.h b/Engine/Plugins/Runtime/MeshModelingToolset/Source/MeshModelingTools/Public/AddPrimitiveTool.h
--- a/Engine/Plugins/Runtime/MeshModelingToolset/Source/MeshModelingTools/Public/AddPrimitiveTool.h
+++ b/Engine/Plugins/Runtime/MeshModelingToolset/Source/MeshModelingTools/Public/AddPrimitiveTool.h
@@ -9,6 +9,7 @@
 #include "Properties/MeshMaterialProperties.h"
 #include "PropertySets/CreateMeshObjectTypeProperties.h"
 #include "UObject/NoExportTypes.h"
+#include "InteractiveToolQueryInterfaces.h"
 
 #include "AddPrimitiveTool.generated.h"
 
@@ -431,7 +432,7 @@ public:
  * Base tool to create primitives
  */
 UCLASS()
-class MESHMODELINGTOOLS_API UAddPrimitiveTool : public USingleClickTool, public IHoverBehaviorTarget
+class MESHMODELINGTOOLS_API UAddPrimitiveTool : public USingleClickTool, public IHoverBehaviorTarget, public IInteractiveToolCameraFocusAPI
 {
 	GENERATED_BODY()
 
@@ -461,6 +462,12 @@ public:
 	virtual bool OnUpdateHover(const FInputDeviceRay& DevicePos) override;
 	virtual void OnEndHover() override;
 
+	// IInteractiveToolCameraFocusAPI implementation
+	virtual bool SupportsWorldSpaceFocusBox() override;
+	virtual FBox GetWorldSpaceFocusBox() override;
+	virtual bool SupportsWorldSpaceFocusPoint() override;
+	virtual bool GetWorldSpaceFocusPoint(const FRay& WorldRay, FVector& PointOut) override;
+
 
 protected:
 	enum class EState
