fix issues with scaled input meshes in edit polygons tool by baking the transform

diff --git a/Engine/Plugins/Runtime/MeshModelingToolset/Source/MeshModelingTools/Private/EditMeshPolygonsTool.cpp b/Engine/Plugins/Runtime/MeshModelingToolset/Source/MeshModelingTools/Private/EditMeshPolygonsTool.cpp
--- a/Engine/Plugins/Runtime/MeshModelingToolset/Source/MeshModelingTools/Private/EditMeshPolygonsTool.cpp
+++ b/Engine/Plugins/Runtime/MeshModelingToolset/Source/MeshModelingTools/Private/EditMeshPolygonsTool.cpp
@@ -232,6 +232,12 @@ void UEditMeshPolygonsTool::Setup()
 	int32 MaxEdges = CVarEdgeLimit.GetValueOnGameThread();
 
 	CurrentMesh = MakeShared<FDynamicMesh3>(UE::ToolTarget::GetDynamicMeshCopy(Target));
+	WorldTransform = UE::ToolTarget::GetLocalToWorldTransform(Target);
+	FVector ScaleToBake = WorldTransform.GetScale();
+	BakedTransform = FTransformSRT3d(FQuaterniond::Identity(), FVector::Zero(), ScaleToBake);
+	WorldTransform.SetScale(FVector::One());
+	MeshTransforms::ApplyTransform(*CurrentMesh, BakedTransform, true);
+
 	if (bTriangleMode)
 	{
 		bToolDisabled = CurrentMesh->EdgeCount() > MaxEdges;
@@ -408,7 +414,6 @@ void UEditMeshPolygonsTool::Setup()
 	Preview = NewObject<UMeshOpPreviewWithBackgroundCompute>();
 	Preview->Setup(TargetWorld);
 	ToolSetupUtil::ApplyRenderingConfigurationToPreview(Preview->PreviewMesh, Target); 
-	WorldTransform = UE::ToolTarget::GetLocalToWorldTransform(Target);
 	Preview->PreviewMesh->SetTransform((FTransform)WorldTransform);
 
 	// We'll use the spatial inside preview mesh mainly for the convenience of having it update automatically.
@@ -709,6 +714,7 @@ void UEditMeshPolygonsTool::OnShutdown(EToolShutdownType ShutdownType)
 
 			// Bake CurrentMesh back to target inside an undo transaction
 			GetToolManager()->BeginUndoTransaction(LOCTEXT("EditMeshPolygonsToolTransactionName", "Deform Mesh"));
+			MeshTransforms::ApplyTransformInverse(*CurrentMesh, BakedTransform, true);
 			UE::ToolTarget::CommitDynamicMeshUpdate(Target, *CurrentMesh, bModifiedTriangleTopology);
 
 			UE::Geometry::SetToolOutputSelection(this, OutputSelection);
diff --git a/Engine/Plugins/Runtime/MeshModelingToolset/Source/MeshModelingTools/Public/EditMeshPolygonsTool.h b/Engine/Plugins/Runtime/MeshModelingToolset/Source/MeshModelingTools/Public/EditMeshPolygonsTool.h
--- a/Engine/Plugins/Runtime/MeshModelingToolset/Source/MeshModelingTools/Public/EditMeshPolygonsTool.h
+++ b/Engine/Plugins/Runtime/MeshModelingToolset/Source/MeshModelingTools/Public/EditMeshPolygonsTool.h
@@ -622,7 +622,8 @@ protected:
 	FFrame3d LockedTransfomerFrame;
 	bool bInGizmoDrag = false;
 
-	UE::Geometry::FTransformSRT3d WorldTransform;
+	UE::Geometry::FTransformSRT3d BakedTransform; // We bake the scale part of the Target -> World transform
+	UE::Geometry::FTransformSRT3d WorldTransform; // Transform from Baked to World
 
 	FFrame3d InitialGizmoFrame;
 	FVector3d InitialGizmoScale;
