Removed an ensure that could fail while reading from a valid texture that has no source art

#rn fix geometry

diff --git a/Engine/Plugins/Runtime/MeshModelingToolset/Source/ModelingComponents/Private/AssetUtils/Texture2DUtil.cpp b/Engine/Plugins/Runtime/MeshModelingToolset/Source/ModelingComponents/Private/AssetUtils/Texture2DUtil.cpp
--- a/Engine/Plugins/Runtime/MeshModelingToolset/Source/ModelingComponents/Private/AssetUtils/Texture2DUtil.cpp
+++ b/Engine/Plugins/Runtime/MeshModelingToolset/Source/ModelingComponents/Private/AssetUtils/Texture2DUtil.cpp
@@ -165,11 +165,9 @@ bool UE::AssetUtils::ReadTexture(
 	if (ensure(TextureMap) == false) return false;
 
 #if WITH_EDITOR
-	const bool bHasVTData = TextureMap->GetPlatformData()->VTData != nullptr;
-	const bool bHasMips = TextureMap->GetPlatformData()->Mips.Num() != 0;
-	ensure(bHasVTData != bHasMips);		// should be one or the other
+	const bool bPlatformDataHasMips = TextureMap->GetPlatformData()->Mips.Num() != 0;
 
-	if (TextureMap->Source.IsValid() && (bPreferPlatformData == false || bHasMips == false))
+	if (TextureMap->Source.IsValid() && (bPreferPlatformData == false || bPlatformDataHasMips == false))
 	{
 		return ReadTexture_SourceData(TextureMap, DestImageOut);
 	}
