Fix previewmesh::findrayintersection normal for preview meshes with non-uniform scaling

diff --git a/Engine/Plugins/Runtime/MeshModelingToolset/Source/ModelingComponents/Private/PreviewMesh.cpp b/Engine/Plugins/Runtime/MeshModelingToolset/Source/ModelingComponents/Private/PreviewMesh.cpp
--- a/Engine/Plugins/Runtime/MeshModelingToolset/Source/ModelingComponents/Private/PreviewMesh.cpp
+++ b/Engine/Plugins/Runtime/MeshModelingToolset/Source/ModelingComponents/Private/PreviewMesh.cpp
@@ -411,7 +411,7 @@ bool UPreviewMesh::FindRayIntersection(const FRay3d& WorldRay, FHitResult& HitOu
 {
 	if (IsVisible() && TemporaryParentActor != nullptr && bBuildSpatialDataStructure)
 	{
-		FTransform3d Transform(TemporaryParentActor->GetActorTransform());
+		FTransformSRT3d Transform(TemporaryParentActor->GetActorTransform());
 		FRay3d LocalRay(Transform.InverseTransformPosition(WorldRay.Origin),
 			Transform.InverseTransformVector(WorldRay.Direction));
 		UE::Geometry::Normalize(LocalRay.Direction);
@@ -426,7 +426,7 @@ bool UPreviewMesh::FindRayIntersection(const FRay3d& WorldRay, FHitResult& HitOu
 
 			HitOut.FaceIndex = HitTriID;
 			HitOut.Distance = Query.RayParameter;
-			HitOut.Normal = (FVector)Transform.TransformVectorNoScale(UseMesh->GetTriNormal(HitTriID));
+			HitOut.Normal = (FVector)Transform.TransformNormal(Triangle.Normal());
 			HitOut.ImpactNormal = HitOut.Normal;
 			HitOut.ImpactPoint = (FVector)Transform.TransformPosition(LocalRay.PointAt(Query.RayParameter));
 			return true;
