fix ensure due to out-of-date spatial data after failed bridge in polyedit tool

diff --git a/Engine/Plugins/Runtime/MeshModelingToolset/Source/MeshModelingTools/Private/EditMeshPolygonsTool.cpp b/Engine/Plugins/Runtime/MeshModelingToolset/Source/MeshModelingTools/Private/EditMeshPolygonsTool.cpp
--- a/Engine/Plugins/Runtime/MeshModelingToolset/Source/MeshModelingTools/Private/EditMeshPolygonsTool.cpp
+++ b/Engine/Plugins/Runtime/MeshModelingToolset/Source/MeshModelingTools/Private/EditMeshPolygonsTool.cpp
@@ -2091,6 +2091,8 @@ void UEditMeshPolygonsTool::ApplyBridgeEdges()
 		{
 			Editor.RemoveTriangles(MinimalHoleFiller.NewTriangles, false);
 			GetToolManager()->DisplayMessage(BridgeFailMessage, EToolMessageLevel::UserWarning);
+			// Even though we've manually 'undone' the changes, this will still change mesh timestamps, so we need to register the mesh update
+			UpdateFromCurrentMesh(false);
 			return;
 		}
 		else
