MeshModelingTools: Fix gizmo not showing up when same selection is made via marquee in lattice tool.

diff --git a/Engine/Plugins/Runtime/MeshModelingToolset/Source/ModelingComponents/Private/Mechanics/LatticeControlPointsMechanic.cpp b/Engine/Plugins/Runtime/MeshModelingToolset/Source/ModelingComponents/Private/Mechanics/LatticeControlPointsMechanic.cpp
--- a/Engine/Plugins/Runtime/MeshModelingToolset/Source/ModelingComponents/Private/Mechanics/LatticeControlPointsMechanic.cpp
+++ b/Engine/Plugins/Runtime/MeshModelingToolset/Source/ModelingComponents/Private/Mechanics/LatticeControlPointsMechanic.cpp
@@ -654,9 +654,11 @@ void ULatticeControlPointsMechanic::OnDragRectangleFinished(const FCameraRectang
 			PreDragSelection, false, Transform, Transform, CurrentChangeStamp), LatticePointDeselectionTransactionText);
 		ParentTool->GetToolManager()->EmitObjectChange(this, MakeUnique<FLatticeControlPointsMechanicSelectionChange>(
 			SelectedPointIDs, true, Transform, Transform, CurrentChangeStamp), LatticePointSelectionTransactionText);
-		UpdateGizmoLocation();
+		
 		OnSelectionChanged.Broadcast();
 	}
+	// We hid the gizmo at rectangle start, so it needs updating now.
+	UpdateGizmoLocation();
 
 	ParentTool->GetToolManager()->EndUndoTransaction();
 
