Lattice tool: hide gizmo if only constrained vertices are selected

diff --git a/Engine/Plugins/Experimental/MeshModelingToolsetExp/Source/MeshModelingToolsExp/Private/LatticeDeformerTool.cpp b/Engine/Plugins/Experimental/MeshModelingToolsetExp/Source/MeshModelingToolsExp/Private/LatticeDeformerTool.cpp
--- a/Engine/Plugins/Experimental/MeshModelingToolsetExp/Source/MeshModelingToolsExp/Private/LatticeDeformerTool.cpp
+++ b/Engine/Plugins/Experimental/MeshModelingToolsetExp/Source/MeshModelingToolsExp/Private/LatticeDeformerTool.cpp
@@ -245,6 +245,20 @@ void ULatticeDeformerTool::Setup()
 	ControlPointsMechanic->SetCoordinateSystem(Settings->GizmoCoordinateSystem);
 	ControlPointsMechanic->UpdateSetPivotMode(Settings->bSetPivotMode);
 
+
+	ControlPointsMechanic->ShouldHideGizmo = ULatticeControlPointsMechanic::FShouldHideGizmo::CreateLambda([this]()->bool
+	{
+		for (int32 VID : ControlPointsMechanic->GetSelectedPointIDs())
+		{
+			if (!ConstrainedLatticePoints.Contains(VID))
+			{
+				return false;	// found a selected point that is not constrained
+			}
+		}
+		return true;
+	});
+
+
 	StartPreview();
 }
 
diff --git a/Engine/Plugins/Runtime/MeshModelingToolset/Source/ModelingComponents/Private/Mechanics/LatticeControlPointsMechanic.cpp b/Engine/Plugins/Runtime/MeshModelingToolset/Source/ModelingComponents/Private/Mechanics/LatticeControlPointsMechanic.cpp
--- a/Engine/Plugins/Runtime/MeshModelingToolset/Source/ModelingComponents/Private/Mechanics/LatticeControlPointsMechanic.cpp
+++ b/Engine/Plugins/Runtime/MeshModelingToolset/Source/ModelingComponents/Private/Mechanics/LatticeControlPointsMechanic.cpp
@@ -253,6 +253,8 @@ void ULatticeControlPointsMechanic::RebuildDrawables()
 
 void ULatticeControlPointsMechanic::UpdateDrawables()
 {
+	UpdateGizmoVisibility();
+
 	for (int32 PointIndex = 0; PointIndex < ControlPoints.Num(); ++PointIndex)
 	{
 		const FVector3d& P = ControlPoints[PointIndex];
@@ -466,7 +468,6 @@ void ULatticeControlPointsMechanic::UpdateGizmoLocation()
 
 	if (SelectedPointIDs.Num() == 0)
 	{
-		PointTransformGizmo->SetVisibility(false);
 		PointTransformGizmo->ReinitializeGizmoTransform(FTransform());
 	}
 	else
@@ -476,10 +477,10 @@ void ULatticeControlPointsMechanic::UpdateGizmoLocation()
 			NewGizmoLocation += ControlPoints[PointID];
 		}
 		NewGizmoLocation /= (double)SelectedPointIDs.Num();
-
-		PointTransformGizmo->SetVisibility(true);
 	}
 
+	UpdateGizmoVisibility();
+
 	FTransform NewTransform(FQuat(LocalToWorldTransform.GetRotation()), (FVector)NewGizmoLocation);
 	PointTransformGizmo->ReinitializeGizmoTransform(NewTransform);
 
@@ -488,6 +489,33 @@ void ULatticeControlPointsMechanic::UpdateGizmoLocation()
 	PointTransformGizmo->SetNewChildScale(GizmoScale);
 }
 
+
+void ULatticeControlPointsMechanic::UpdateGizmoVisibility()
+{
+	if (!PointTransformGizmo)
+	{
+		return;
+	}
+
+	if (SelectedPointIDs.Num() == 0)
+	{
+		PointTransformGizmo->SetVisibility(false);
+	}
+	else if (bIsDraggingRectangle)
+	{
+		PointTransformGizmo->SetVisibility(false);
+	}
+	else if (ShouldHideGizmo.IsBound() && ShouldHideGizmo.Execute())
+	{
+		PointTransformGizmo->SetVisibility(false);
+	}
+	else
+	{
+		PointTransformGizmo->SetVisibility(true);
+	}
+}
+
+
 bool ULatticeControlPointsMechanic::DeselectPoint(int32 PointID)
 {
 	bool PointFound = false;
@@ -594,12 +622,8 @@ void ULatticeControlPointsMechanic::OnUpdateModifierState(int ModifierID, bool b
 void ULatticeControlPointsMechanic::OnDragRectangleStarted()
 {
 	PreDragSelection = SelectedPointIDs;
-
-	// Hide gizmo while dragging
-	if (PointTransformGizmo)
-	{
-		PointTransformGizmo->SetVisibility(false);
-	}
+	bIsDraggingRectangle = true;
+	UpdateGizmoVisibility();
 }
 
 void ULatticeControlPointsMechanic::OnDragRectangleChanged(const FCameraRectangle& Rectangle)
@@ -644,7 +668,9 @@ void ULatticeControlPointsMechanic::OnDragRectangleChanged(const FCameraRectangl
 void ULatticeControlPointsMechanic::OnDragRectangleFinished(const FCameraRectangle& Rectangle, bool bCancelled)
 {
 	using namespace LatticeControlPointMechanicLocals;
-	
+
+	bIsDraggingRectangle = false;
+
 	ParentTool->GetToolManager()->BeginUndoTransaction(LatticePointSelectionTransactionText);
 
 	if (!IsEqual(PreDragSelection, SelectedPointIDs))
@@ -727,8 +753,7 @@ void FLatticeControlPointsMechanicSelectionChange::Apply(UObject* Object)
 	Mechanic->PointTransformGizmo->ReinitializeGizmoTransform(NewTransform);
 	Mechanic->PointTransformGizmo->SetNewChildScale(FVector{ 1.0, 1.0, 1.0 });	// Clear the child scale
 
-	bool bAnyPointSelected = (Mechanic->SelectedPointIDs.Num() > 0);
-	Mechanic->PointTransformGizmo->SetVisibility(bAnyPointSelected);
+	Mechanic->UpdateGizmoVisibility();
 	Mechanic->OnSelectionChanged.Broadcast();
 }
 
@@ -750,8 +775,7 @@ void FLatticeControlPointsMechanicSelectionChange::Revert(UObject* Object)
 	Mechanic->PointTransformGizmo->ReinitializeGizmoTransform(PreviousTransform);
 	Mechanic->PointTransformGizmo->SetNewChildScale(FVector{ 1.0, 1.0, 1.0 });	// Clear the child scale
 
-	bool bAnyPointSelected = (Mechanic->SelectedPointIDs.Num() > 0);
-	Mechanic->PointTransformGizmo->SetVisibility(bAnyPointSelected);
+	Mechanic->UpdateGizmoVisibility();
 	Mechanic->OnSelectionChanged.Broadcast();
 }
 
diff --git a/Engine/Plugins/Runtime/MeshModelingToolset/Source/ModelingComponents/Public/Mechanics/LatticeControlPointsMechanic.h b/Engine/Plugins/Runtime/MeshModelingToolset/Source/ModelingComponents/Public/Mechanics/LatticeControlPointsMechanic.h
--- a/Engine/Plugins/Runtime/MeshModelingToolset/Source/ModelingComponents/Public/Mechanics/LatticeControlPointsMechanic.h
+++ b/Engine/Plugins/Runtime/MeshModelingToolset/Source/ModelingComponents/Public/Mechanics/LatticeControlPointsMechanic.h
@@ -102,6 +102,9 @@ public:
 
 	void UpdatePointLocations(const TMap<int32, FVector3d>& NewLocations);
 
+	DECLARE_DELEGATE_RetVal(bool, FShouldHideGizmo);
+	FShouldHideGizmo ShouldHideGizmo;
+
 protected:
 
 	TArray<FVector3d> ControlPoints;
@@ -167,6 +170,7 @@ protected:
 	// Support for selection
 	UPROPERTY()
 	TObjectPtr<URectangleMarqueeMechanic> MarqueeMechanic;
+	bool bIsDraggingRectangle = false;
 	TSet<int32> SelectedPointIDs;
 	TSet<int32> PreDragSelection;
 	TArray<int32> CurrentDragSelection;
@@ -188,6 +192,7 @@ protected:
 	void SelectPoint(int32 PointID);
 	bool DeselectPoint(int32 PointID);
 	void UpdateGizmoLocation();
+	void UpdateGizmoVisibility();
 
 	void RebuildDrawables();
 
