MeshModelingTools: Fix PolyEd doing full proxy updates while dragging gizmo.

diff --git a/Engine/Plugins/Runtime/MeshModelingToolset/Source/MeshModelingTools/Private/EditMeshPolygonsTool.cpp b/Engine/Plugins/Runtime/MeshModelingToolset/Source/MeshModelingTools/Private/EditMeshPolygonsTool.cpp
--- a/Engine/Plugins/Runtime/MeshModelingToolset/Source/MeshModelingTools/Private/EditMeshPolygonsTool.cpp
+++ b/Engine/Plugins/Runtime/MeshModelingToolset/Source/MeshModelingTools/Private/EditMeshPolygonsTool.cpp
@@ -1038,7 +1038,11 @@ void UEditMeshPolygonsTool::ComputeUpdate_Gizmo()
 		LinearDeformer.ClearSolution(Mesh);
 	}
 
-	Preview->PreviewMesh->UpdatePreview(CurrentMesh.Get());
+	Preview->PreviewMesh->UpdatePreview(CurrentMesh.Get(),
+		// It's important to use the fast update path for the gizmo manipulations that only
+		// affect positions.
+		UPreviewMesh::ERenderUpdateMode::FastUpdate,
+		EMeshRenderAttributeFlags::Positions | EMeshRenderAttributeFlags::VertexNormals);
 
 	GetToolManager()->PostInvalidation();
 }
