Fixed a bug where the number of the vertical slices being generated is always off by one in the Sphere Primitive Tool.

diff --git a/Engine/Plugins/Runtime/MeshModelingToolset/Source/MeshModelingTools/Private/AddPrimitiveTool.cpp b/Engine/Plugins/Runtime/MeshModelingToolset/Source/MeshModelingTools/Private/AddPrimitiveTool.cpp
--- a/Engine/Plugins/Runtime/MeshModelingToolset/Source/MeshModelingTools/Private/AddPrimitiveTool.cpp
+++ b/Engine/Plugins/Runtime/MeshModelingToolset/Source/MeshModelingTools/Private/AddPrimitiveTool.cpp
@@ -691,7 +691,7 @@ void UAddSpherePrimitiveTool::GenerateMesh(FDynamicMesh3* OutMesh) const
 	{
 		FSphereGenerator SphereGen;
 		SphereGen.Radius = SphereSettings->Radius;
-		SphereGen.NumTheta = SphereSettings->VerticalSlices + 1;
+		SphereGen.NumTheta = SphereSettings->VerticalSlices;
 		SphereGen.NumPhi = SphereSettings->HorizontalSlices + 1;
 		SphereGen.bPolygroupPerQuad = (ShapeSettings->PolygroupMode == EMakeMeshPolygroupMode::PerQuad);
 		SphereGen.Generate();
