ModelingComponents/Selection:
- create a new MeshTopologySelector base class. The existing GroupTopologySelector class now inherits from it, as does a new BoundarySelector class
- likewise, create a new MeshTopologySelectionMehchanic base class and have (existing) PolygonSelectionMechanic and (new) BoundarySelectionMechanic inherit from it
- the new Boundary classes use FMeshBoundary loops to define selectable loops rather than FGroupTopology

HoleFillTool:
- change to using BoundarySelectionMechanic instead of GroupTopologySelector

Misc:
- allow FMeshBoundaryLoops to fail untangling a loop with bowties but still continue processing other loops

diff --git a/Engine/Plugins/Experimental/MeshModelingToolsetExp/Source/MeshModelingToolsExp/Private/HoleFillTool.cpp b/Engine/Plugins/Experimental/MeshModelingToolsetExp/Source/MeshModelingToolsExp/Private/HoleFillTool.cpp
--- a/Engine/Plugins/Experimental/MeshModelingToolsetExp/Source/MeshModelingToolsExp/Private/HoleFillTool.cpp
+++ b/Engine/Plugins/Experimental/MeshModelingToolsetExp/Source/MeshModelingToolsExp/Private/HoleFillTool.cpp
@@ -12,8 +12,7 @@
 #include "BaseBehaviors/MouseHoverBehavior.h"
 #include "MeshBoundaryLoops.h"
 #include "MeshOpPreviewHelpers.h"
-#include "Selection/PolygonSelectionMechanic.h"
-
+#include "Selection/BoundarySelectionMechanic.h"
 #include "TargetInterfaces/MaterialProvider.h"
 #include "TargetInterfaces/PrimitiveComponentBackedTarget.h"
 #include "ModelingToolTargetUtil.h"
@@ -45,12 +44,12 @@ void UHoleFillToolActions::PostAction(EHoleFillToolActions Action)
 
 void UHoleFillStatisticsProperties::Initialize(const UHoleFillTool& HoleFillTool)
 {
-	if (HoleFillTool.Topology == nullptr)
+	if (!HoleFillTool.BoundaryLoops.IsValid())
 	{
 		return;
 	}
 
-	int Initial = HoleFillTool.Topology->Edges.Num();
+	int Initial = HoleFillTool.BoundaryLoops->Loops.Num();
 	int Selected = 0;
 	int Successful = 0;
 	int Failed = 0;
@@ -65,12 +64,12 @@ void UHoleFillStatisticsProperties::Initialize(const UHoleFillTool& HoleFillTool
 
 void UHoleFillStatisticsProperties::Update(const UHoleFillTool& HoleFillTool, const FHoleFillOp& Op)
 {
-	if (HoleFillTool.Topology == nullptr)
+	if (!HoleFillTool.BoundaryLoops.IsValid())
 	{
 		return;
 	}
 
-	int Initial = HoleFillTool.Topology->Edges.Num();
+	int Initial = HoleFillTool.BoundaryLoops->Loops.Num();
 	int Selected = Op.Loops.Num();
 	int Failed = Op.NumFailedLoops;
 	int Successful = Selected - Failed;
@@ -159,12 +158,30 @@ void UHoleFillTool::Setup()
 	MeshSpatial.SetMesh(OriginalMesh.Get());
 
 	// initialize topology
-	Topology = MakeUnique<FBasicTopology>(OriginalMesh.Get(), false);
-	bool bTopologyOK = Topology->RebuildTopology();
+	constexpr bool bAutoComputeLoops = true;
+	BoundaryLoops = MakeUnique<UE::Geometry::FMeshBoundaryLoops>(OriginalMesh.Get(), bAutoComputeLoops);
+	const bool bLoopsOK = !BoundaryLoops->bAborted;
+
+	if (bLoopsOK && BoundaryLoops->Spans.Num() > 0)
+	{
+		// Boundary loop finding finished but some boundaries were too degenerate to form simple loops (this might be due to the 
+		// presence of bowtie vertices on boundaries, for example)
+		// Since this tool can tolerate failing to fill a boundary loop, just convert the spans to loops and let them fail
+		for (const FEdgeSpan& Span : BoundaryLoops->Spans)
+		{
+			FEdgeLoop SpanLoop(OriginalMesh.Get());
+			if (SpanLoop.InitializeFromVertices(Span.Vertices))
+			{
+				BoundaryLoops->Loops.Add(SpanLoop);
+			}
+		}
+		BoundaryLoops->Spans.Empty();
+	}
+
 
 	// Set up selection mechanic to find and select edges
 	IPrimitiveComponentBackedTarget* TargetComponent = Cast<IPrimitiveComponentBackedTarget>(Target);
-	SelectionMechanic = NewObject<UPolygonSelectionMechanic>(this);
+	SelectionMechanic = NewObject<UBoundarySelectionMechanic>(this);
 	SelectionMechanic->bAddSelectionFilterPropertiesToParentTool = false;
 	SelectionMechanic->Setup(this);
 	SelectionMechanic->Properties->bSelectEdges = true;
@@ -173,7 +190,7 @@ void UHoleFillTool::Setup()
 	SelectionMechanic->Initialize(OriginalMesh.Get(),
 		(FTransform3d)TargetComponent->GetWorldTransform(),
 		GetTargetWorld(),
-		Topology.Get(),
+		BoundaryLoops.Get(),
 		[this]() { return &MeshSpatial; }
 	);
 	// allow toggling selection without modifier key
@@ -191,7 +208,7 @@ void UHoleFillTool::Setup()
 	SetupPreview();
 	InvalidatePreviewResult();
 
-	if (!bTopologyOK)
+	if (!bLoopsOK)
 	{
 		GetToolManager()->DisplayMessage(
 			LOCTEXT("LoopFindError", "Error finding hole boundary loops."),
@@ -201,7 +218,7 @@ void UHoleFillTool::Setup()
 		SetToolPropertySourceEnabled(SmoothHoleFillProperties, false);
 		SetToolPropertySourceEnabled(Actions, false);
 	}
-	else if (Topology->Edges.Num() == 0)
+	else if (BoundaryLoops->Loops.Num() == 0)
 	{
 		GetToolManager()->DisplayMessage(
 			LOCTEXT("NoHoleNotification", "This mesh has no holes to fill."),
@@ -371,7 +388,7 @@ void UHoleFillTool::ApplyAction(EHoleFillToolActions ActionType)
 void UHoleFillTool::SelectAll()
 {
 	FGroupTopologySelection NewSelection;
-	for (int32 i = 0; i < Topology->Edges.Num(); ++i)
+	for (int32 i = 0; i < BoundaryLoops->Loops.Num(); ++i)
 	{
 		NewSelection.SelectedEdgeIDs.Add(i);
 	}
@@ -403,12 +420,9 @@ void UHoleFillTool::UpdateActiveBoundaryLoopSelection()
 	ActiveBoundaryLoopSelection.Reserve(NumEdges);
 	for (int32 EdgeID : ActiveSelection.SelectedEdgeIDs)
 	{
-		if (Topology->IsBoundaryEdge(EdgeID))
-		{
-			FSelectedBoundaryLoop& Loop = ActiveBoundaryLoopSelection.Emplace_GetRef();
-			Loop.EdgeTopoID = EdgeID;
-			Loop.EdgeIDs = Topology->GetGroupEdgeEdges(EdgeID);
-		}
+		FSelectedBoundaryLoop& Loop = ActiveBoundaryLoopSelection.Emplace_GetRef();
+		Loop.EdgeTopoID = EdgeID;
+		Loop.EdgeIDs = BoundaryLoops->Loops[EdgeID].Edges;
 	}
 }
 
@@ -425,16 +439,16 @@ void UHoleFillTool::Render(IToolsContextRenderAPI* RenderAPI)
 void UHoleFillTool::GetLoopsToFill(TArray<FEdgeLoop>& OutLoops) const
 {
 	OutLoops.Reset();
-	FMeshBoundaryLoops BoundaryLoops(OriginalMesh.Get());
+	check(BoundaryLoops.IsValid());
 
 	for (const FSelectedBoundaryLoop& FillEdge : ActiveBoundaryLoopSelection)
 	{
 		if (OriginalMesh->IsBoundaryEdge(FillEdge.EdgeIDs[0]))		// may no longer be boundary due to previous fill
 		{
-			int32 LoopID = BoundaryLoops.FindLoopContainingEdge(FillEdge.EdgeIDs[0]);
+			int32 LoopID = BoundaryLoops->FindLoopContainingEdge(FillEdge.EdgeIDs[0]);
 			if (LoopID >= 0)
 			{
-				OutLoops.Add(BoundaryLoops.Loops[LoopID]);
+				OutLoops.Add(BoundaryLoops->Loops[LoopID]);
 			}
 		}
 	}
diff --git a/Engine/Plugins/Experimental/MeshModelingToolsetExp/Source/MeshModelingToolsExp/Public/HoleFillTool.h b/Engine/Plugins/Experimental/MeshModelingToolsetExp/Source/MeshModelingToolsExp/Public/HoleFillTool.h
--- a/Engine/Plugins/Experimental/MeshModelingToolsetExp/Source/MeshModelingToolsExp/Public/HoleFillTool.h
+++ b/Engine/Plugins/Experimental/MeshModelingToolsetExp/Source/MeshModelingToolsExp/Public/HoleFillTool.h
@@ -11,11 +11,11 @@
 #include "BaseTools/MeshSurfacePointTool.h"
 #include "BaseBehaviors/BehaviorTargetInterfaces.h"
 #include "DynamicMesh/DynamicMeshAABBTree3.h"
-#include "GroupTopology.h"
 #include "BaseTools/SingleSelectionMeshEditingTool.h"
+#include "MeshBoundaryLoops.h"
 #include "HoleFillTool.generated.h"
 
-class UPolygonSelectionMechanic;
+class UBoundarySelectionMechanic;
 class UDynamicMeshReplacementChangeTarget;
 class UMeshOpPreviewWithBackgroundCompute;
 class UDynamicMeshComponent;
@@ -256,7 +256,7 @@ protected:
 	TObjectPtr<UMeshOpPreviewWithBackgroundCompute> Preview = nullptr;
 
 	UPROPERTY()
-	TObjectPtr<UPolygonSelectionMechanic> SelectionMechanic = nullptr;
+	TObjectPtr<UBoundarySelectionMechanic> SelectionMechanic = nullptr;
 
 	// Input mesh. Ownership shared with Op.
 	TSharedPtr<UE::Geometry::FDynamicMesh3, ESPMode::ThreadSafe> OriginalMesh;
@@ -281,30 +281,7 @@ protected:
 	void SelectAll();
 	void ClearSelection();
 
-	// Object used to get boundary loop information. All triangles return the same GroupID, so the only boundaries that 
-	// are returned are the the actual mesh boundaries.
-	// TODO: It seems like overkill to use a FGroupTopology subclass when we don't actually care about groups.
-	class FBasicTopology : public UE::Geometry::FGroupTopology
-	{
-	public:
-		FBasicTopology(const UE::Geometry::FDynamicMesh3* Mesh, bool bAutoBuild) :
-			FGroupTopology(Mesh, false)
-		{
-			if (bAutoBuild)
-			{
-				// Virtual func resolution doesn't work in constructors. Though we're not currently
-				// overriding RebuildTopology, let's do the proper thing in case we get copy-pasted
-				// somewhere where we do.
-				RebuildTopology();
-			}
-		}
-
-		int GetGroupID(int TriangleID) const override
-		{
-			return Mesh->IsTriangle(TriangleID) ? 1 : 0;
-		}
-	};
-	TUniquePtr<FBasicTopology> Topology;
+	TUniquePtr<UE::Geometry::FMeshBoundaryLoops> BoundaryLoops;
 
 	struct FSelectedBoundaryLoop
 	{
diff --git a/Engine/Plugins/Runtime/GeometryProcessing/Source/DynamicMesh/Private/MeshBoundaryLoops.cpp b/Engine/Plugins/Runtime/GeometryProcessing/Source/DynamicMesh/Private/MeshBoundaryLoops.cpp
--- a/Engine/Plugins/Runtime/GeometryProcessing/Source/DynamicMesh/Private/MeshBoundaryLoops.cpp
+++ b/Engine/Plugins/Runtime/GeometryProcessing/Source/DynamicMesh/Private/MeshBoundaryLoops.cpp
@@ -321,18 +321,25 @@ bool FMeshBoundaryLoops::Compute()
 			bool bSubloopsOK = ExtractSubloops(loop_verts, loop_edges, bowties, subloops);
 			if (bSubloopsOK == false)
 			{
-				goto CATASTROPHIC_ABORT;
-			}
-			for (FEdgeLoop& loop : subloops.Loops)
-			{
-				Loops.Add(loop);
+				if (FailureBehavior == EFailureBehaviors::Abort)
+				{
+					goto CATASTROPHIC_ABORT;
+				}
+
+				if (subloops.Spans.Num() > 0)
+				{
+					bFellBackToSpansOnFailure = true;
+					for (FEdgeSpan& span : subloops.Spans)
+					{
+						Spans.Add(span);
+					}
+				}
 			}
-			if (subloops.Spans.Num() > 0)
+			else
 			{
-				bFellBackToSpansOnFailure = true;
-				for (FEdgeSpan& span : subloops.Spans)
+				for (FEdgeLoop& loop : subloops.Loops)
 				{
-					Spans.Add(span);
+					Loops.Add(loop);
 				}
 			}
 		}
diff --git a/Engine/Plugins/Runtime/MeshModelingToolset/Source/MeshModelingTools/Private/EditMeshPolygonsTool.cpp b/Engine/Plugins/Runtime/MeshModelingToolset/Source/MeshModelingTools/Private/EditMeshPolygonsTool.cpp
--- a/Engine/Plugins/Runtime/MeshModelingToolset/Source/MeshModelingTools/Private/EditMeshPolygonsTool.cpp
+++ b/Engine/Plugins/Runtime/MeshModelingToolset/Source/MeshModelingTools/Private/EditMeshPolygonsTool.cpp
@@ -158,7 +158,7 @@ void UEditMeshPolygonsSelectionModeToolBuilder::InitializeNewTool(USingleSelecti
 		PolyTool->SetToSelectionModeInterface();
 
 		UPolygonSelectionMechanic* SelectionMechanic = PolyTool->SelectionMechanic;
-		UPolygonSelectionMechanicProperties* SelectionProps = SelectionMechanic->Properties;
+		UMeshTopologySelectionMechanicProperties* SelectionProps = SelectionMechanic->Properties;
 		SelectionProps->bSelectFaces = SelectionProps->bSelectEdges = SelectionProps->bSelectVertices = false;
 		SelectionProps->bSelectEdgeLoops = SelectionProps->bSelectEdgeRings = false;
 
diff --git a/Engine/Plugins/Runtime/MeshModelingToolset/Source/MeshModelingTools/Private/ToolActivities/PolyEditInsertEdgeActivity.cpp b/Engine/Plugins/Runtime/MeshModelingToolset/Source/MeshModelingTools/Private/ToolActivities/PolyEditInsertEdgeActivity.cpp
--- a/Engine/Plugins/Runtime/MeshModelingToolset/Source/MeshModelingTools/Private/ToolActivities/PolyEditInsertEdgeActivity.cpp
+++ b/Engine/Plugins/Runtime/MeshModelingToolset/Source/MeshModelingTools/Private/ToolActivities/PolyEditInsertEdgeActivity.cpp
@@ -88,7 +88,8 @@ void UPolyEditInsertEdgeActivity::Setup(UInteractiveTool* ParentToolIn)
 
 	ActivityContext = ParentTool->GetToolManager()->GetContextObjectStore()->FindContext<UPolyEditActivityContext>();
 
-	TopologySelector = ActivityContext->SelectionMechanic->GetTopologySelector();
+	// TODO: When the deprecated function GetTopologySelector is removed from UPolygonSelectionMechanic, we can remove the UMeshTopologySelectionMechanic:: prefix here
+	TopologySelector = ActivityContext->SelectionMechanic->UMeshTopologySelectionMechanic::GetTopologySelector();
 
 	ActivityContext->OnUndoRedo.AddWeakLambda(this, [this](bool bGroupTopologyModified)
 	{
diff --git a/Engine/Plugins/Runtime/MeshModelingToolset/Source/MeshModelingTools/Private/ToolActivities/PolyEditInsertEdgeLoopActivity.cpp b/Engine/Plugins/Runtime/MeshModelingToolset/Source/MeshModelingTools/Private/ToolActivities/PolyEditInsertEdgeLoopActivity.cpp
--- a/Engine/Plugins/Runtime/MeshModelingToolset/Source/MeshModelingTools/Private/ToolActivities/PolyEditInsertEdgeLoopActivity.cpp
+++ b/Engine/Plugins/Runtime/MeshModelingToolset/Source/MeshModelingTools/Private/ToolActivities/PolyEditInsertEdgeLoopActivity.cpp
@@ -109,7 +109,8 @@ void UPolyEditInsertEdgeLoopActivity::Setup(UInteractiveTool* ParentToolIn)
 
 	ActivityContext = ParentTool->GetToolManager()->GetContextObjectStore()->FindContext<UPolyEditActivityContext>();
 
-	TopologySelector = ActivityContext->SelectionMechanic->GetTopologySelector();
+	// TODO: When the deprecated function GetTopologySelector is removed from UPolygonSelectionMechanic, we can remove the UMeshTopologySelectionMechanic:: prefix here
+	TopologySelector = ActivityContext->SelectionMechanic->UMeshTopologySelectionMechanic::GetTopologySelector();
 
 	ActivityContext->OnUndoRedo.AddWeakLambda(this, [this](bool bGroupTopologyModified)
 	{
diff --git a/Engine/Plugins/Runtime/MeshModelingToolset/Source/MeshModelingTools/Public/ToolActivities/PolyEditInsertEdgeActivity.h b/Engine/Plugins/Runtime/MeshModelingToolset/Source/MeshModelingTools/Public/ToolActivities/PolyEditInsertEdgeActivity.h
--- a/Engine/Plugins/Runtime/MeshModelingToolset/Source/MeshModelingTools/Public/ToolActivities/PolyEditInsertEdgeActivity.h
+++ b/Engine/Plugins/Runtime/MeshModelingToolset/Source/MeshModelingTools/Public/ToolActivities/PolyEditInsertEdgeActivity.h
@@ -12,10 +12,10 @@
 #include "InteractiveToolChange.h" //FToolCommandChange
 #include "MeshOpPreviewHelpers.h" //FDynamicMeshOpResult
 #include "Operations/GroupEdgeInserter.h"
-#include "Selection/GroupTopologySelector.h"
+#include "Selection/MeshTopologySelector.h"
 #include "SingleSelectionTool.h"
 #include "ToolDataVisualizer.h"
-
+#include "GroupTopology.h"
 #include "PolyEditInsertEdgeActivity.generated.h"
 
 PREDECLARE_USE_GEOMETRY_CLASS(FDynamicMeshChange);
@@ -107,7 +107,7 @@ protected:
 	bool bIsRunning = false;
 
 	FTransform TargetTransform;
-	TSharedPtr<FGroupTopologySelector, ESPMode::ThreadSafe> TopologySelector;
+	TSharedPtr<FMeshTopologySelector, ESPMode::ThreadSafe> TopologySelector;
 
 	TArray<TPair<FVector3d, FVector3d>> PreviewEdges;
 	TArray<FVector3d> PreviewPoints;
@@ -116,7 +116,7 @@ protected:
 
 	FToolDataVisualizer ExistingEdgesRenderer;
 	FToolDataVisualizer PreviewEdgeRenderer;
-	FGroupTopologySelector::FSelectionSettings TopologySelectorSettings;
+	FMeshTopologySelector::FSelectionSettings TopologySelectorSettings;
 
 
 	// Inputs from user interaction:
@@ -137,7 +137,7 @@ protected:
 
 	bool bShowingBaseMesh = false;
 	bool bLastComputeSucceeded = false;
-	TSharedPtr<FGroupTopology, ESPMode::ThreadSafe> LatestOpTopologyResult;
+	TSharedPtr<UE::Geometry::FGroupTopology, ESPMode::ThreadSafe> LatestOpTopologyResult;
 	TSharedPtr<TSet<int32>, ESPMode::ThreadSafe> LatestOpChangedTids;
 
 	int32 CurrentChangeStamp = 0;
diff --git a/Engine/Plugins/Runtime/MeshModelingToolset/Source/MeshModelingTools/Public/ToolActivities/PolyEditInsertEdgeLoopActivity.h b/Engine/Plugins/Runtime/MeshModelingToolset/Source/MeshModelingTools/Public/ToolActivities/PolyEditInsertEdgeLoopActivity.h
--- a/Engine/Plugins/Runtime/MeshModelingToolset/Source/MeshModelingTools/Public/ToolActivities/PolyEditInsertEdgeLoopActivity.h
+++ b/Engine/Plugins/Runtime/MeshModelingToolset/Source/MeshModelingTools/Public/ToolActivities/PolyEditInsertEdgeLoopActivity.h
@@ -11,7 +11,7 @@
 #include "InteractiveToolBuilder.h" //UInteractiveToolBuilder
 #include "InteractiveToolChange.h" //FToolCommandChange
 #include "MeshOpPreviewHelpers.h" //FDynamicMeshOpResult
-#include "Selection/GroupTopologySelector.h"
+#include "Selection/MeshTopologySelector.h"
 #include "SingleSelectionTool.h"
 #include "ToolDataVisualizer.h"
 
@@ -140,7 +140,7 @@ protected:
 	bool bIsRunning = false;
 
 	FTransform TargetTransform;
-	TSharedPtr<FGroupTopologySelector, ESPMode::ThreadSafe> TopologySelector;
+	TSharedPtr<FMeshTopologySelector, ESPMode::ThreadSafe> TopologySelector;
 
 	TArray<TPair<FVector3d, FVector3d>> PreviewEdges;
 
@@ -152,7 +152,7 @@ protected:
 
 	FToolDataVisualizer PreviewEdgeRenderer;
 	FToolDataVisualizer ProblemTopologyRenderer;
-	FGroupTopologySelector::FSelectionSettings TopologySelectorSettings;
+	FMeshTopologySelector::FSelectionSettings TopologySelectorSettings;
 	float ProblemVertTickWidth = 8;
 
 	void SetupPreview();
@@ -178,6 +178,6 @@ protected:
 
 	// Copied over on op completion
 	bool bLastComputeSucceeded = false;
-	TSharedPtr<FGroupTopology, ESPMode::ThreadSafe> LatestOpTopologyResult;
+	TSharedPtr<UE::Geometry::FGroupTopology, ESPMode::ThreadSafe> LatestOpTopologyResult;
 	TSharedPtr<TSet<int32>, ESPMode::ThreadSafe> LatestOpChangedTids;
 };
diff --git a/Engine/Plugins/Runtime/MeshModelingToolset/Source/ModelingComponents/Private/Selection/BoundarySelectionMechanic.cpp b/Engine/Plugins/Runtime/MeshModelingToolset/Source/ModelingComponents/Private/Selection/BoundarySelectionMechanic.cpp
new file mode 100644
--- /dev/null
+++ b/Engine/Plugins/Runtime/MeshModelingToolset/Source/ModelingComponents/Private/Selection/BoundarySelectionMechanic.cpp
@@ -0,0 +1,94 @@
+// Copyright Epic Games, Inc. All Rights Reserved.
+
+#include "Selection/BoundarySelectionMechanic.h"
+#include "MeshBoundaryLoops.h"
+#include "Selection/BoundarySelector.h"
+#include "ToolSceneQueriesUtil.h"
+
+using namespace UE::Geometry;
+
+#define LOCTEXT_NAMESPACE "UBoundarySelectionMechanic"
+
+void UBoundarySelectionMechanic::Initialize(
+	const FDynamicMesh3* MeshIn,
+	FTransform3d TargetTransformIn,
+	UWorld* WorldIn,
+	const UE::Geometry::FMeshBoundaryLoops* BoundaryLoopsIn,
+	TFunction<FDynamicMeshAABBTree3* ()> GetSpatialSourceFuncIn)
+{
+	TopoSelector = MakeShared<FBoundarySelector>(MeshIn, BoundaryLoopsIn);;
+	UMeshTopologySelectionMechanic::Initialize(MeshIn, TargetTransformIn, WorldIn, GetSpatialSourceFuncIn);
+}
+
+bool UBoundarySelectionMechanic::UpdateHighlight(const FRay& WorldRay)
+{
+	checkf(DrawnTriangleSetComponent != nullptr, TEXT("Initialize() not called on UMeshTopologySelectionMechanic."));
+
+	FRay3d LocalRay(TargetTransform.InverseTransformPosition((FVector3d)WorldRay.Origin),
+		TargetTransform.InverseTransformVector((FVector3d)WorldRay.Direction));
+	UE::Geometry::Normalize(LocalRay.Direction);
+
+	HilightSelection.Clear();
+	FVector3d LocalPosition, LocalNormal;
+	FGroupTopologySelector::FSelectionSettings TopoSelectorSettings = GetTopoSelectorSettings(CameraState.bIsOrthographic);
+	bool bHit = TopoSelector->FindSelectedElement(TopoSelectorSettings, LocalRay, HilightSelection, LocalPosition, LocalNormal);
+
+	// Don't hover highlight a selection that we already selected, because people didn't like that
+	if (PersistentSelection.Contains(HilightSelection))
+	{
+		HilightSelection.Clear();
+	}
+
+	return bHit;
+}
+
+
+bool UBoundarySelectionMechanic::UpdateSelection(const FRay& WorldRay, FVector3d& LocalHitPositionOut, FVector3d& LocalHitNormalOut)
+{
+	FRay3d LocalRay(TargetTransform.InverseTransformPosition((FVector3d)WorldRay.Origin),
+		TargetTransform.InverseTransformVector((FVector3d)WorldRay.Direction));
+	UE::Geometry::Normalize(LocalRay.Direction);
+
+	const FGroupTopologySelection PreviousSelection = PersistentSelection;
+
+	FVector3d LocalPosition, LocalNormal;
+	FGroupTopologySelection Selection;
+	FGroupTopologySelector::FSelectionSettings TopoSelectorSettings = GetTopoSelectorSettings(CameraState.bIsOrthographic);
+	if (TopoSelector->FindSelectedElement(TopoSelectorSettings, LocalRay, Selection, LocalPosition, LocalNormal))
+	{
+		LocalHitPositionOut = LocalPosition;
+		LocalHitNormalOut = LocalNormal;
+	}
+
+	if (ShouldAddToSelectionFunc())
+	{
+		if (ShouldRemoveFromSelectionFunc())
+		{
+			PersistentSelection.Toggle(Selection);
+		}
+		else
+		{
+			PersistentSelection.Append(Selection);
+		}
+	}
+	else if (ShouldRemoveFromSelectionFunc())
+	{
+		PersistentSelection.Remove(Selection);
+	}
+	else
+	{
+		PersistentSelection = Selection;
+	}
+
+	if (PersistentSelection != PreviousSelection)
+	{
+		SelectionTimestamp++;
+		OnSelectionChanged.Broadcast();
+		return true;
+	}
+
+	return false;
+}
+
+
+#undef LOCTEXT_NAMESPACE
diff --git a/Engine/Plugins/Runtime/MeshModelingToolset/Source/ModelingComponents/Private/Selection/BoundarySelector.cpp b/Engine/Plugins/Runtime/MeshModelingToolset/Source/ModelingComponents/Private/Selection/BoundarySelector.cpp
new file mode 100644
--- /dev/null
+++ b/Engine/Plugins/Runtime/MeshModelingToolset/Source/ModelingComponents/Private/Selection/BoundarySelector.cpp
@@ -0,0 +1,264 @@
+// Copyright Epic Games, Inc. All Rights Reserved.
+
+#include "Selection/BoundarySelector.h"
+#include "MeshBoundaryLoops.h"
+#include "ToolDataVisualizer.h"
+#include "ToolSceneQueriesUtil.h"
+#include "Mechanics/RectangleMarqueeMechanic.h" // FCameraRectangle
+
+using namespace UE::Geometry;
+
+#define LOCTEXT_NAMESPACE "FBoundarySelector"
+
+bool FBoundaryTopologyProvider::GetLoopTangent(int LoopID, FVector3d& TangentOut) const
+{
+	check(LoopID >= 0 && LoopID < BoundaryLoops->Loops.Num());
+
+	const FEdgeLoop& Loop = BoundaryLoops->Loops[LoopID];
+	FVector3d StartPos = BoundaryLoops->Mesh->GetVertex(Loop.Vertices[0]);
+	FVector3d EndPos = BoundaryLoops->Mesh->GetVertex(Loop.Vertices[Loop.Vertices.Num() - 1]);
+
+	TangentOut = EndPos - StartPos;
+	const bool bNormalized = TangentOut.Normalize(100 * FMathd::ZeroTolerance);
+
+	if (bNormalized)
+	{
+		return true;
+	}
+	else
+	{
+		TangentOut = FVector3d::UnitX();
+		return false;
+	}
+}
+
+double FBoundaryTopologyProvider::GetLoopArcLength(int32 LoopID, TArray<double>* PerVertexLengthsOut) const
+{
+	check(BoundaryLoops);
+	check(LoopID >= 0 && LoopID < BoundaryLoops->Loops.Num());
+	const FDynamicMesh3* Mesh = BoundaryLoops->Mesh;
+	check(Mesh);
+
+	const TArray<int32>& Vertices = GetGroupEdgeVertices(LoopID);
+	int32 NumV = Vertices.Num();
+	if (PerVertexLengthsOut != nullptr)
+	{
+		PerVertexLengthsOut->SetNum(NumV);
+		(*PerVertexLengthsOut)[0] = 0.0;
+	}
+	double AccumLength = 0;
+	for (int32 k = 1; k < NumV; ++k)
+	{
+		AccumLength += Distance(Mesh->GetVertex(Vertices[k]), Mesh->GetVertex(Vertices[k - 1]));
+		if (PerVertexLengthsOut != nullptr)
+		{
+			(*PerVertexLengthsOut)[k] = AccumLength;
+		}
+	}
+	return AccumLength;
+}
+
+
+FVector3d FBoundaryTopologyProvider::GetLoopMidpoint(int32 LoopID, double* ArcLengthOut, TArray<double>* PerVertexLengthsOut) const
+{
+	check(BoundaryLoops);
+	check(LoopID >= 0 && LoopID < BoundaryLoops->Loops.Num());
+	const FDynamicMesh3* Mesh = BoundaryLoops->Mesh;
+	check(Mesh);
+
+	const TArray<int32>& Vertices = GetGroupEdgeVertices(LoopID);
+	int32 NumV = Vertices.Num();
+
+	// trivial case
+	if (NumV == 2)
+	{
+		FVector3d A(Mesh->GetVertex(Vertices[0])), B(Mesh->GetVertex(Vertices[1]));
+		if (ArcLengthOut)
+		{
+			*ArcLengthOut = Distance(A, B);
+		}
+		if (PerVertexLengthsOut)
+		{
+			(*PerVertexLengthsOut).SetNum(2);
+			(*PerVertexLengthsOut)[0] = 0;
+			(*PerVertexLengthsOut)[1] = Distance(A, B);
+		}
+		return (A + B) * 0.5;
+	}
+
+	// if we want lengths anyway we can avoid second loop
+	if (PerVertexLengthsOut)
+	{
+		double Len = GetLoopArcLength(LoopID, PerVertexLengthsOut);
+		if (ArcLengthOut)
+		{
+			*ArcLengthOut = Len;
+		}
+		Len /= 2;
+		int32 k = 0;
+		while ((*PerVertexLengthsOut)[k] < Len)
+		{
+			k++;
+		}
+		int32 kprev = k - 1;
+		double a = (*PerVertexLengthsOut)[k - 1], b = (*PerVertexLengthsOut)[k];
+		double t = (Len - a) / (b - a);
+		FVector3d A(Mesh->GetVertex(Vertices[k - 1])), B(Mesh->GetVertex(Vertices[k]));
+		return Lerp(A, B, t);
+	}
+
+	// compute arclen and then walk forward until we get halfway
+	double Len = GetLoopArcLength(LoopID);
+	if (ArcLengthOut)
+	{
+		*ArcLengthOut = Len;
+	}
+	Len /= 2;
+	double AccumLength = 0;
+	for (int32 k = 1; k < NumV; ++k)
+	{
+		double NewLen = AccumLength + Distance(Mesh->GetVertex(Vertices[k]), Mesh->GetVertex(Vertices[k - 1]));
+		if (NewLen > Len)
+		{
+			double t = (Len - AccumLength) / (NewLen - AccumLength);
+			FVector3d A(Mesh->GetVertex(Vertices[k - 1])), B(Mesh->GetVertex(Vertices[k]));
+			return Lerp(A, B, t);
+		}
+		AccumLength = NewLen;
+	}
+
+	// somehow failed?
+	return (Mesh->GetVertex(Vertices[0]) + Mesh->GetVertex(Vertices[NumV - 1])) * 0.5;
+}
+
+
+
+FFrame3d FBoundaryTopologyProvider::GetSelectionFrame(const FGroupTopologySelection& Selection, FFrame3d* InitialLocalFrame) const
+{
+	check(BoundaryLoops);
+	const FDynamicMesh3* Mesh = BoundaryLoops->Mesh;
+	check(Mesh);
+
+
+	int32 NumLoops = Selection.SelectedEdgeIDs.Num();
+	FFrame3d StartFrame = (InitialLocalFrame) ? (*InitialLocalFrame) : FFrame3d();
+
+	if (NumLoops == 1)
+	{
+		int32 LoopID = Selection.GetASelectedEdgeID();
+		int32 MeshEdgeID = BoundaryLoops->Loops[LoopID].Edges[0];
+
+		// align Z axis of frame to face normal of one of the connected faces. 
+		FIndex2i EdgeTris = Mesh->GetEdgeT(MeshEdgeID);
+		int32 UseFace = (EdgeTris.B != IndexConstants::InvalidID) ? FMath::Min(EdgeTris.A, EdgeTris.B)
+			: EdgeTris.A;
+		FVector3d FaceNormal = Mesh->GetTriNormal(UseFace);
+		if (FaceNormal.Length() > 0.1)
+		{
+			StartFrame.AlignAxis(2, FaceNormal);
+		}
+
+		// align X axis along the edge, around the aligned Z axis
+		FVector3d Tangent;
+		if (GetLoopTangent(LoopID, Tangent))
+		{
+			StartFrame.ConstrainedAlignAxis(0, Tangent, StartFrame.Z());
+		}
+
+		StartFrame.Origin = GetLoopMidpoint(LoopID);
+		return StartFrame;
+	}
+
+	// If we have multiple loops, just align the frame with the world axes
+
+	const int NumSelectedLoops = Selection.SelectedEdgeIDs.Num();
+
+	FVector3d AccumulatedOrigin = FVector3d::Zero();
+	for (int32 LoopID : Selection.SelectedEdgeIDs)
+	{
+		const FEdgeLoop& Loop = BoundaryLoops->Loops[LoopID];
+		FVector3d StartPos = Mesh->GetVertex(Loop.Vertices[0]);
+		FVector3d EndPos = Mesh->GetVertex(Loop.Vertices[Loop.Vertices.Num() - 1]);
+		AccumulatedOrigin += 0.5 * (StartPos + EndPos);
+	}
+
+	check(NumSelectedLoops > 1);
+
+	FFrame3d AccumulatedFrame;
+	AccumulatedOrigin /= (double)NumSelectedLoops;
+	AccumulatedFrame = FFrame3d(AccumulatedOrigin, FQuaterniond::Identity());
+
+	return AccumulatedFrame;
+}
+
+FAxisAlignedBox3d FBoundaryTopologyProvider::GetSelectionBounds(
+	const FGroupTopologySelection& Selection,
+	TFunctionRef<FVector3d(const FVector3d&)> TransformFunc) const
+{
+	check(BoundaryLoops);
+	const FDynamicMesh3* Mesh = BoundaryLoops->Mesh;
+	check(Mesh);
+
+	if (ensure(!Selection.IsEmpty()) == false)
+	{
+		return Mesh->GetBounds();
+	}
+
+	FAxisAlignedBox3d Bounds = FAxisAlignedBox3d::Empty();
+	for (int32 EdgeID : Selection.SelectedEdgeIDs)
+	{
+		const FEdgeLoop& Loop = BoundaryLoops->Loops[EdgeID];
+		for (int32 vid : Loop.Vertices)
+		{
+			Bounds.Contain(TransformFunc(Mesh->GetVertex(vid)));
+		}
+	}
+
+	return Bounds;
+}
+
+
+FBoundarySelector::FBoundarySelector(const FDynamicMesh3* MeshIn, const FMeshBoundaryLoops* BoundaryLoopsIn) :
+	FMeshTopologySelector()
+{
+	Mesh = MeshIn;
+	TopologyProvider = MakeUnique<FBoundaryTopologyProvider>(BoundaryLoopsIn);
+	bGeometryInitialized = false;
+	bGeometryUpToDate = false;
+}
+
+void FBoundarySelector::DrawSelection(const FGroupTopologySelection& Selection, FToolDataVisualizer* Renderer, const FViewCameraState* CameraState, ECornerDrawStyle CornerDrawStyle)
+{
+	FLinearColor UseColor = Renderer->LineColor;
+	float LineWidth = Renderer->LineThickness;
+
+	for (int EdgeID : Selection.SelectedEdgeIDs)
+	{
+		const TArray<int>& Vertices = TopologyProvider->GetGroupEdgeVertices(EdgeID);
+		int NV = Vertices.Num() - 1;
+
+		// Draw the edge, but also draw the endpoints in ortho mode (to make projected edges visible)
+		FVector A = (FVector)Mesh->GetVertex(Vertices[0]);
+		if (CameraState->bIsOrthographic)
+		{
+			Renderer->DrawPoint(A, UseColor, 10, false);
+		}
+
+		for (int k = 0; k < NV; ++k)
+		{
+			FVector B = (FVector)Mesh->GetVertex(Vertices[k + 1]);
+			Renderer->DrawLine(A, B, UseColor, LineWidth, false);
+			A = B;
+		}
+		FVector B = (FVector)Mesh->GetVertex(Vertices[0]);
+		Renderer->DrawLine(A, B, UseColor, LineWidth, false);
+
+		if (CameraState->bIsOrthographic)
+		{
+			Renderer->DrawPoint(A, UseColor, LineWidth, false);
+		}
+	}
+}
+
+
+#undef LOCTEXT_NAMESPACE
diff --git a/Engine/Plugins/Runtime/MeshModelingToolset/Source/ModelingComponents/Private/Selection/GroupTopologySelector.cpp b/Engine/Plugins/Runtime/MeshModelingToolset/Source/ModelingComponents/Private/Selection/GroupTopologySelector.cpp
--- a/Engine/Plugins/Runtime/MeshModelingToolset/Source/ModelingComponents/Private/Selection/GroupTopologySelector.cpp
+++ b/Engine/Plugins/Runtime/MeshModelingToolset/Source/ModelingComponents/Private/Selection/GroupTopologySelector.cpp
@@ -1,798 +1,151 @@
 // Copyright Epic Games, Inc. All Rights Reserved.
 
-
 #include "Selection/GroupTopologySelector.h"
-
-#include "ConvexVolume.h"
-#include "Engine/Polys.h"
 #include "Mechanics/RectangleMarqueeMechanic.h" // FCameraRectangle
-#include "MeshQueries.h"
 #include "ToolDataVisualizer.h"
 #include "ToolSceneQueriesUtil.h"
 
-using namespace UE::Geometry;
-
-// Local utility function forward declarations
-bool IsOccluded(const FGeometrySet3::FNearest& ClosestElement, const FVector3d& ViewOrigin, const FDynamicMeshAABBTree3* Spatial, bool bBackFacesOcclude);
-bool IsOccluded(const FVector3d& Point, const FVector3d& ViewOrigin, const FDynamicMeshAABBTree3* Spatial, bool bBackFacesOcclude);
-void AddNewEdgeLoopEdgesFromCorner(const FGroupTopology& Topology, int32 EdgeID, int32 CornerID, TSet<int32>& EdgeSet);
-bool GetNextEdgeLoopEdge(const FGroupTopology& Topology, int32 IncomingEdgeID, int32 CornerID, int32& NextEdgeIDOut);
-void AddNewEdgeRingEdges(const FGroupTopology& Topology, int32 StartEdgeID, int32 ForwardGroupID, TSet<int32>& EdgeSet);
-bool GetQuadOppositeEdge(const FGroupTopology& Topology, int32 EdgeIDIn, int32 GroupID, int32& OppositeEdgeIDOut);
-
-FGroupTopologySelector::FGroupTopologySelector()
-{
-	// initialize to sane values
-	PointsWithinToleranceTest =
-		[](const FVector3d& A, const FVector3d& B, double TolScale) { return Distance(A, B) < (TolScale*1.0); };
-	GetSpatial =
-		[]() { return nullptr; };
-}
-
-
-
-void FGroupTopologySelector::Initialize(const FDynamicMesh3* MeshIn, const FGroupTopology* TopologyIn)
+void FGroupTopologyUtils::AddNewEdgeLoopEdgesFromCorner(int32 EdgeID, int32 CornerID, TSet<int32>& EdgeSet) const
 {
-	Mesh = MeshIn;
-	Topology = TopologyIn;
-	bGeometryInitialized = false;
-	bGeometryUpToDate = false;
-}
-
-
-void FGroupTopologySelector::Invalidate(bool bTopologyDeformed, bool bTopologyModified)
-{
-	if (bTopologyDeformed)
-	{
-		bGeometryUpToDate = false;
-	}
-	if (bTopologyModified)
-	{
-		bGeometryUpToDate = bGeometryInitialized = false;
-	}
-}
-
-
-const FGeometrySet3& FGroupTopologySelector::GetGeometrySet()
-{
-	if (bGeometryInitialized == false)
+	int32 LastCornerID = CornerID;
+	int32 LastEdgeID = EdgeID;
+	while (true)
 	{
-		GeometrySet.Reset();
-		int NumCorners = Topology->Corners.Num();
-		for (int CornerID = 0; CornerID < NumCorners; ++CornerID)
+		int32 NextEid;
+		if (!GetNextEdgeLoopEdge(LastEdgeID, LastCornerID, NextEid))
 		{
-			FVector3d Position = Mesh->GetVertex(Topology->Corners[CornerID].VertexID);
-			GeometrySet.AddPoint(CornerID, Position);
+			break; // Probably not a valence 4 corner
 		}
-		int NumEdges = Topology->Edges.Num();
-		for (int EdgeID = 0; EdgeID < NumEdges; ++EdgeID)
+		if (EdgeSet.Contains(NextEid))
 		{
-			FPolyline3d Polyline;
-			Topology->Edges[EdgeID].Span.GetPolyline(Polyline);
-			GeometrySet.AddCurve(EdgeID, Polyline);
+			break; // Either we finished the loop, or we'll continue it from another selection
 		}
 
-		bGeometryInitialized = true;
-		bGeometryUpToDate = true;
-	}
+		EdgeSet.Add(NextEid);
 
-	if (bGeometryUpToDate == false)
-	{
-		int NumCorners = Topology->Corners.Num();
-		for (int CornerID = 0; CornerID < NumCorners; ++CornerID)
-		{
-			FVector3d Position = Mesh->GetVertex(Topology->Corners[CornerID].VertexID);
-			GeometrySet.UpdatePoint(CornerID, Position);
-		}
-		int NumEdges = Topology->Edges.Num();
-		for (int EdgeID = 0; EdgeID < NumEdges; ++EdgeID)
-		{
-			FPolyline3d Polyline;
-			Topology->Edges[EdgeID].Span.GetPolyline(Polyline);
-			GeometrySet.UpdateCurve(EdgeID, Polyline);
-		}
-		bGeometryUpToDate = true;
-	}
+		LastEdgeID = NextEid;
+		const FGroupTopology::FGroupEdge& LastEdge = GroupTopology->Edges[LastEdgeID];
+		LastCornerID = LastEdge.EndpointCorners[0] == LastCornerID ? LastEdge.EndpointCorners[1] : LastEdge.EndpointCorners[0];
 
-	return GeometrySet;
+		check(LastCornerID != IndexConstants::InvalidID);
+	}
 }
 
-bool FGroupTopologySelector::FindSelectedElement(const FSelectionSettings& Settings, const FRay3d& Ray,
-	FGroupTopologySelection& ResultOut, FVector3d& SelectedPositionOut, FVector3d& SelectedNormalOut, int32* EdgeSegmentIdOut)
+bool FGroupTopologyUtils::GetNextEdgeLoopEdge(int32 IncomingEdgeID, int32 CornerID, int32& NextEdgeIDOut) const
 {
-	// These get used for finding intersections with triangles and corners/edges, repectively.
-	FDynamicMeshAABBTree3* Spatial = GetSpatial();
-	const FGeometrySet3& TopoSpatial = GetGeometrySet();
-
-	IMeshSpatial::FQueryOptions SpatialQueryOptions;
-	if (Spatial && !Settings.bHitBackFaces)
-	{
-		SpatialQueryOptions.TriangleFilterF = [Spatial, &Ray](int32 Tid) {
-			return Spatial->GetMesh()->GetTriNormal(Tid).Dot(Ray.Direction) < 0;
-		};
-	}
-
-	// We start by intersecting with the mesh triangles because even when selecting corners or edges, we set
-	// the normal based on the true triangle that we hit. If we end up with a simple face selection, we will
-	// end up using this result.
-	double RayParameter = -1;
-	int HitTriangleID = IndexConstants::InvalidID;
-	FVector3d TriangleHitPos;
-	bool bActuallyHitSurface = (Spatial != nullptr) ? Spatial->FindNearestHitTriangle(Ray, RayParameter, HitTriangleID, SpatialQueryOptions) : false;
-	if (bActuallyHitSurface)
-	{
-		TriangleHitPos = Ray.PointAt(RayParameter);
-		SelectedNormalOut = Mesh->GetTriNormal(HitTriangleID);
-	}
-	else
-	{
-		SelectedNormalOut = FVector3d::UnitZ();
-	}
-	bool bHaveFaceHit = (bActuallyHitSurface && Settings.bEnableFaceHits);
-	
-	// Deal with corner hits first (and edges that project to a corner)
-	FGroupTopologySelection CornerResults;
-	FVector3d CornerPosition;
-	int32 CornerSegmentEdgeID = 0;
-	bool bHaveCornerHit = false;
-	if (Settings.bEnableCornerHits || (Settings.bEnableEdgeHits && Settings.bPreferProjectedElement))
-	{
-		if (DoCornerBasedSelection(Settings, Ray, Spatial, TopoSpatial, CornerResults, CornerPosition, &CornerSegmentEdgeID))
-		{
-			bHaveCornerHit = true;
-		}
-	}
-
-	// If corner selection didn't yield results, try edge selection
-	FGroupTopologySelection EdgeResults;
-	FVector3d EdgePosition;
-	int32 EdgeSegmentEdgeID = 0;
-	bool bHaveEdgeHit = false;
-	if (Settings.bEnableEdgeHits || (Settings.bEnableFaceHits && Settings.bPreferProjectedElement))
-	{
-		if (DoEdgeBasedSelection(Settings, Ray, Spatial, TopoSpatial, EdgeResults, EdgePosition, &EdgeSegmentEdgeID))
-		{
-			bHaveEdgeHit = true;
-		}
-	}
-
-	// if we have both corner and edge, want to keep the one we are closer to
-	if (bHaveCornerHit && bHaveEdgeHit)
-	{
-		if (PointsWithinToleranceTest(CornerPosition, Ray.ClosestPoint(CornerPosition), 0.75))
-		{
-			bHaveEdgeHit = false;
-		}
-		else
-		{
-			bHaveCornerHit = false;
-		}
-	}
-
-	// if we have a corner or edge hit, and a face hit, pick face unless we are really close to corner/edge
-	if ((bHaveCornerHit || bHaveEdgeHit) && bHaveFaceHit)
-	{
-		FVector3d TestPos = (bHaveCornerHit) ? CornerPosition : EdgePosition;
-		if (!PointsWithinToleranceTest(TestPos, Ray.ClosestPoint(TestPos), 0.15))
-		{
-			bHaveEdgeHit = bHaveCornerHit = false;
-		}
-	}
+	// It's worth noting that the approach here breaks down in pathological cases where the same group is present
+	// multiple times around a corner (i.e. the group is not contiguous, and separate islands share a corner).
+	// It's not practical to worry about those cases.
 
+	NextEdgeIDOut = IndexConstants::InvalidID;
+	const FGroupTopology::FCorner& CurrentCorner = GroupTopology->Corners[CornerID];
 
-	if (bHaveCornerHit)
-	{
-		ResultOut = CornerResults;
-		SelectedPositionOut = CornerPosition;
-		if (EdgeSegmentIdOut != nullptr)
-		{
-			*EdgeSegmentIdOut = CornerSegmentEdgeID;
-		}
-		return true;
-	}
-	if (bHaveEdgeHit)
-	{
-		ResultOut = EdgeResults;
-		SelectedPositionOut = EdgePosition;
-		if (EdgeSegmentIdOut != nullptr)
-		{
-			*EdgeSegmentIdOut = EdgeSegmentEdgeID;
-		}
-		return true;
-	}
-	if (bHaveFaceHit)
+	if (CurrentCorner.NeighbourGroupIDs.Num() != 4)
 	{
-		// If we still haven't found a selection, go ahead and select the face that we found earlier
-		ResultOut.SelectedGroupIDs.Add(Topology->GetGroupID(HitTriangleID));
-		SelectedPositionOut = TriangleHitPos;
-		return true;
+		return false; // Not a valence 4 corner
 	}
 
+	const FGroupTopology::FGroupEdge& IncomingEdge = GroupTopology->Edges[IncomingEdgeID];
 
-	return false;
-}
-
-bool FGroupTopologySelector::DoCornerBasedSelection(const FSelectionSettings& Settings,
-	const FRay3d& Ray, FDynamicMeshAABBTree3* Spatial, const FGeometrySet3& TopoSpatial,
-	FGroupTopologySelection& ResultOut, FVector3d& SelectedPositionOut, int32 *EdgeSegmentIdOut) const
-{
-	// These will store our results, depending on whether we select all along the ray or not.
-	FGeometrySet3::FNearest SingleElement;
-	// ElementsWithinTolerance gives all nearby elements returned by the topology selector, whereas DownRayElements filters these to just
-	// the closest element and those that project directly onto it (and it only stores the IDs, since we don't care about anything else).
-	TArray<FGeometrySet3::FNearest> ElementsWithinTolerance;
-	TArray<int32> DownRayElements;
-
-	auto LocalTolTest = [this](const FVector3d& A, const FVector3d& B) { return PointsWithinToleranceTest(A, B, 1.0); };
+	// We want to find the edge that shares this corner but does not border either of the neighboring groups of
+	// the incoming edge.
 
-	// Start by getting the closest element
-	const FGeometrySet3::FNearest* ClosestElement = nullptr;
-	if (!Settings.bSelectDownRay)
+	for (int32 Gid : CurrentCorner.NeighbourGroupIDs)
 	{
-		if (TopoSpatial.FindNearestPointToRay(Ray, SingleElement, LocalTolTest))
+		if (Gid == IncomingEdge.Groups[0] || Gid == IncomingEdge.Groups[1])
 		{
-			ClosestElement = &SingleElement;
+			continue; // This is one of the neighboring groups of the incoming edge
 		}
-	}
-	else
-	{
-		// We're collecting all corners within tolerance, but we still need the closest element
-		if (TopoSpatial.CollectPointsNearRay(Ray, ElementsWithinTolerance, LocalTolTest))
-		{
-			double MinRayT = TNumericLimits<double>::Max();
-			for (const FGeometrySet3::FNearest& Element : ElementsWithinTolerance)
-			{
-				if (Element.RayParam < MinRayT)
-				{
-					MinRayT = Element.RayParam;
-					ClosestElement = &Element;
-				}
-			}
-		}//end if found corners
-	}//end if selecting down ray
-
-	// Bail early if we haven't found a corner
-	if (ClosestElement == nullptr)
-	{
-		return false;
-	}
 
-	// Also bail if the closest element is not visible.
-	if (!Settings.bIgnoreOcclusion && IsOccluded(*ClosestElement, Ray.Origin, Spatial, Settings.bHitBackFaces))
-	{
-		return false;
-	}
-
-	// The closest point is already found
-	SelectedPositionOut = ClosestElement->NearestGeoPoint;
-
-	// If we have other corners, we actually need to filter them to only those that lie in line with the closest element. Note that
-	// this would be done differently depending on whether we're dealing with an orthographic or perspective view: in an
-	// orthographic view, we need to see that they lie on a ray with view ray direction and closest corner origin, whereas in
-	// perspective, they would need to lie on a ray from camer through closest corner (which will differ due to tolerance).
-	// Because the "select down ray" behavior is only useful in orthographic viewports in the first place, we do it that way.
-	if (Settings.bSelectDownRay)
-	{
-		DownRayElements.Add(ClosestElement->ID);
-		for (const FGeometrySet3::FNearest& Element : ElementsWithinTolerance)
+		// Iterate through all edges of group
+		const FGroupTopology::FGroup* Group = GroupTopology->FindGroupByID(Gid);
+		for (const FGroupTopology::FGroupBoundary& Boundary : Group->Boundaries)
 		{
-			if (ClosestElement == &Element)
-			{
-				continue; // Already added
-			}
-
-			// Make sure that closest corner to current element is parallel with view ray
-			FVector3d ClosestTowardElement = UE::Geometry::Normalized(Element.NearestGeoPoint - ClosestElement->NearestGeoPoint);
-			// There would usually be one more abs() in here, but we know that other elements are down ray direction
-			if (abs(ClosestTowardElement.Dot(Ray.Direction) - 1.0) < KINDA_SMALL_NUMBER)
+			for (int32 Eid : Boundary.GroupEdges)
 			{
-				DownRayElements.Add(Element.ID);
-			}
-		}//end assembling aligned corners
-	}
+				const FGroupTopology::FGroupEdge& CandidateEdge = GroupTopology->Edges[Eid];
 
-	// Try to select edges that project to corners.
-	if (Settings.bPreferProjectedElement && Settings.bEnableEdgeHits)
-	{
-		TSet<int32> AddedTopologyEdges;
-
-		// See if the closest vertex has an attached edge that is colinear with the view ray. Due to the fact that
-		// topology "edges" are actually polylines, we could actually have more than one even for the closest corner
-		// (if the "edge" towards us curves away). We'll only grab an edge down our view ray.
-		int32 ClosestVid = Topology->GetCornerVertexID(ClosestElement->ID);
-		for (int32 Eid : Mesh->VtxEdgesItr(ClosestVid))
-		{
-			FDynamicMesh3::FEdge Edge = Mesh->GetEdge(Eid);
-			int32 OtherVid = (Edge.Vert.A == ClosestVid) ? Edge.Vert.B : Edge.Vert.A;
-			FVector3d EdgeVector = UE::Geometry::Normalized(Mesh->GetVertex(OtherVid) - Mesh->GetVertex(ClosestVid));
-			if (abs(EdgeVector.Dot(Ray.Direction) - 1.0) < KINDA_SMALL_NUMBER)
-			{
-				int TopologyEdgeIndex = Topology->FindGroupEdgeID(Eid);
-				if (TopologyEdgeIndex >= 0)
+				// Edge must share corner but not neighboring groups
+				if ((CandidateEdge.EndpointCorners[0] == CornerID || CandidateEdge.EndpointCorners[1] == CornerID)
+					&& CandidateEdge.Groups[0] != IncomingEdge.Groups[0] && CandidateEdge.Groups[0] != IncomingEdge.Groups[1]
+					&& CandidateEdge.Groups[1] != IncomingEdge.Groups[0] && CandidateEdge.Groups[1] != IncomingEdge.Groups[1])
 				{
-					ResultOut.SelectedEdgeIDs.Add(TopologyEdgeIndex);
-					AddedTopologyEdges.Add(TopologyEdgeIndex);
-
-					if (EdgeSegmentIdOut)
-					{
-						Topology->GetGroupEdgeEdges(TopologyEdgeIndex).Find(Eid, *EdgeSegmentIdOut);
-					}
-
-					break;
+					NextEdgeIDOut = Eid;
+					return true;
 				}
 			}
 		}
-
-		// If relevant, get all the other colinear edges
-		if (Settings.bSelectDownRay && AddedTopologyEdges.Num() > 0)
-		{
-			for (int i = 1; i < DownRayElements.Num(); ++i) // skip 0 because it is closest and done
-			{
-				// Look though any attached edges.
-				for (int32 Eid : Mesh->VtxEdgesItr(Topology->GetCornerVertexID(DownRayElements[i])))
-				{
-					FDynamicMesh3::FEdge Edge = Mesh->GetEdge(Eid);
-					FVector3d EdgeVector = UE::Geometry::Normalized(Mesh->GetVertex(Edge.Vert.A) - Mesh->GetVertex(Edge.Vert.B));
-
-					// Compare absolute value of dot product to 1. We already made sure that one of the vertices is in line
-					// with the closest vertex earlier.
-					if (abs(abs(EdgeVector.Dot(Ray.Direction)) - 1.0) < KINDA_SMALL_NUMBER)
-					{
-						int TopologyEdgeIndex = Topology->FindGroupEdgeID(Eid);
-						if (TopologyEdgeIndex >= 0 && !AddedTopologyEdges.Contains(TopologyEdgeIndex))
-						{
-							ResultOut.SelectedEdgeIDs.Add(TopologyEdgeIndex);
-							AddedTopologyEdges.Add(TopologyEdgeIndex);
-							// Don't break here because we may have parallel edges in both directions, since we aren't
-							// going through the vertices in a particular order.
-						}
-					}
-				}//end checking edges
-			}//end going through corners
-		}
-
-		// If we found edges, we're done
-		if (AddedTopologyEdges.Num() > 0)
-		{
-			return true;
-		}
-	}//end selecting projected edges
-
-	// If getting projected edges didn't work out, go ahead and add the corners.
-	if (Settings.bEnableCornerHits)
-	{
-		if (Settings.bSelectDownRay)
-		{
-			for (int32 Id : DownRayElements)
-			{
-				ResultOut.SelectedCornerIDs.Add(Id);
-			}
-		}
-		else
-		{
-			ResultOut.SelectedCornerIDs.Add(ClosestElement->ID);
-		}
-		return true;
 	}
-
 	return false;
 }
 
-bool FGroupTopologySelector::DoEdgeBasedSelection(const FSelectionSettings& Settings, const FRay3d& Ray,
-	FDynamicMeshAABBTree3* Spatial, const FGeometrySet3& TopoSpatial,
-	FGroupTopologySelection& ResultOut, FVector3d& SelectedPositionOut, int32* EdgeSegmentIdOut) const
+void FGroupTopologyUtils::AddNewEdgeRingEdges(int32 StartEdgeID, int32 ForwardGroupID, TSet<int32>& EdgeSet) const
 {
-	// These will store our results, depending on whether we select all along the ray or not.
-	FGeometrySet3::FNearest SingleElement;
-	// ElementsWithinTolerance gives all nearby elements returned by the topology selector, whereas DownRayElements filters these to just
-	// the closest element and those that project directly onto it (it stores the polyline and segment IDs, which are all we care about).
-	TArray<FGeometrySet3::FNearest> ElementsWithinTolerance;
-	TArray<FIndex2i> DownRayElements;
-
-	auto LocalTolTest = [this](const FVector3d& A, const FVector3d& B) { return PointsWithinToleranceTest(A, B, 1.0); };
-
-	// Start by getting the closest element
-	const FGeometrySet3::FNearest* ClosestElement = nullptr;
-	if (!Settings.bSelectDownRay)
-	{
-		if (TopoSpatial.FindNearestCurveToRay(Ray, SingleElement, LocalTolTest))
-		{
-			ClosestElement = &SingleElement;
-		}
-	}
-	else
-	{
-		// Need all curves within tolerance, but also need to know the closest.
-		if (TopoSpatial.CollectCurvesNearRay(Ray, ElementsWithinTolerance, LocalTolTest))
-		{
-			double MinRayT = TNumericLimits<double>::Max();
-			for (const FGeometrySet3::FNearest& Element : ElementsWithinTolerance)
-			{
-				if (Element.RayParam < MinRayT)
-				{
-					MinRayT = Element.RayParam;
-					ClosestElement = &Element;
-				}
-			}
-		}//end if found edges
-	}//end if selecting down ray
-
-	// Bail early if we haven't found at least one edge
-	if (ClosestElement == nullptr)
-	{
-		return false;
-	}
-
-	// Also bail if the closest element is not visible.
-	if (!Settings.bIgnoreOcclusion && IsOccluded(*ClosestElement, Ray.Origin, Spatial, Settings.bHitBackFaces))
-	{
-		return false;
-	}
-
-	// The closest point is already found
-	SelectedPositionOut = ClosestElement->NearestGeoPoint;
-
-	// If we have other edges, we need to filter them to only those that project onto the closest element. This would be done
-	// differently for perspective cameras vs orthographic projection, but since the behavior is only useful in ortho mode,
-	// we do it that way.
-	if (Settings.bSelectDownRay)
+	int32 CurrentEdgeID = StartEdgeID;
+	int32 CurrentForwardGroupID = ForwardGroupID;
+	while (true)
 	{
-		// Closest element is a given
-		DownRayElements.Add(FIndex2i(ClosestElement->ID, ClosestElement->PolySegmentIdx));
-
-		// We want edges that lie in a plane through the closest edge that is coplanar with the view direction. If we had wanted 
-		// to do this for perspective mode, we would have picked a plane through the edge and the camera origin.
-		int32 ClosestEid = Topology->GetGroupEdgeEdges(ClosestElement->ID)[ClosestElement->PolySegmentIdx];
-		FDynamicMesh3::FEdge ClosestEdge = Mesh->GetEdge(ClosestEid);
-		UE::Geometry::FPlane3d PlaneThroughClosestEdge(Mesh->GetVertex(ClosestEdge.Vert.A), Mesh->GetVertex(ClosestEdge.Vert.B), 
-			Mesh->GetVertex(ClosestEdge.Vert.A) + Ray.Direction);
-
-		for (const FGeometrySet3::FNearest& Element : ElementsWithinTolerance)
+		if (CurrentForwardGroupID == IndexConstants::InvalidID)
 		{
-			if (ClosestElement == &Element)
-			{
-				continue; // already added
-			}
-
-			// See if the edge endpoints lie in the plane
-			int32 Eid = Topology->GetGroupEdgeEdges(Element.ID)[Element.PolySegmentIdx];
-			FDynamicMesh3::FEdge Edge = Mesh->GetEdge(Eid);
-			if (abs(PlaneThroughClosestEdge.DistanceTo(Mesh->GetVertex(Edge.Vert.A))) < KINDA_SMALL_NUMBER
-				&& abs(PlaneThroughClosestEdge.DistanceTo(Mesh->GetVertex(Edge.Vert.B))) < KINDA_SMALL_NUMBER)
-			{
-				DownRayElements.Add(FIndex2i(Element.ID, Element.PolySegmentIdx));
-			}
+			break; // Reached a boundary
 		}
-	}
-
-	// Try to select faces that project to the closest edge
-	if (Settings.bPreferProjectedElement && Settings.bEnableFaceHits)
-	{
-		TSet<int32> AddedGroups;
-
-		// Start with the closest edge
-		int32 Eid = Topology->GetGroupEdgeEdges(ClosestElement->ID)[ClosestElement->PolySegmentIdx];
-		FDynamicMesh3::FEdge Edge = Mesh->GetEdge(Eid);
-		FVector3d VertA = Mesh->GetVertex(Edge.Vert.A);
-		FVector3d VertB = Mesh->GetVertex(Edge.Vert.B);
-
-
-		// Grab a plane through the two verts that contains the ray direction (again, this is assuming that we'd
-		// only do this in ortho mode, otherwise the plane would go through the ray origin.
-		UE::Geometry::FPlane3d PlaneThroughClosestEdge(VertA, VertB, VertA + Ray.Direction);
 
-		// Checking that the face is coplanar simply entails checking that the opposite vert is in the plane. However,
-		// it is possible even for the closest edge to have multiple coplanar faces if a group going towards
-		// the camera curves away. We need to be able to select just one, and we want to select the one that
-		// extends down the view ray. For that, we'll make sure that a vector to the opposite vertex lies on the same
-		// side of the edge as the view ray.
-		FVector3d EdgeVector = VertB - VertA;
-		FVector3d EdgeVecCrossDirection = EdgeVector.Cross(Ray.Direction);
-		FIndex2i OppositeVids = Mesh->GetEdgeOpposingV(Eid);
-
-		FVector3d OppositeVert = Mesh->GetVertex(OppositeVids.A);
-		if (abs(PlaneThroughClosestEdge.DistanceTo(OppositeVert)) < KINDA_SMALL_NUMBER
-			&& EdgeVector.Cross(OppositeVert - VertA).Dot(EdgeVecCrossDirection) > 0)
+		int32 NextEdgeID;
+		if (!GetQuadOppositeEdge(CurrentEdgeID, CurrentForwardGroupID, NextEdgeID))
 		{
-			int GroupId = Topology->GetGroupID(Edge.Tri.A);
-			ResultOut.SelectedGroupIDs.Add(GroupId);
-			AddedGroups.Add(GroupId);
+			break; // Probably not a quad
 		}
-		else if (OppositeVids.B != FDynamicMesh3::InvalidID)
+		if (EdgeSet.Contains(NextEdgeID))
 		{
-			OppositeVert = Mesh->GetVertex(OppositeVids.B);
-
-			if (abs(PlaneThroughClosestEdge.DistanceTo(OppositeVert)) < KINDA_SMALL_NUMBER
-				&& EdgeVector.Cross(OppositeVert - VertA).Dot(EdgeVecCrossDirection) > 0)
-			{
-				int GroupId = Topology->GetGroupID(Edge.Tri.B);
-				ResultOut.SelectedGroupIDs.Add(GroupId);
-				AddedGroups.Add(GroupId);
-			}
+			break; // Either we finished the loop, or we'll continue it from another selection
 		}
 
-		// If relevant, get all the other coplanar faces
-		if (Settings.bSelectDownRay && AddedGroups.Num() > 0)
-		{
-			for (int i = 1; i < DownRayElements.Num(); ++i) // skip 0 because it is closest and done
-			{
-				// We already made sure that all these edges are coplanar, so we'll just be checking opposite verts.
-				Eid = Topology->GetGroupEdgeEdges(DownRayElements[i].A)[DownRayElements[i].B];
-				Edge = Mesh->GetEdge(Eid);
-				OppositeVids = Mesh->GetEdgeOpposingV(Eid);
-
-				// No need to check directionality of faces here, just need them to be coplanar.
-				if (abs(PlaneThroughClosestEdge.DistanceTo(Mesh->GetVertex(OppositeVids.A))) < KINDA_SMALL_NUMBER)
-				{
-					int GroupId = Topology->GetGroupID(Edge.Tri.A);
-					if (!AddedGroups.Contains(GroupId))
-					{
-						ResultOut.SelectedGroupIDs.Add(GroupId);
-						AddedGroups.Add(GroupId);
-					}
-				}
-				if (OppositeVids.B != FDynamicMesh3::InvalidID
-					&& abs(PlaneThroughClosestEdge.DistanceTo(Mesh->GetVertex(OppositeVids.B))) < KINDA_SMALL_NUMBER)
-				{
-					int GroupId = Topology->GetGroupID(Edge.Tri.B);
-					if (!AddedGroups.Contains(GroupId))
-					{
-						ResultOut.SelectedGroupIDs.Add(GroupId);
-						AddedGroups.Add(GroupId);
-					}
-				}
-			}//end going through other edges
-		}
+		EdgeSet.Add(NextEdgeID);
 
-		// If we selected faces, we're done
-		if (AddedGroups.Num() > 0)
-		{
-			return true;
-		}
+		CurrentEdgeID = NextEdgeID;
+		const FGroupTopology::FGroupEdge& Edge = GroupTopology->Edges[CurrentEdgeID];
+		CurrentForwardGroupID = (Edge.Groups[0] == CurrentForwardGroupID) ? Edge.Groups[1] : Edge.Groups[0];
 	}
+}
 
-	// If we didn't end up selecting projected faces, and we have edges to select, select them
-	if (Settings.bEnableEdgeHits)
+bool FGroupTopologyUtils::GetQuadOppositeEdge(int32 EdgeIDIn, int32 GroupID, int32& OppositeEdgeIDOut) const
+{
+	const FGroupTopology::FGroup* Group = GroupTopology->FindGroupByID(GroupID);
+	check(Group);
+
+	// Find the boundary that contains this edge
+	for (int32 i = 0; i < Group->Boundaries.Num(); ++i)
 	{
-		if (Settings.bSelectDownRay)
+		const FGroupTopology::FGroupBoundary& Boundary = Group->Boundaries[i];
+		int32 EdgeIndex = Boundary.GroupEdges.IndexOfByKey(EdgeIDIn);
+		if (EdgeIndex != INDEX_NONE)
 		{
-			for (const FIndex2i& ElementTuple : DownRayElements)
+			if (Boundary.GroupEdges.Num() != 4)
 			{
-				ResultOut.SelectedEdgeIDs.Add(ElementTuple.A);
+				return false;
 			}
-		}
-		else
-		{
-			ResultOut.SelectedEdgeIDs.Add(ClosestElement->ID);
-		}
 
-		if (EdgeSegmentIdOut)
-		{
-			*EdgeSegmentIdOut = ClosestElement->PolySegmentIdx;
+			OppositeEdgeIDOut = Boundary.GroupEdges[(EdgeIndex + 2) % 4];
+			return true;
 		}
-
-		return true;
 	}
+	check(false); // No boundary of the given group contained the given edge
 	return false;
 }
 
-bool IsOccluded(const FGeometrySet3::FNearest& ClosestElement, const FVector3d& ViewOrigin, 
-	const FDynamicMeshAABBTree3* Spatial, bool bBackFacesOcclude)
-{
-	return IsOccluded(ClosestElement.NearestGeoPoint, ViewOrigin, Spatial, bBackFacesOcclude);
-}
 
-bool IsOccluded(const FVector3d& Point, const FVector3d& ViewOrigin, 
-	const FDynamicMeshAABBTree3* Spatial, bool bBackFacesOcclude)
+FIndex2i FGroupTopologyUtils::GetEdgeEndpointCorners(int EdgeID) const
 {
-	// We will shoot a ray backwards to see if we hit something. 
-	FRay3d ToEyeRay(Point, UE::Geometry::Normalized(ViewOrigin - Point), true);
-	ToEyeRay.Origin += (double)(100 * FMathf::ZeroTolerance) * ToEyeRay.Direction;
-
-	IMeshSpatial::FQueryOptions QueryOptions;
-	if (!bBackFacesOcclude)
-	{
-		// If back faces don't occlude (ie we weren't hitting back faces on the way forward), 
-		// then the direction toward the camera should be same as normal direction.
-		QueryOptions.TriangleFilterF = [Spatial, &ToEyeRay](int32 Tid)
-		{
-			return Spatial->GetMesh()->GetTriNormal(Tid).Dot(ToEyeRay.Direction) > 0;
-		};
-	}
-
-	if (Spatial->FindNearestHitTriangle(ToEyeRay, QueryOptions) >= 0)
-	{
-		return true;
-	}
-	return false;
+	return GroupTopology->Edges[EdgeID].EndpointCorners;
 }
 
-bool FGroupTopologySelector::FindSelectedElement(const FSelectionSettings& Settings, 
-	const FCameraRectangle& CameraRectangle, FTransform3d TargetTransform, FGroupTopologySelection& ResultOut,
-	TMap<int32, bool>* TriIsOccludedCache)
+FIndex2i FGroupTopologyUtils::GetEdgeGroups(int EdgeID) const
 {
-	// One minor easy thing we can do to speed up the below is to detect cases where transform does not have a non-uniform
-	// scale, and in those cases we transform the rectangle once rather than transforming all the query points/curves. 
-	// This isn't worth doing for non-uniform scale transforms because the rectangle basis stops being orthonormal and 
-	// becomes a problem to work with.
-
-	ResultOut.Clear();
-	FDynamicMeshAABBTree3* Spatial = GetSpatial();
-
-	// Needed for occlusion test, which happens in local space.
-	FVector3d LocalCameraOrigin = TargetTransform.InverseTransformPosition((FVector3d)CameraRectangle.CameraState.Position);
-	
-	// Corner selection takes priority over edges.
-	if (Settings.bEnableCornerHits)
-	{
-		GeometrySet.ParallelFindAllPointsSatisfying(
-			[&CameraRectangle, &TargetTransform, &Settings, &LocalCameraOrigin, Spatial]
-			(const FVector3d& PointPosition)
-			{
-				FVector TransformedPoint = TargetTransform.TransformPosition(PointPosition);
-				return CameraRectangle.IsProjectedPointInRectangle(TransformedPoint)
-					&& (Settings.bIgnoreOcclusion || !IsOccluded(PointPosition, LocalCameraOrigin, 
-						Spatial, Settings.bHitBackFaces));
-			},
-			ResultOut.SelectedCornerIDs);
-	}
-
-	// If we didn't get corners, look for edges.
-	if (ResultOut.SelectedCornerIDs.IsEmpty() && Settings.bEnableEdgeHits)
-	{
-		GeometrySet.ParallelFindAllCurvesSatisfying(
-			[&CameraRectangle, &TargetTransform, &Settings, &LocalCameraOrigin, Spatial](const FPolyline3d& Curve)
-			{
-				// Testing occlusion properly seems like it would be a pain, so we'll just consider something occluded
-				// if one of the endpoints is occluded. This will handle the common case of not wanting to select a
-				// hidden edge that is connected to a visible corner.
-				if (!Settings.bIgnoreOcclusion 
-					&& (IsOccluded(Curve.Start(), LocalCameraOrigin, Spatial, Settings.bHitBackFaces)
-						|| IsOccluded(Curve.End(), LocalCameraOrigin, Spatial, Settings.bHitBackFaces)))
-				{
-					return false;
-				}
-
-				// Check whether any of the component segments intersect the rectangle
-				const TArray<FVector3d>& Verts = Curve.GetVertices();
-				FVector3d CurrentVert = TargetTransform.TransformPosition(Verts[0]);
-				for (int32 i = 1; i < Verts.Num(); ++i)
-				{
-					FVector3d NextVert = TargetTransform.TransformPosition(Verts[i]);
-					if (CameraRectangle.IsProjectedSegmentIntersectingRectangle((FVector)CurrentVert, (FVector)NextVert))
-					{
-						return true;
-					}
-					CurrentVert = NextVert;
-				}
-				return false;
-			},
-			ResultOut.SelectedEdgeIDs);
-	}
-	
-	if (Settings.bEnableFaceHits && ResultOut.IsEmpty())
-	{
-		// Get a frustum volume in local space, then use it to figure out if it contains or intersects
-		// the aabb tree boxes in the traversal. 
-		// TODO: Frustum transformation is done the same way in FractureEditorMode.cpp- should this be placed
-		// somewhere common? Not sure where to put it though.
-		FConvexVolume WorldSpaceFrustum = CameraRectangle.FrustumAsConvexVolume();
-		FMatrix InverseTargetTransform(FTransform(TargetTransform).ToInverseMatrixWithScale());
-
-		FConvexVolume LocalFrustum;
-		LocalFrustum.Planes.Empty(6);
-		for (const FPlane& Plane : WorldSpaceFrustum.Planes)
-		{
-			LocalFrustum.Planes.Add(Plane.TransformBy(InverseTargetTransform));
-		}
-		LocalFrustum.Init();
-
-		FDynamicMeshAABBTree3::FTreeTraversal Traversal;
-		// Used as additional state for the traversal, relying on the fact that our traversal is depth-first
-		// and serial. When we descend to a box that is fully contained, we use CurrentSelectAllDepth to mark
-		// all deeper boxes as contained, and reset it once we return back out of that subtree (done similarly
-		// in UVEditorMeshSelectionMechanic).
-		int CurrentDepth = -1;
-		int CurrentSelectAllDepth = TNumericLimits<int>::Max();
-
-		Traversal.NextBoxF = [&LocalFrustum, &CurrentDepth, &CurrentSelectAllDepth](const FAxisAlignedBox3d& Box, int Depth)
-		{
-			CurrentDepth = Depth;
-			if (Depth > CurrentSelectAllDepth)
-			{
-				// We're currently inside a fully-selected box
-				return true;
-			}
-			// Reset CurrentSelectAllDepth because we must have left the cotained subtree (if it was set to begin with)
-			CurrentSelectAllDepth = TNumericLimits<int>::Max();
-
-			bool bFullyContained = false;
-			bool bIntersects = LocalFrustum.IntersectBox(Box.Center(), Box.Extents(), bFullyContained);
-			
-			if (bFullyContained)
-			{
-				CurrentSelectAllDepth = Depth;
-			}
-			return bIntersects;
-		};
-
-		Traversal.NextTriangleF = 
-			[this, &ResultOut, &CurrentDepth, &CurrentSelectAllDepth, &LocalFrustum, 
-			&LocalCameraOrigin, &Spatial, &Settings, TriIsOccludedCache](int Tid)
-		{
-			int32 TriGroupID = Topology->GetGroupID(Tid);
-
-			// If this group already got selected, no need to do anything else. Helps avoid some occlusion tests.
-			if (ResultOut.SelectedGroupIDs.Contains(TriGroupID))
-			{
-				return;
-			}
-
-			// Apply frustum check. Skip if box was fully in frustum
-			if (!(CurrentDepth >= CurrentSelectAllDepth))
-			{
-				// TODO: This isn't an ideal way to check whether the triangle is inside the frustum because FPoly is
-				// an old class that we probably want to avoid using. It's fine enough here, but it should be replaced.
-				// Related note: It may be tempting to just check for containment of the centroid, but we probably don't
-				// want to do that because it makes selection behavior dependent on the tesselation of the underlying
-				// group. I.e. a remeshed cube face ends up being much more reliable to select than one made of two triangles 
-				// where a marquee select of even the group center could fail if it's between the tri centroids.
-				FIndex3i TriVids = Mesh->GetTriangle(Tid);
-				FPoly TrianglePolygon;
-				TrianglePolygon.Init();
-				for (int i = 0; i < 3; ++i)
-				{
-					TrianglePolygon.Vertices.Add((FVector3f)Mesh->GetVertex(TriVids[i]));
-				}
-
-				if (!LocalFrustum.ClipPolygon(TrianglePolygon))
-				{
-					// Triangle was fully outside the frustum
-					return;
-				}
-			}
-
-			// Apply occlusion filter. Note that the back face filter is applied via TriangleFilterF in the traversal.
-			if (!Settings.bIgnoreOcclusion)
-			{
-				bool bTriangleIsOccluded = false;
-
-				// See if we've already calculated occlusion for this tri.
-				bool* CachedOcclusion = TriIsOccludedCache->Find(Tid);
-				if (CachedOcclusion)
-				{
-					bTriangleIsOccluded = *CachedOcclusion;
-				}
-				else
-				{
-					bTriangleIsOccluded = IsOccluded(Mesh->GetTriCentroid(Tid), LocalCameraOrigin, Spatial, Settings.bHitBackFaces);
-					TriIsOccludedCache->Add(Tid, bTriangleIsOccluded);
-				}
-
-				if (bTriangleIsOccluded)
-				{
-					return;
-				}
-			}
-
-			ResultOut.SelectedGroupIDs.Add(TriGroupID);
-		};
+	return GroupTopology->Edges[EdgeID].Groups;
+}
 
-		// Add the back face filter, if applicable
-		IMeshSpatial::FQueryOptions Options;
-		if (!Settings.bHitBackFaces)
-		{
-			Options.TriangleFilterF = [this, &LocalCameraOrigin](int32 Tid)
-			{
-				return Mesh->GetTriNormal(Tid).Dot(Mesh->GetVertex(Mesh->GetTriangle(Tid)[0]) - LocalCameraOrigin) < 0;
-			};
-		}
-		
-		Spatial->DoTraversal(Traversal, Options);
-	}
 
-	return ResultOut.IsEmpty();
-}
+// -------------------------------------------------
 
 bool FGroupTopologySelector::ExpandSelectionByEdgeLoops(FGroupTopologySelection& Selection)
 {
@@ -800,15 +153,15 @@ bool FGroupTopologySelector::ExpandSelectionByEdgeLoops(FGroupTopologySelection&
 	int32 OriginalNumEdges = Selection.SelectedEdgeIDs.Num();
 	for (int32 Eid : Selection.SelectedEdgeIDs)
 	{
-		const FGroupTopology::FGroupEdge& Edge = Topology->Edges[Eid];
-		if (Edge.EndpointCorners[0] == IndexConstants::InvalidID)
+		const FIndex2i EndpointCorners = GroupTopologyUtils.GetEdgeEndpointCorners(Eid);
+		if (EndpointCorners[0] == IndexConstants::InvalidID)
 		{
 			continue; // This FGroupEdge is a loop unto itself (and already in our selection, since we're looking at it).
 		}
 
 		// Go forward and backward adding edges
-		AddNewEdgeLoopEdgesFromCorner(*Topology, Eid, Edge.EndpointCorners[0], EdgeSet);
-		AddNewEdgeLoopEdgesFromCorner(*Topology, Eid, Edge.EndpointCorners[1], EdgeSet);
+		GroupTopologyUtils.AddNewEdgeLoopEdgesFromCorner(Eid, EndpointCorners[0], EdgeSet);
+		GroupTopologyUtils.AddNewEdgeLoopEdgesFromCorner(Eid, EndpointCorners[1], EdgeSet);
 	}
 
 	if (EdgeSet.Num() > OriginalNumEdges)
@@ -822,96 +175,22 @@ bool FGroupTopologySelector::ExpandSelectionByEdgeLoops(FGroupTopologySelection&
 	}
 }
 
-void AddNewEdgeLoopEdgesFromCorner(const FGroupTopology& Topology, int32 EdgeID, int32 CornerID, TSet<int32>& EdgeSet)
-{
-	int32 LastCornerID = CornerID;
-	int32 LastEdgeID = EdgeID;
-	while (true)
-	{
-		int32 NextEid;
-		if (!GetNextEdgeLoopEdge(Topology, LastEdgeID, LastCornerID, NextEid))
-		{
-			break; // Probably not a valence 4 corner
-		}
-		if (EdgeSet.Contains(NextEid))
-		{
-			break; // Either we finished the loop, or we'll continue it from another selection
-		}
-
-		EdgeSet.Add(NextEid);
-
-		LastEdgeID = NextEid;
-		const FGroupTopology::FGroupEdge& LastEdge = Topology.Edges[LastEdgeID];
-		LastCornerID = LastEdge.EndpointCorners[0] == LastCornerID ? LastEdge.EndpointCorners[1] : LastEdge.EndpointCorners[0];
-
-		check(LastCornerID != IndexConstants::InvalidID);
-	}
-}
-
-bool GetNextEdgeLoopEdge(const FGroupTopology& Topology, int32 IncomingEdgeID, int32 CornerID, int32& NextEdgeIDOut)
-{
-	// It's worth noting that the approach here breaks down in pathological cases where the same group is present
-	// multiple times around a corner (i.e. the group is not contiguous, and separate islands share a corner).
-	// It's not practical to worry about those cases.
-
-	NextEdgeIDOut = IndexConstants::InvalidID;
-	const FGroupTopology::FCorner& CurrentCorner = Topology.Corners[CornerID];
-
-	if (CurrentCorner.NeighbourGroupIDs.Num() != 4)
-	{
-		return false; // Not a valence 4 corner
-	}
-
-	const FGroupTopology::FGroupEdge& IncomingEdge = Topology.Edges[IncomingEdgeID];
-
-	// We want to find the edge that shares this corner but does not border either of the neighboring groups of
-	// the incoming edge.
-
-	for (int32 Gid : CurrentCorner.NeighbourGroupIDs)
-	{
-		if (Gid == IncomingEdge.Groups[0] || Gid == IncomingEdge.Groups[1])
-		{
-			continue; // This is one of the neighboring groups of the incoming edge
-		}
-
-		// Iterate through all edges of group
-		const FGroupTopology::FGroup* Group = Topology.FindGroupByID(Gid);
-		for (const FGroupTopology::FGroupBoundary& Boundary : Group->Boundaries)
-		{
-			for (int32 Eid : Boundary.GroupEdges)
-			{
-				const FGroupTopology::FGroupEdge& CandidateEdge = Topology.Edges[Eid];
-
-				// Edge must share corner but not neighboring groups
-				if ((CandidateEdge.EndpointCorners[0] == CornerID || CandidateEdge.EndpointCorners[1] == CornerID)
-					&& CandidateEdge.Groups[0] != IncomingEdge.Groups[0] && CandidateEdge.Groups[0] != IncomingEdge.Groups[1]
-					&& CandidateEdge.Groups[1] != IncomingEdge.Groups[0] && CandidateEdge.Groups[1] != IncomingEdge.Groups[1])
-				{
-					NextEdgeIDOut = Eid;
-					return true;
-				}
-			}
-		}
-	}
-	return false;
-}
-
 bool FGroupTopologySelector::ExpandSelectionByEdgeRings(FGroupTopologySelection& Selection)
 {
 	TSet<int32> EdgeSet(Selection.SelectedEdgeIDs);
 	int32 OriginalNumEdges = Selection.SelectedEdgeIDs.Num();
 	for (int32 Eid : Selection.SelectedEdgeIDs)
 	{
-		const FGroupTopology::FGroupEdge& Edge = Topology->Edges[Eid];
+		const FIndex2i EdgeGroups = GroupTopologyUtils.GetEdgeGroups(Eid);
 
 		// Go forward and backward adding edges
-		if (Edge.Groups[0] != IndexConstants::InvalidID)
+		if (EdgeGroups[0] != IndexConstants::InvalidID)
 		{
-			AddNewEdgeRingEdges(*Topology, Eid, Edge.Groups[0], EdgeSet);
+			GroupTopologyUtils.AddNewEdgeRingEdges(Eid, EdgeGroups[0], EdgeSet);
 		}
-		if (Edge.Groups[0] != IndexConstants::InvalidID)
+		if (EdgeGroups[0] != IndexConstants::InvalidID)
 		{
-			AddNewEdgeRingEdges(*Topology, Eid, Edge.Groups[1], EdgeSet);
+			GroupTopologyUtils.AddNewEdgeRingEdges(Eid, EdgeGroups[1], EdgeSet);
 		}
 	}
 
@@ -926,58 +205,23 @@ bool FGroupTopologySelector::ExpandSelectionByEdgeRings(FGroupTopologySelection&
 	}
 }
 
-void AddNewEdgeRingEdges(const FGroupTopology& Topology, int32 StartEdgeID, int32 ForwardGroupID, TSet<int32>& EdgeSet)
+FGroupTopologySelector::FGroupTopologySelector(const FDynamicMesh3* Mesh, const FGroupTopology* Topology) : 
+	FMeshTopologySelector()
 {
-	int32 CurrentEdgeID = StartEdgeID;
-	int32 CurrentForwardGroupID = ForwardGroupID;
-	while (true)
-	{
-		if (CurrentForwardGroupID == IndexConstants::InvalidID)
-		{
-			break; // Reached a boundary
-		}
-
-		int32 NextEdgeID;
-		if (!GetQuadOppositeEdge(Topology, CurrentEdgeID, CurrentForwardGroupID, NextEdgeID))
-		{
-			break; // Probably not a quad
-		}
-		if (EdgeSet.Contains(NextEdgeID))
-		{
-			break; // Either we finished the loop, or we'll continue it from another selection
-		}
-
-		EdgeSet.Add(NextEdgeID);
-
-		CurrentEdgeID = NextEdgeID;
-		const FGroupTopology::FGroupEdge& Edge = Topology.Edges[CurrentEdgeID];
-		CurrentForwardGroupID = (Edge.Groups[0] == CurrentForwardGroupID) ? Edge.Groups[1] : Edge.Groups[0];
-	}
+	Initialize(Mesh, Topology);
 }
 
-bool GetQuadOppositeEdge(const FGroupTopology& Topology, int32 EdgeIDIn, int32 GroupID, int32& OppositeEdgeIDOut)
+void FGroupTopologySelector::Initialize(const FDynamicMesh3* MeshIn, const FGroupTopology* TopologyIn)
 {
-	const FGroupTopology::FGroup* Group = Topology.FindGroupByID(GroupID);
-	check(Group);
-
-	// Find the boundary that contains this edge
-	for (int32 i = 0; i < Group->Boundaries.Num(); ++i)
-	{
-		const FGroupTopology::FGroupBoundary& Boundary = Group->Boundaries[i];
-		int32 EdgeIndex = Boundary.GroupEdges.IndexOfByKey(EdgeIDIn);
-		if (EdgeIndex != INDEX_NONE)
-		{
-			if (Boundary.GroupEdges.Num() != 4)
-			{
-				return false;
-			}
+	Mesh = MeshIn;
+	TopologyProvider = MakeUnique<FGroupTopologyProvider>(TopologyIn);
+	GroupTopologyUtils.GroupTopology = TopologyIn;
+	bGeometryInitialized = false;
+	bGeometryUpToDate = false;
 
-			OppositeEdgeIDOut = Boundary.GroupEdges[(EdgeIndex + 2) % 4];
-			return true;
-		}
-	}
-	check(false); // No boundary of the given group contained the given edge
-	return false;
+	PRAGMA_DISABLE_DEPRECATION_WARNINGS
+	Topology = TopologyIn;
+	PRAGMA_ENABLE_DEPRECATION_WARNINGS
 }
 
 
@@ -990,7 +234,7 @@ void FGroupTopologySelector::DrawSelection(const FGroupTopologySelection& Select
 	{
 		for (int CornerID : Selection.SelectedCornerIDs)
 		{
-			int VertexID = Topology->GetCornerVertexID(CornerID);
+			int VertexID = TopologyProvider->GetCornerVertexID(CornerID);
 			FVector Position = (FVector)Mesh->GetVertex(VertexID);
 
 			Renderer->DrawPoint(Position, Renderer->PointColor, Renderer->PointSize, false);
@@ -1000,7 +244,7 @@ void FGroupTopologySelector::DrawSelection(const FGroupTopologySelection& Select
 	{
 		for (int CornerID : Selection.SelectedCornerIDs)
 		{
-			int VertexID = Topology->GetCornerVertexID(CornerID);
+			int VertexID = TopologyProvider->GetCornerVertexID(CornerID);
 			FVector Position = (FVector)Mesh->GetVertex(VertexID);
 			FVector WorldPosition = Renderer->TransformP(Position);
 
@@ -1015,7 +259,7 @@ void FGroupTopologySelector::DrawSelection(const FGroupTopologySelection& Select
 
 	for (int EdgeID : Selection.SelectedEdgeIDs)
 	{
-		const TArray<int>& Vertices = Topology->GetGroupEdgeVertices(EdgeID);
+		const TArray<int>& Vertices = TopologyProvider->GetGroupEdgeVertices(EdgeID);
 		int NV = Vertices.Num() - 1;
 
 		// Draw the edge, but also draw the endpoints in ortho mode (to make projected edges visible)
@@ -1026,7 +270,7 @@ void FGroupTopologySelector::DrawSelection(const FGroupTopologySelection& Select
 		}
 		for (int k = 0; k < NV; ++k)
 		{
-			FVector B = (FVector)Mesh->GetVertex(Vertices[k+1]);
+			FVector B = (FVector)Mesh->GetVertex(Vertices[k + 1]);
 			Renderer->DrawLine(A, B, UseColor, LineWidth, false);
 			A = B;
 		}
@@ -1040,10 +284,10 @@ void FGroupTopologySelector::DrawSelection(const FGroupTopologySelection& Select
 	// when they project to an edge.
 	if (CameraState->bIsOrthographic && Selection.SelectedGroupIDs.Num() > 0)
 	{
-		Topology->ForGroupSetEdges(Selection.SelectedGroupIDs,
-			[&UseColor, LineWidth, Renderer, this] (FGroupTopology::FGroupEdge Edge, int EdgeID) 
+		TopologyProvider->ForGroupSetEdges(Selection.SelectedGroupIDs,
+			[&UseColor, LineWidth, Renderer, this](int EdgeID)
 		{
-			const TArray<int>& Vertices = Topology->GetGroupEdgeVertices(EdgeID);
+			const TArray<int>& Vertices = TopologyProvider->GetGroupEdgeVertices(EdgeID);
 			int NV = Vertices.Num() - 1;
 			FVector A = (FVector)Mesh->GetVertex(Vertices[0]);
 			for (int k = 0; k < NV; ++k)
@@ -1054,4 +298,7 @@ void FGroupTopologySelector::DrawSelection(const FGroupTopologySelection& Select
 			}
 		});
 	}
-}
\ No newline at end of file
+}
+
+
+#undef LOCTEXT_NAMESPACE
diff --git a/Engine/Plugins/Runtime/MeshModelingToolset/Source/ModelingComponents/Private/Selection/MeshTopologySelectionMechanic.cpp b/Engine/Plugins/Runtime/MeshModelingToolset/Source/ModelingComponents/Private/Selection/MeshTopologySelectionMechanic.cpp
new file mode 100644
--- /dev/null
+++ b/Engine/Plugins/Runtime/MeshModelingToolset/Source/ModelingComponents/Private/Selection/MeshTopologySelectionMechanic.cpp
@@ -0,0 +1,761 @@
+// Copyright Epic Games, Inc. All Rights Reserved.
+
+#include "Selection/MeshTopologySelectionMechanic.h"
+
+#include "BaseBehaviors/SingleClickOrDragBehavior.h"
+#include "BaseBehaviors/MouseHoverBehavior.h"
+#include "InteractiveToolManager.h"
+#include "Util/ColorConstants.h"
+#include "Selection/PersistentMeshSelection.h"
+#include "ToolSceneQueriesUtil.h"
+#include "ToolSetupUtil.h"
+
+using namespace UE::Geometry;
+
+#define LOCTEXT_NAMESPACE "UMeshTopologySelectionMechanic"
+
+void UMeshTopologySelectionMechanicProperties::InvertSelection()
+{
+	if (Mechanic.IsValid())
+	{
+		Mechanic->InvertSelection();
+	}
+}
+
+void UMeshTopologySelectionMechanicProperties::SelectAll()
+{
+	if (Mechanic.IsValid())
+	{
+		Mechanic->SelectAll();
+	}
+}
+
+UMeshTopologySelectionMechanic::~UMeshTopologySelectionMechanic()
+{
+	checkf(PreviewGeometryActor == nullptr, TEXT("Shutdown() should be called before UMeshTopologySelectionMechanic is destroyed."));
+}
+
+void UMeshTopologySelectionMechanic::Setup(UInteractiveTool* ParentToolIn)
+{
+	UInteractionMechanic::Setup(ParentToolIn);
+
+	HoverBehavior = NewObject<UMouseHoverBehavior>();
+	// We use modifiers on hover to change the highlighting according to what can be selected  	
+	HoverBehavior->Modifiers.RegisterModifier(ShiftModifierID, FInputDeviceState::IsShiftKeyDown);
+	HoverBehavior->Modifiers.RegisterModifier(CtrlModifierID, FInputDeviceState::IsCtrlKeyDown);
+	HoverBehavior->Initialize(this);
+	HoverBehavior->SetDefaultPriority(BasePriority);
+	ParentToolIn->AddInputBehavior(HoverBehavior, this);
+
+	MarqueeMechanic = NewObject<URectangleMarqueeMechanic>();
+	MarqueeMechanic->bUseExternalClickDragBehavior = true;
+	MarqueeMechanic->Setup(ParentToolIn);
+	MarqueeMechanic->OnDragRectangleStarted.AddUObject(this, &UMeshTopologySelectionMechanic::OnDragRectangleStarted);
+	MarqueeMechanic->OnDragRectangleChanged.AddUObject(this, &UMeshTopologySelectionMechanic::OnDragRectangleChanged);
+	MarqueeMechanic->OnDragRectangleFinished.AddUObject(this, &UMeshTopologySelectionMechanic::OnDragRectangleFinished);
+	MarqueeMechanic->SetBasePriority(BasePriority.MakeLower());
+
+	ClickOrDragBehavior = NewObject<USingleClickOrDragInputBehavior>();
+	ClickOrDragBehavior->Initialize(this, MarqueeMechanic);
+	ClickOrDragBehavior->Modifiers.RegisterModifier(ShiftModifierID, FInputDeviceState::IsShiftKeyDown);
+	ClickOrDragBehavior->Modifiers.RegisterModifier(CtrlModifierID, FInputDeviceState::IsCtrlKeyDown);
+	ClickOrDragBehavior->SetDefaultPriority(BasePriority);
+	ParentTool->AddInputBehavior(ClickOrDragBehavior, this);
+
+	Properties = NewObject<UMeshTopologySelectionMechanicProperties>(this);
+	Properties->Initialize(this);
+	if (bAddSelectionFilterPropertiesToParentTool)
+	{
+		AddToolPropertySource(Properties);
+	}
+	Properties->WatchProperty(Properties->bSelectVertices, [this](bool bSelectVertices) {
+		UpdateMarqueeEnabled(); 
+	});
+	Properties->WatchProperty(Properties->bSelectEdges, [this](bool bSelectVertices) { 
+		UpdateMarqueeEnabled(); 
+	});
+	Properties->WatchProperty(Properties->bSelectFaces, [this](bool bSelectFaces) {
+		UpdateMarqueeEnabled();
+	});
+	Properties->WatchProperty(Properties->bEnableMarquee, [this](bool bEnableMarquee) { 
+		UpdateMarqueeEnabled(); 
+	});
+
+	// set up visualizers
+	PolyEdgesRenderer.LineColor = FLinearColor::Red;
+	PolyEdgesRenderer.LineThickness = 2.0;
+	PolyEdgesRenderer.PointColor = FLinearColor::Red;
+	PolyEdgesRenderer.PointSize = 8.0f;
+	HilightRenderer.LineColor = FLinearColor::Green;
+	HilightRenderer.LineThickness = 4.0f;
+	HilightRenderer.PointColor = FLinearColor::Green;
+	HilightRenderer.PointSize = 10.0f;
+	SelectionRenderer.LineColor = LinearColors::Gold3f();
+	SelectionRenderer.LineThickness = 4.0f;
+	SelectionRenderer.PointColor = LinearColors::Gold3f();
+	SelectionRenderer.PointSize = 10.0f;
+
+	float HighlightedFacePercentDepthOffset = 0.5f;
+	HighlightedFaceMaterial = ToolSetupUtil::GetSelectionMaterial(FLinearColor::Green, ParentToolIn->GetToolManager(), HighlightedFacePercentDepthOffset);
+	// The rest of the highlighting setup has to be done in Initialize(), since we need the world to set up our drawing component.
+}
+
+void UMeshTopologySelectionMechanic::Shutdown()
+{
+	if (PreviewGeometryActor)
+	{
+		PreviewGeometryActor->Destroy();
+		PreviewGeometryActor = nullptr;
+	}
+}
+
+void UMeshTopologySelectionMechanic::DisableBehaviors(UInteractiveTool* ParentToolIn)
+{
+	ParentToolIn->RemoveInputBehaviorsBySource(this);
+	ParentToolIn->RemoveInputBehaviorsBySource(MarqueeMechanic);
+
+	// TODO: Is it worth adding a way to remove the property watchers for marquee?
+}
+
+void UMeshTopologySelectionMechanic::SetIsEnabled(bool bOn)
+{
+	bIsEnabled = bOn;
+	UpdateMarqueeEnabled();
+}
+
+void UMeshTopologySelectionMechanic::SetBasePriority(const FInputCapturePriority &Priority)
+{
+	BasePriority = Priority;
+	if (ClickOrDragBehavior)
+	{
+		ClickOrDragBehavior->SetDefaultPriority(Priority);
+	}
+	if (HoverBehavior)
+	{
+		HoverBehavior->SetDefaultPriority(Priority);
+	}
+	if (MarqueeMechanic)
+	{
+		MarqueeMechanic->SetBasePriority(Priority.MakeLower());
+	}
+}
+
+TPair<FInputCapturePriority, FInputCapturePriority> UMeshTopologySelectionMechanic::GetPriorityRange() const
+{
+	TPair<FInputCapturePriority, FInputCapturePriority> Result;
+	Result.Key = BasePriority;
+	Result.Value = MarqueeMechanic->GetPriorityRange().Value;
+	return Result;
+}
+
+void UMeshTopologySelectionMechanic::Render(IToolsContextRenderAPI* RenderAPI)
+{
+	MarqueeMechanic->Render(RenderAPI);
+
+	// Cache the view camera state so we can use for snapping/etc.
+	// This should not happen in Render() though...
+	GetParentTool()->GetToolManager()->GetContextQueriesAPI()->GetCurrentViewState(CameraState);
+
+	FViewCameraState RenderCameraState = RenderAPI->GetCameraState();
+
+	const FDynamicMesh3* TargetMesh = this->Mesh;
+	FTransform Transform = (FTransform)TargetTransform;
+
+	const FTopologyProvider* TopologyProvider = TopoSelector->GetTopologyProvider();
+
+	PolyEdgesRenderer.BeginFrame(RenderAPI, RenderCameraState);
+	PolyEdgesRenderer.SetTransform(Transform);
+	for (int EdgeID = 0; EdgeID < TopologyProvider->GetNumEdges(); ++EdgeID)
+	{
+		FVector3d A, B;
+		for (int eid : TopologyProvider->GetGroupEdgeEdges(EdgeID))
+		{
+			TargetMesh->GetEdgeV(eid, A, B);
+			PolyEdgesRenderer.DrawLine(A, B);
+		}
+	}
+	if (bShowSelectableCorners)
+	{
+		for (int CornerID = 0; CornerID < TopologyProvider->GetNumCorners(); ++CornerID)
+		{
+			FVector3d A = TargetMesh->GetVertex(TopologyProvider->GetCornerVertexID(CornerID));
+			PolyEdgesRenderer.DrawPoint(A);
+		}
+	}
+	PolyEdgesRenderer.EndFrame();
+
+	if (PersistentSelection.IsEmpty() == false)
+	{
+		SelectionRenderer.BeginFrame(RenderAPI, RenderCameraState);
+		SelectionRenderer.SetTransform(Transform);
+		TopoSelector->DrawSelection(PersistentSelection, &SelectionRenderer, &RenderCameraState);
+		SelectionRenderer.EndFrame();
+	}
+
+	HilightRenderer.BeginFrame(RenderAPI, RenderCameraState);
+	HilightRenderer.SetTransform(Transform);
+	TopoSelector->DrawSelection(HilightSelection, &HilightRenderer, &RenderCameraState, FGroupTopologySelector::ECornerDrawStyle::Circle);
+	HilightRenderer.EndFrame();
+}
+
+void UMeshTopologySelectionMechanic::DrawHUD(FCanvas* Canvas, IToolsContextRenderAPI* RenderAPI)
+{
+	MarqueeMechanic->DrawHUD(Canvas, RenderAPI);
+}
+
+void UMeshTopologySelectionMechanic::ClearHighlight()
+{
+	checkf(DrawnTriangleSetComponent != nullptr, TEXT("Initialize() not called on UMeshTopologySelectionMechanic."));
+
+	HilightSelection.Clear();
+	DrawnTriangleSetComponent->Clear();
+	CurrentlyHighlightedGroups.Empty();
+}
+
+
+void UMeshTopologySelectionMechanic::NotifyMeshChanged(bool bTopologyModified)
+{
+	ClearHighlight();
+	TopoSelector->Invalidate(true, bTopologyModified);
+	if (bTopologyModified)
+	{
+		PersistentSelection.Clear();
+		SelectionTimestamp++;
+		OnSelectionChanged.Broadcast();
+	}
+}
+
+
+bool UMeshTopologySelectionMechanic::TopologyHitTest(const FRay& WorldRay, FHitResult& OutHit, bool bUseOrthoSettings)
+{
+	FGroupTopologySelection Selection;
+	return TopologyHitTest(WorldRay, OutHit, Selection, bUseOrthoSettings);
+}
+
+bool UMeshTopologySelectionMechanic::TopologyHitTest(const FRay& WorldRay, FHitResult& OutHit, FGroupTopologySelection& OutSelection, bool bUseOrthoSettings)
+{
+	FRay3d LocalRay(TargetTransform.InverseTransformPosition((FVector3d)WorldRay.Origin),
+		TargetTransform.InverseTransformVector((FVector3d)WorldRay.Direction));
+	UE::Geometry::Normalize(LocalRay.Direction);
+
+	FVector3d LocalPosition, LocalNormal;
+	int32 EdgeSegmentId; // Only used if hit is an edge
+	FGroupTopologySelector::FSelectionSettings TopoSelectorSettings = GetTopoSelectorSettings(bUseOrthoSettings);
+	if (TopoSelector->FindSelectedElement(TopoSelectorSettings, LocalRay, OutSelection, LocalPosition, LocalNormal, &EdgeSegmentId) == false)
+	{
+		return false;
+	}
+
+	if (OutSelection.SelectedCornerIDs.Num() > 0)
+	{
+		OutHit.FaceIndex = OutSelection.GetASelectedCornerID();
+		OutHit.Distance = LocalRay.GetParameter(LocalPosition);
+		OutHit.ImpactPoint = (FVector)TargetTransform.TransformPosition(LocalRay.PointAt(OutHit.Distance));
+	}
+	else if (OutSelection.SelectedEdgeIDs.Num() > 0)
+	{
+		OutHit.FaceIndex = OutSelection.GetASelectedEdgeID();
+		OutHit.Distance = LocalRay.GetParameter(LocalPosition);
+		OutHit.ImpactPoint = (FVector)TargetTransform.TransformPosition(LocalRay.PointAt(OutHit.Distance));
+		OutHit.Item = EdgeSegmentId;
+	}
+	else
+	{
+		FDynamicMeshAABBTree3* Spatial = GetSpatialFunc();
+		int HitTID = Spatial->FindNearestHitTriangle(LocalRay);
+		if (HitTID != IndexConstants::InvalidID)
+		{
+			FTriangle3d Triangle;
+			Spatial->GetMesh()->GetTriVertices(HitTID, Triangle.V[0], Triangle.V[1], Triangle.V[2]);
+			FIntrRay3Triangle3d Query(LocalRay, Triangle);
+			if (Query.Find())
+			{
+				OutHit.FaceIndex = HitTID;
+				OutHit.Distance = (float)Query.RayParameter;
+				OutHit.Normal = (FVector)TargetTransform.TransformVectorNoScale(Spatial->GetMesh()->GetTriNormal(HitTID));
+				OutHit.ImpactPoint = (FVector)TargetTransform.TransformPosition(LocalRay.PointAt(Query.RayParameter));
+			}
+			else
+			{
+				return false;
+			}
+		}
+	}
+	return true;
+}
+
+
+
+FGroupTopologySelector::FSelectionSettings UMeshTopologySelectionMechanic::GetTopoSelectorSettings(bool bUseOrthoSettings)
+{
+	FGroupTopologySelector::FSelectionSettings Settings;
+
+	Settings.bEnableFaceHits = Properties->bSelectFaces;
+	Settings.bEnableEdgeHits = Properties->bSelectEdges;
+	Settings.bEnableCornerHits = Properties->bSelectVertices;
+	Settings.bHitBackFaces = Properties->bHitBackFaces;
+
+	if (!PersistentSelection.IsEmpty() && (ShouldAddToSelectionFunc() || ShouldRemoveFromSelectionFunc()))
+	{
+		// If we have a selection and we're adding/removing/toggling elements make sure we only hit elements with compatible types
+		Settings.bEnableFaceHits = Settings.bEnableFaceHits && PersistentSelection.SelectedGroupIDs.Num() > 0;
+		Settings.bEnableEdgeHits = Settings.bEnableEdgeHits && PersistentSelection.SelectedEdgeIDs.Num() > 0;
+		Settings.bEnableCornerHits = Settings.bEnableCornerHits && PersistentSelection.SelectedCornerIDs.Num() > 0;
+	}
+
+	if (bUseOrthoSettings)
+	{
+		Settings.bPreferProjectedElement = Properties->bPreferProjectedElement;
+		Settings.bSelectDownRay = Properties->bSelectDownRay;
+		Settings.bIgnoreOcclusion = Properties->bIgnoreOcclusion;
+	}
+
+	return Settings;
+}
+
+bool UMeshTopologySelectionMechanic::HasSelection() const
+{
+	return PersistentSelection.IsEmpty() == false;
+}
+
+void UMeshTopologySelectionMechanic::SetSelection(const FGroupTopologySelection& Selection, bool bBroadcast)
+{
+	PersistentSelection = Selection;
+	SelectionTimestamp++;
+	if (bBroadcast)
+	{
+		OnSelectionChanged.Broadcast();
+	}
+}
+
+
+void UMeshTopologySelectionMechanic::ClearSelection()
+{
+	PersistentSelection.Clear();
+	SelectionTimestamp++;
+	OnSelectionChanged.Broadcast();
+}
+
+void UMeshTopologySelectionMechanic::InvertSelection()
+{
+	if (PersistentSelection.IsEmpty())
+	{
+		SelectAll();
+		return;
+	}
+
+	ParentTool->GetToolManager()->BeginUndoTransaction(LOCTEXT("SelectionChange", "Selection"));
+	BeginChange();
+
+	const FGroupTopologySelection PreviousSelection = PersistentSelection;
+	PersistentSelection.Clear();
+
+	const FTopologyProvider* TopologyProvider = TopoSelector->GetTopologyProvider();
+
+	if (!PreviousSelection.SelectedCornerIDs.IsEmpty())
+	{
+		for (int32 CornerID = 0; CornerID < TopologyProvider->GetNumCorners(); ++CornerID)
+		{
+			if (!PreviousSelection.SelectedCornerIDs.Contains(CornerID))
+			{
+				PersistentSelection.SelectedCornerIDs.Add(CornerID);
+			}
+		}
+	}
+	else if (!PreviousSelection.SelectedEdgeIDs.IsEmpty())
+	{
+		for (int32 EdgeID = 0; EdgeID < TopologyProvider->GetNumEdges(); ++EdgeID)
+		{
+			if (!PreviousSelection.SelectedEdgeIDs.Contains(EdgeID))
+			{
+				PersistentSelection.SelectedEdgeIDs.Add(EdgeID);
+			}
+		}
+	}
+	else if (!PreviousSelection.SelectedGroupIDs.IsEmpty())
+	{
+		for (int GroupIndex = 0; GroupIndex < TopologyProvider->GetNumGroups(); ++GroupIndex)
+		{
+			const int GroupID = TopologyProvider->GetGroupIDAt(GroupIndex);
+			if (!PreviousSelection.SelectedGroupIDs.Contains(GroupID))
+			{
+				PersistentSelection.SelectedGroupIDs.Add(GroupID);
+			}
+		}
+	}
+
+	SelectionTimestamp++;
+	OnSelectionChanged.Broadcast();
+	EndChangeAndEmitIfModified();
+	ParentTool->GetToolManager()->EndUndoTransaction();
+}
+
+void UMeshTopologySelectionMechanic::SelectAll()
+{
+	const FGroupTopologySelection PreviousSelection = PersistentSelection;
+
+	ParentTool->GetToolManager()->BeginUndoTransaction(LOCTEXT("SelectionChange", "Selection"));
+	BeginChange();
+
+	auto SelectAllIndices = [](int32 MaxExclusiveIndex, TSet<int32>& ContainerOut)
+	{
+		for (int32 i = 0; i < MaxExclusiveIndex; ++i)
+		{
+			ContainerOut.Add(i);
+		}
+	};
+
+	PersistentSelection.Clear();
+
+	const FTopologyProvider* TopologyProvider = TopoSelector->GetTopologyProvider();
+
+	// Select based on settings, prefering corners to edges to groups (since this is the preference we have
+	// elsewhere, eg in marquee).
+	if (Properties->bSelectVertices)
+	{
+		SelectAllIndices(TopologyProvider->GetNumCorners(), PersistentSelection.SelectedCornerIDs);
+	}
+	else if (Properties->bSelectEdges)
+	{
+		SelectAllIndices(TopologyProvider->GetNumEdges(), PersistentSelection.SelectedEdgeIDs);
+	}
+	else if (Properties->bSelectFaces)
+	{
+		for (int GroupIndex = 0; GroupIndex < TopologyProvider->GetNumGroups(); ++GroupIndex)
+		{
+			const int GroupID = TopologyProvider->GetGroupIDAt(GroupIndex);
+			PersistentSelection.SelectedGroupIDs.Add(GroupID);
+		}
+	}
+
+	SelectionTimestamp++;
+	OnSelectionChanged.Broadcast();
+	
+	if (PreviousSelection != PersistentSelection)
+	{
+		SelectionTimestamp++;
+		OnSelectionChanged.Broadcast();
+	}
+	
+	EndChangeAndEmitIfModified();
+	ParentTool->GetToolManager()->EndUndoTransaction();
+}
+
+FInputRayHit UMeshTopologySelectionMechanic::IsHitByClick(const FInputDeviceRay& ClickPos)
+{
+	if (!bIsEnabled)
+	{
+		return FInputRayHit(); // bHit is false
+	}
+
+	FHitResult OutHit;
+	FGroupTopologySelection Selection;
+	if (TopologyHitTest(ClickPos.WorldRay, OutHit, Selection, CameraState.bIsOrthographic))
+	{
+		return FInputRayHit(OutHit.Distance);
+	}
+
+	// Return a hit so we always capture and can clear the selection
+	return FInputRayHit(TNumericLimits<float>::Max());
+}
+
+void UMeshTopologySelectionMechanic::OnClicked(const FInputDeviceRay& ClickPos)
+{
+	// update selection
+	ParentTool->GetToolManager()->BeginUndoTransaction(LOCTEXT("SelectionChange", "Selection"));
+	BeginChange();
+
+	// This will fire off a OnSelectionChanged delegate.
+	UpdateSelection(ClickPos.WorldRay, LastClickedHitPosition, LastClickedHitNormal);
+
+	EndChangeAndEmitIfModified();
+	ParentTool->GetToolManager()->EndUndoTransaction();
+}
+
+FInputRayHit UMeshTopologySelectionMechanic::BeginHoverSequenceHitTest(const FInputDeviceRay& PressPos)
+{
+	FHitResult OutHit;
+	if (bIsEnabled && TopologyHitTest(PressPos.WorldRay, OutHit))
+	{
+		return FInputRayHit(OutHit.Distance);
+	}
+	return FInputRayHit(); // bHit is false
+}
+
+void UMeshTopologySelectionMechanic::OnBeginHover(const FInputDeviceRay& DevicePos)
+{
+}
+
+bool UMeshTopologySelectionMechanic::OnUpdateHover(const FInputDeviceRay& DevicePos)
+{
+	UpdateHighlight(DevicePos.WorldRay);
+	return true;
+}
+
+void UMeshTopologySelectionMechanic::OnEndHover()
+{
+	ClearHighlight();
+}
+
+void UMeshTopologySelectionMechanic::OnUpdateModifierState(int ModifierID, bool bIsOn)
+{
+	switch (ModifierID)
+	{
+	case ShiftModifierID:
+		bShiftToggle = bIsOn;
+		break;
+	case CtrlModifierID:
+		bCtrlToggle = bIsOn;
+		break;
+	default:
+		break;
+	}
+}
+
+void UMeshTopologySelectionMechanic::OnDragRectangleStarted()
+{
+	bCurrentlyMarqueeDragging = true;
+
+	ParentTool->GetToolManager()->BeginUndoTransaction(LOCTEXT("SelectionChange", "Selection"));
+	BeginChange();
+
+	PreDragPersistentSelection = PersistentSelection;
+	LastUpdateRectangleSelection = PersistentSelection;
+	PreDragTopoSelectorSettings = GetTopoSelectorSettings(false);
+	PreDragTopoSelectorSettings.bIgnoreOcclusion = Properties->bMarqueeIgnoreOcclusion; // uses a separate setting for marquee
+}
+
+void UMeshTopologySelectionMechanic::OnDragRectangleChanged(const FCameraRectangle& CurrentRectangle)
+{
+	FGroupTopologySelection RectangleSelection;
+
+	TopoSelector->FindSelectedElement(PreDragTopoSelectorSettings, CurrentRectangle, TargetTransform,
+		RectangleSelection, &TriIsOccludedCache);
+
+	if (ShouldAddToSelectionFunc())
+	{
+		PersistentSelection = PreDragPersistentSelection;
+		if (ShouldRemoveFromSelectionFunc())
+		{
+			PersistentSelection.Toggle(RectangleSelection);
+		}
+		else
+		{
+			PersistentSelection.Append(RectangleSelection);
+		}
+	}
+	else if (ShouldRemoveFromSelectionFunc())
+	{
+		PersistentSelection = PreDragPersistentSelection;
+		PersistentSelection.Remove(RectangleSelection);
+	}
+	else
+	{
+		// Neither key pressed.
+		PersistentSelection = RectangleSelection;
+	}
+
+	// If we modified the currently selected edges/vertices, they will be properly displayed in our
+	// Render() call. However, the mechanic is not responsible for face highlighting, so if we modified
+	// that, we need to notify the user so that they can update the highlighting (since OnSelectionChanged
+	// only gets broadcast at rectangle end).
+	if ((!PersistentSelection.SelectedGroupIDs.IsEmpty() || !LastUpdateRectangleSelection.SelectedGroupIDs.IsEmpty()) // if groups are involved
+		&& PersistentSelection != LastUpdateRectangleSelection)
+	{
+		LastUpdateRectangleSelection = PersistentSelection;
+		OnFaceSelectionPreviewChanged.Broadcast();
+	}
+}
+
+void UMeshTopologySelectionMechanic::OnDragRectangleFinished(const FCameraRectangle& Rectangle, bool bCancelled)
+{
+	bCurrentlyMarqueeDragging = false;
+
+	TriIsOccludedCache.Reset();
+
+	if (PersistentSelection != PreDragPersistentSelection)
+	{
+		SelectionTimestamp++;
+		OnSelectionChanged.Broadcast();
+	}
+
+	EndChangeAndEmitIfModified();
+	ParentTool->GetToolManager()->EndUndoTransaction();
+}
+
+void UMeshTopologySelectionMechanic::UpdateMarqueeEnabled()
+{
+	MarqueeMechanic->SetIsEnabled(
+		bIsEnabled 
+		&& Properties->bEnableMarquee
+		&& (Properties->bSelectVertices || Properties->bSelectEdges || Properties->bSelectFaces));
+}
+
+
+void UMeshTopologySelectionMechanic::BeginChange()
+{
+	// If you hit this ensure, either you didn't match a BeginChange with an EndChange(), or you
+	// interleaved two actions that both needed a BeginChange(). For instance, are you sure you
+	// are not performing actions while a marquee rectangle is active?
+	ensure(ActiveChange.IsValid() == false);
+	ActiveChange = MakeUnique<FMeshTopologySelectionMechanicSelectionChange>();
+	ActiveChange->Before = PersistentSelection;
+	ActiveChange->Timestamp = SelectionTimestamp;
+}
+
+TUniquePtr<FToolCommandChange> UMeshTopologySelectionMechanic::EndChange()
+{
+	if (ensure(ActiveChange.IsValid()) && SelectionTimestamp != ActiveChange->Timestamp)
+	{
+		ActiveChange->After = PersistentSelection;
+		return MoveTemp(ActiveChange);
+	}
+	ActiveChange = TUniquePtr<FMeshTopologySelectionMechanicSelectionChange>();
+	return TUniquePtr<FToolCommandChange>();
+}
+
+bool UMeshTopologySelectionMechanic::EndChangeAndEmitIfModified()
+{
+	if (ensure(ActiveChange.IsValid()) && SelectionTimestamp != ActiveChange->Timestamp)
+	{
+		ActiveChange->After = PersistentSelection;
+		GetParentTool()->GetToolManager()->EmitObjectChange(this, MoveTemp(ActiveChange),
+			LOCTEXT("SelectionChangeMessage", "Selection Change"));
+		return true;
+	}
+	ActiveChange = TUniquePtr<FMeshTopologySelectionMechanicSelectionChange>();
+	return false;
+}
+
+void UMeshTopologySelectionMechanic::GetClickedHitPosition(FVector3d& PositionOut, FVector3d& NormalOut) const
+{
+	PositionOut = LastClickedHitPosition;
+	NormalOut = LastClickedHitNormal;
+}
+
+FFrame3d UMeshTopologySelectionMechanic::GetSelectionFrame(bool bWorld, FFrame3d* InitialLocalFrame) const
+{
+	FFrame3d UseFrame;
+	if (PersistentSelection.IsEmpty() == false)
+	{
+		UseFrame = TopoSelector->GetTopologyProvider()->GetSelectionFrame(PersistentSelection, InitialLocalFrame);
+	}
+
+	if (bWorld)
+	{
+		UseFrame.Transform(TargetTransform);
+	}
+
+	return UseFrame;
+}
+
+FAxisAlignedBox3d UMeshTopologySelectionMechanic::GetSelectionBounds(bool bWorld) const
+{
+	if ( ! PersistentSelection.IsEmpty() )
+	{
+		const FTopologyProvider* TopologyProvider = TopoSelector->GetTopologyProvider();
+
+		if (bWorld)
+		{
+			return TopologyProvider->GetSelectionBounds(PersistentSelection, [this](const FVector3d& Pos) { return TargetTransform.TransformPosition(Pos); });
+		}
+		else
+		{
+			return TopologyProvider->GetSelectionBounds(PersistentSelection, [this](const FVector3d& Pos) { return Pos; });
+		}
+	}
+	else 
+	{
+		if (bWorld)
+		{
+			return FAxisAlignedBox3d(Mesh->GetBounds(), TargetTransform);
+		}
+		else
+		{
+			return Mesh->GetBounds();
+		}
+	}
+}
+
+void UMeshTopologySelectionMechanic::SetShowSelectableCorners(bool bShowCorners)
+{
+	bShowSelectableCorners = bShowCorners;
+}
+
+
+
+void FMeshTopologySelectionMechanicSelectionChange::Apply(UObject* Object)
+{
+	UMeshTopologySelectionMechanic* Mechanic = Cast<UMeshTopologySelectionMechanic>(Object);
+	if (Mechanic)
+	{
+		Mechanic->PersistentSelection = After;
+		Mechanic->OnSelectionChanged.Broadcast();
+	}
+}
+void FMeshTopologySelectionMechanicSelectionChange::Revert(UObject* Object)
+{
+	UMeshTopologySelectionMechanic* Mechanic = Cast<UMeshTopologySelectionMechanic>(Object);
+	if (Mechanic)
+	{
+		Mechanic->PersistentSelection = Before;
+		Mechanic->OnSelectionChanged.Broadcast();
+	}
+}
+FString FMeshTopologySelectionMechanicSelectionChange::ToString() const
+{
+	return TEXT("FMeshTopologySelectionMechanicSelectionChange");
+}
+
+void UMeshTopologySelectionMechanic::Initialize(const FDynamicMesh3* MeshIn,
+	FTransform3d TargetTransformIn,
+	UWorld* WorldIn,
+	TFunction<FDynamicMeshAABBTree3* ()> GetSpatialSourceFuncIn)
+{
+	this->Mesh = MeshIn;
+	this->TargetTransform = TargetTransformIn;
+
+	this->GetSpatialFunc = GetSpatialSourceFuncIn;
+	TopoSelector->SetSpatialSource(GetSpatialFunc);
+	TopoSelector->PointsWithinToleranceTest = [this](const FVector3d& Position1, const FVector3d& Position2, double TolScale) {
+		if (CameraState.bIsOrthographic)
+		{
+			// We could just always use ToolSceneQueriesUtil::PointSnapQuery. But in ortho viewports, we happen to know
+			// that the only points that we will ever give this function will be the closest points between a ray and
+			// some geometry, meaning that the vector between them will be orthogonal to the view ray. With this knowledge,
+			// we can do the tolerance computation more efficiently than PointSnapQuery can, since we don't need to project
+			// down to the view plane.
+			// As in PointSnapQuery, we convert our angle-based tolerance to one we can use in an ortho viewport (instead of
+			// dividing our field of view into 90 visual angle degrees, we divide the plane into 90 units).
+			float OrthoTolerance = ToolSceneQueriesUtil::GetDefaultVisualAngleSnapThreshD() * CameraState.OrthoWorldCoordinateWidth / 90.0;
+			OrthoTolerance *= TolScale;
+			return DistanceSquared(TargetTransform.TransformPosition(Position1), TargetTransform.TransformPosition(Position2)) < OrthoTolerance * OrthoTolerance;
+		}
+		else
+		{
+			return ToolSceneQueriesUtil::PointSnapQuery(CameraState,
+				TargetTransform.TransformPosition(Position1), TargetTransform.TransformPosition(Position2),
+				ToolSceneQueriesUtil::GetDefaultVisualAngleSnapThreshD() * TolScale);
+		}
+	};
+
+	// Set up the component we use to draw highlighted triangles. Only needs to be done once, not when the mesh
+	// changes (we are assuming that we won't swap worlds without creating a new mechanic).
+	if (PreviewGeometryActor == nullptr)
+	{
+		FRotator Rotation(0.0f, 0.0f, 0.0f);
+		FActorSpawnParameters SpawnInfo;;
+		PreviewGeometryActor = WorldIn->SpawnActor<APreviewGeometryActor>(FVector::ZeroVector, Rotation, SpawnInfo);
+
+		DrawnTriangleSetComponent = NewObject<UTriangleSetComponent>(PreviewGeometryActor);
+		PreviewGeometryActor->SetRootComponent(DrawnTriangleSetComponent);
+		DrawnTriangleSetComponent->RegisterComponent();
+	}
+
+	PreviewGeometryActor->SetActorTransform((FTransform)TargetTransformIn);
+
+	DrawnTriangleSetComponent->Clear();
+	CurrentlyHighlightedGroups.Empty();
+}
+
+#undef LOCTEXT_NAMESPACE
diff --git a/Engine/Plugins/Runtime/MeshModelingToolset/Source/ModelingComponents/Private/Selection/MeshTopologySelector.cpp b/Engine/Plugins/Runtime/MeshModelingToolset/Source/ModelingComponents/Private/Selection/MeshTopologySelector.cpp
new file mode 100644
--- /dev/null
+++ b/Engine/Plugins/Runtime/MeshModelingToolset/Source/ModelingComponents/Private/Selection/MeshTopologySelector.cpp
@@ -0,0 +1,790 @@
+// Copyright Epic Games, Inc. All Rights Reserved.
+
+
+#include "Selection/MeshTopologySelector.h"
+#include "DynamicMesh/DynamicMesh3.h"
+#include "Spatial/SpatialInterfaces.h"
+#include "Engine/Polys.h"
+#include "Mechanics/RectangleMarqueeMechanic.h" // FCameraRectangle
+#include "ConvexVolume.h"
+
+using namespace UE::Geometry;
+
+#define LOCTEXT_NAMESPACE "FMeshTopologySelector"
+
+// Local utility function forward declarations
+bool IsOccluded(const FGeometrySet3::FNearest& ClosestElement, const FVector3d& ViewOrigin, const FDynamicMeshAABBTree3* Spatial, bool bBackFacesOcclude);
+bool IsOccluded(const FVector3d& Point, const FVector3d& ViewOrigin, const FDynamicMeshAABBTree3* Spatial, bool bBackFacesOcclude);
+
+FMeshTopologySelector::FMeshTopologySelector()
+{
+	// initialize to sane values
+	PointsWithinToleranceTest =
+		[](const FVector3d& A, const FVector3d& B, double TolScale) { return Distance(A, B) < (TolScale*1.0); };
+	GetSpatial =
+		[]() { return nullptr; };
+}
+
+
+void FMeshTopologySelector::Invalidate(bool bTopologyDeformed, bool bTopologyModified)
+{
+	if (bTopologyDeformed)
+	{
+		bGeometryUpToDate = false;
+	}
+	if (bTopologyModified)
+	{
+		bGeometryUpToDate = bGeometryInitialized = false;
+	}
+}
+
+
+const FGeometrySet3& FMeshTopologySelector::GetGeometrySet()
+{
+	if (bGeometryInitialized == false)
+	{
+		GeometrySet.Reset();
+		int NumCorners = TopologyProvider->GetNumCorners();
+		for (int CornerID = 0; CornerID < NumCorners; ++CornerID)
+		{
+			FVector3d Position = Mesh->GetVertex(TopologyProvider->GetCornerVertexID(CornerID));
+			GeometrySet.AddPoint(CornerID, Position);
+		}
+		int NumEdges = TopologyProvider->GetNumEdges();
+		for (int EdgeID = 0; EdgeID < NumEdges; ++EdgeID)
+		{
+			FPolyline3d Polyline;
+			TopologyProvider->GetEdgePolyline(EdgeID, Polyline);
+			GeometrySet.AddCurve(EdgeID, Polyline);
+		}
+
+		bGeometryInitialized = true;
+		bGeometryUpToDate = true;
+	}
+
+	if (bGeometryUpToDate == false)
+	{
+		int NumCorners = TopologyProvider->GetNumCorners();
+		for (int CornerID = 0; CornerID < NumCorners; ++CornerID)
+		{
+			FVector3d Position = Mesh->GetVertex(TopologyProvider->GetCornerVertexID(CornerID));
+			GeometrySet.UpdatePoint(CornerID, Position);
+		}
+		int NumEdges = TopologyProvider->GetNumEdges();
+		for (int EdgeID = 0; EdgeID < NumEdges; ++EdgeID)
+		{
+			FPolyline3d Polyline;
+			TopologyProvider->GetEdgePolyline(EdgeID, Polyline);
+			GeometrySet.UpdateCurve(EdgeID, Polyline);
+		}
+		bGeometryUpToDate = true;
+	}
+
+	return GeometrySet;
+}
+
+bool FMeshTopologySelector::FindSelectedElement(const FSelectionSettings& Settings, const FRay3d& Ray,
+	FGroupTopologySelection& ResultOut, FVector3d& SelectedPositionOut, FVector3d& SelectedNormalOut, int32* EdgeSegmentIdOut)
+{
+	// These get used for finding intersections with triangles and corners/edges, respectively.
+	FDynamicMeshAABBTree3* Spatial = GetSpatial();
+	const FGeometrySet3& TopoSpatial = GetGeometrySet();
+
+	IMeshSpatial::FQueryOptions SpatialQueryOptions;
+	if (Spatial && !Settings.bHitBackFaces)
+	{
+		SpatialQueryOptions.TriangleFilterF = [Spatial, &Ray](int32 Tid) {
+			return Spatial->GetMesh()->GetTriNormal(Tid).Dot(Ray.Direction) < 0;
+		};
+	}
+
+	// We start by intersecting with the mesh triangles because even when selecting corners or edges, we set
+	// the normal based on the true triangle that we hit. If we end up with a simple face selection, we will
+	// end up using this result.
+	double RayParameter = -1;
+	int HitTriangleID = IndexConstants::InvalidID;
+	FVector3d TriangleHitPos;
+	bool bActuallyHitSurface = (Spatial != nullptr) ? Spatial->FindNearestHitTriangle(Ray, RayParameter, HitTriangleID, SpatialQueryOptions) : false;
+	if (bActuallyHitSurface)
+	{
+		TriangleHitPos = Ray.PointAt(RayParameter);
+		SelectedNormalOut = Mesh->GetTriNormal(HitTriangleID);
+	}
+	else
+	{
+		SelectedNormalOut = FVector3d::UnitZ();
+	}
+	bool bHaveFaceHit = (bActuallyHitSurface && Settings.bEnableFaceHits);
+	
+	// Deal with corner hits first (and edges that project to a corner)
+	FGroupTopologySelection CornerResults;
+	FVector3d CornerPosition;
+	int32 CornerSegmentEdgeID = 0;
+	bool bHaveCornerHit = false;
+	if (Settings.bEnableCornerHits || (Settings.bEnableEdgeHits && Settings.bPreferProjectedElement))
+	{
+		if (DoCornerBasedSelection(Settings, Ray, Spatial, TopoSpatial, CornerResults, CornerPosition, &CornerSegmentEdgeID))
+		{
+			bHaveCornerHit = true;
+		}
+	}
+
+	// If corner selection didn't yield results, try edge selection
+	FGroupTopologySelection EdgeResults;
+	FVector3d EdgePosition;
+	int32 EdgeSegmentEdgeID = 0;
+	bool bHaveEdgeHit = false;
+	if (Settings.bEnableEdgeHits || (Settings.bEnableFaceHits && Settings.bPreferProjectedElement))
+	{
+		if (DoEdgeBasedSelection(Settings, Ray, Spatial, TopoSpatial, EdgeResults, EdgePosition, &EdgeSegmentEdgeID))
+		{
+			bHaveEdgeHit = true;
+		}
+	}
+
+	// if we have both corner and edge hit, pick the corner if it is within this (relative) tolerance to the hit point
+	constexpr double CornerProximityTolerance = 0.75;
+
+	// if we have both corner and edge, want to keep the one we are closer to
+	if (bHaveCornerHit && bHaveEdgeHit)
+	{
+		if (PointsWithinToleranceTest(CornerPosition, Ray.ClosestPoint(CornerPosition), CornerProximityTolerance))
+		{
+			bHaveEdgeHit = false;
+		}
+		else
+		{
+			bHaveCornerHit = false;
+		}
+	}
+
+	// if we have a corner or edge hit and a face hit, pick corner/edge if it is within this (relative) tolerance to the hit point
+	constexpr double EdgeOrCornerProximityTolerance = 0.15;
+
+	// if we have a corner or edge hit, and a face hit, pick face unless we are really close to corner/edge
+	if ((bHaveCornerHit || bHaveEdgeHit) && bHaveFaceHit)
+	{
+		FVector3d TestPos = (bHaveCornerHit) ? CornerPosition : EdgePosition;
+		if (!PointsWithinToleranceTest(TestPos, Ray.ClosestPoint(TestPos), EdgeOrCornerProximityTolerance))
+		{
+			bHaveEdgeHit = bHaveCornerHit = false;
+		}
+	}
+
+
+	if (bHaveCornerHit)
+	{
+		ResultOut = CornerResults;
+		SelectedPositionOut = CornerPosition;
+		if (EdgeSegmentIdOut != nullptr)
+		{
+			*EdgeSegmentIdOut = CornerSegmentEdgeID;
+		}
+		return true;
+	}
+	if (bHaveEdgeHit)
+	{
+		ResultOut = EdgeResults;
+		SelectedPositionOut = EdgePosition;
+		if (EdgeSegmentIdOut != nullptr)
+		{
+			*EdgeSegmentIdOut = EdgeSegmentEdgeID;
+		}
+		return true;
+	}
+	if (bHaveFaceHit)
+	{
+		// If we still haven't found a selection, go ahead and select the face that we found earlier
+		ResultOut.SelectedGroupIDs.Add(TopologyProvider->GetGroupIDForTriangle(HitTriangleID));
+		SelectedPositionOut = TriangleHitPos;
+		return true;
+	}
+
+
+	return false;
+}
+
+bool FMeshTopologySelector::DoCornerBasedSelection(const FSelectionSettings& Settings,
+	const FRay3d& Ray, FDynamicMeshAABBTree3* Spatial, const FGeometrySet3& TopoSpatial,
+	FGroupTopologySelection& ResultOut, FVector3d& SelectedPositionOut, int32 *EdgeSegmentIdOut) const
+{
+	// These will store our results, depending on whether we select all along the ray or not.
+	FGeometrySet3::FNearest SingleElement;
+	// ElementsWithinTolerance gives all nearby elements returned by the topology selector, whereas DownRayElements filters these to just
+	// the closest element and those that project directly onto it (and it only stores the IDs, since we don't care about anything else).
+	TArray<FGeometrySet3::FNearest> ElementsWithinTolerance;
+	TArray<int32> DownRayElements;
+
+	auto LocalTolTest = [this](const FVector3d& A, const FVector3d& B) { return PointsWithinToleranceTest(A, B, 1.0); };
+
+	// Start by getting the closest element
+	const FGeometrySet3::FNearest* ClosestElement = nullptr;
+	if (!Settings.bSelectDownRay)
+	{
+		if (TopoSpatial.FindNearestPointToRay(Ray, SingleElement, LocalTolTest))
+		{
+			ClosestElement = &SingleElement;
+		}
+	}
+	else
+	{
+		// We're collecting all corners within tolerance, but we still need the closest element
+		if (TopoSpatial.CollectPointsNearRay(Ray, ElementsWithinTolerance, LocalTolTest))
+		{
+			double MinRayT = TNumericLimits<double>::Max();
+			for (const FGeometrySet3::FNearest& Element : ElementsWithinTolerance)
+			{
+				if (Element.RayParam < MinRayT)
+				{
+					MinRayT = Element.RayParam;
+					ClosestElement = &Element;
+				}
+			}
+		}//end if found corners
+	}//end if selecting down ray
+
+	// Bail early if we haven't found a corner
+	if (ClosestElement == nullptr)
+	{
+		return false;
+	}
+
+	// Also bail if the closest element is not visible.
+	if (!Settings.bIgnoreOcclusion && IsOccluded(*ClosestElement, Ray.Origin, Spatial, Settings.bHitBackFaces))
+	{
+		return false;
+	}
+
+	// The closest point is already found
+	SelectedPositionOut = ClosestElement->NearestGeoPoint;
+
+	// If we have other corners, we actually need to filter them to only those that lie in line with the closest element. Note that
+	// this would be done differently depending on whether we're dealing with an orthographic or perspective view: in an
+	// orthographic view, we need to see that they lie on a ray with view ray direction and closest corner origin, whereas in
+	// perspective, they would need to lie on a ray from camer through closest corner (which will differ due to tolerance).
+	// Because the "select down ray" behavior is only useful in orthographic viewports in the first place, we do it that way.
+	if (Settings.bSelectDownRay)
+	{
+		DownRayElements.Add(ClosestElement->ID);
+		for (const FGeometrySet3::FNearest& Element : ElementsWithinTolerance)
+		{
+			if (ClosestElement == &Element)
+			{
+				continue; // Already added
+			}
+
+			// Make sure that closest corner to current element is parallel with view ray
+			FVector3d ClosestTowardElement = UE::Geometry::Normalized(Element.NearestGeoPoint - ClosestElement->NearestGeoPoint);
+			// There would usually be one more abs() in here, but we know that other elements are down ray direction
+			if (abs(ClosestTowardElement.Dot(Ray.Direction) - 1.0) < KINDA_SMALL_NUMBER)
+			{
+				DownRayElements.Add(Element.ID);
+			}
+		}//end assembling aligned corners
+	}
+
+	// Try to select edges that project to corners.
+	if (Settings.bPreferProjectedElement && Settings.bEnableEdgeHits)
+	{
+		TSet<int32> AddedTopologyEdges;
+
+		// See if the closest vertex has an attached edge that is colinear with the view ray. Due to the fact that
+		// topology "edges" are actually polylines, we could actually have more than one even for the closest corner
+		// (if the "edge" towards us curves away). We'll only grab an edge down our view ray.
+		int32 ClosestVid = TopologyProvider->GetCornerVertexID(ClosestElement->ID);
+		for (int32 Eid : Mesh->VtxEdgesItr(ClosestVid))
+		{
+			FDynamicMesh3::FEdge Edge = Mesh->GetEdge(Eid);
+			int32 OtherVid = (Edge.Vert.A == ClosestVid) ? Edge.Vert.B : Edge.Vert.A;
+			FVector3d EdgeVector = UE::Geometry::Normalized(Mesh->GetVertex(OtherVid) - Mesh->GetVertex(ClosestVid));
+			if (abs(EdgeVector.Dot(Ray.Direction) - 1.0) < KINDA_SMALL_NUMBER)
+			{
+				int TopologyEdgeIndex = TopologyProvider->FindGroupEdgeID(Eid);
+				if (TopologyEdgeIndex >= 0)
+				{
+					ResultOut.SelectedEdgeIDs.Add(TopologyEdgeIndex);
+					AddedTopologyEdges.Add(TopologyEdgeIndex);
+
+					if (EdgeSegmentIdOut)
+					{
+						TopologyProvider->GetGroupEdgeEdges(TopologyEdgeIndex).Find(Eid, *EdgeSegmentIdOut);
+					}
+
+					break;
+				}
+			}
+		}
+
+		// If relevant, get all the other colinear edges
+		if (Settings.bSelectDownRay && AddedTopologyEdges.Num() > 0)
+		{
+			for (int i = 1; i < DownRayElements.Num(); ++i) // skip 0 because it is closest and done
+			{
+				// Look though any attached edges.
+				for (int32 Eid : Mesh->VtxEdgesItr(TopologyProvider->GetCornerVertexID(DownRayElements[i])))
+				{
+					FDynamicMesh3::FEdge Edge = Mesh->GetEdge(Eid);
+					FVector3d EdgeVector = UE::Geometry::Normalized(Mesh->GetVertex(Edge.Vert.A) - Mesh->GetVertex(Edge.Vert.B));
+
+					// Compare absolute value of dot product to 1. We already made sure that one of the vertices is in line
+					// with the closest vertex earlier.
+					if (abs(abs(EdgeVector.Dot(Ray.Direction)) - 1.0) < KINDA_SMALL_NUMBER)
+					{
+						int TopologyEdgeIndex = TopologyProvider->FindGroupEdgeID(Eid);
+						if (TopologyEdgeIndex >= 0 && !AddedTopologyEdges.Contains(TopologyEdgeIndex))
+						{
+							ResultOut.SelectedEdgeIDs.Add(TopologyEdgeIndex);
+							AddedTopologyEdges.Add(TopologyEdgeIndex);
+							// Don't break here because we may have parallel edges in both directions, since we aren't
+							// going through the vertices in a particular order.
+						}
+					}
+				}//end checking edges
+			}//end going through corners
+		}
+
+		// If we found edges, we're done
+		if (AddedTopologyEdges.Num() > 0)
+		{
+			return true;
+		}
+	}//end selecting projected edges
+
+	// If getting projected edges didn't work out, go ahead and add the corners.
+	if (Settings.bEnableCornerHits)
+	{
+		if (Settings.bSelectDownRay)
+		{
+			for (int32 Id : DownRayElements)
+			{
+				ResultOut.SelectedCornerIDs.Add(Id);
+			}
+		}
+		else
+		{
+			ResultOut.SelectedCornerIDs.Add(ClosestElement->ID);
+		}
+		return true;
+	}
+
+	return false;
+}
+
+bool FMeshTopologySelector::DoEdgeBasedSelection(const FSelectionSettings& Settings, const FRay3d& Ray,
+	FDynamicMeshAABBTree3* Spatial, const FGeometrySet3& TopoSpatial,
+	FGroupTopologySelection& ResultOut, FVector3d& SelectedPositionOut, int32* EdgeSegmentIdOut) const
+{
+	// These will store our results, depending on whether we select all along the ray or not.
+	FGeometrySet3::FNearest SingleElement;
+	// ElementsWithinTolerance gives all nearby elements returned by the topology selector, whereas DownRayElements filters these to just
+	// the closest element and those that project directly onto it (it stores the polyline and segment IDs, which are all we care about).
+	TArray<FGeometrySet3::FNearest> ElementsWithinTolerance;
+	TArray<FIndex2i> DownRayElements;
+
+	auto LocalTolTest = [this](const FVector3d& A, const FVector3d& B) { return PointsWithinToleranceTest(A, B, 1.0); };
+
+	// Start by getting the closest element
+	const FGeometrySet3::FNearest* ClosestElement = nullptr;
+	if (!Settings.bSelectDownRay)
+	{
+		if (TopoSpatial.FindNearestCurveToRay(Ray, SingleElement, LocalTolTest))
+		{
+			ClosestElement = &SingleElement;
+		}
+	}
+	else
+	{
+		// Need all curves within tolerance, but also need to know the closest.
+		if (TopoSpatial.CollectCurvesNearRay(Ray, ElementsWithinTolerance, LocalTolTest))
+		{
+			double MinRayT = TNumericLimits<double>::Max();
+			for (const FGeometrySet3::FNearest& Element : ElementsWithinTolerance)
+			{
+				if (Element.RayParam < MinRayT)
+				{
+					MinRayT = Element.RayParam;
+					ClosestElement = &Element;
+				}
+			}
+		}//end if found edges
+	}//end if selecting down ray
+
+	// Bail early if we haven't found at least one edge
+	if (ClosestElement == nullptr)
+	{
+		return false;
+	}
+
+	// Also bail if the closest element is not visible.
+	if (!Settings.bIgnoreOcclusion && IsOccluded(*ClosestElement, Ray.Origin, Spatial, Settings.bHitBackFaces))
+	{
+		return false;
+	}
+
+	// The closest point is already found
+	SelectedPositionOut = ClosestElement->NearestGeoPoint;
+
+	// If we have other edges, we need to filter them to only those that project onto the closest element. This would be done
+	// differently for perspective cameras vs orthographic projection, but since the behavior is only useful in ortho mode,
+	// we do it that way.
+	if (Settings.bSelectDownRay)
+	{
+		// Closest element is a given
+		DownRayElements.Add(FIndex2i(ClosestElement->ID, ClosestElement->PolySegmentIdx));
+
+		// We want edges that lie in a plane through the closest edge that is coplanar with the view direction. If we had wanted 
+		// to do this for perspective mode, we would have picked a plane through the edge and the camera origin.
+		int32 ClosestEid = TopologyProvider->GetGroupEdgeEdges(ClosestElement->ID)[ClosestElement->PolySegmentIdx];
+		FDynamicMesh3::FEdge ClosestEdge = Mesh->GetEdge(ClosestEid);
+		UE::Geometry::FPlane3d PlaneThroughClosestEdge(Mesh->GetVertex(ClosestEdge.Vert.A), Mesh->GetVertex(ClosestEdge.Vert.B), 
+			Mesh->GetVertex(ClosestEdge.Vert.A) + Ray.Direction);
+
+		for (const FGeometrySet3::FNearest& Element : ElementsWithinTolerance)
+		{
+			if (ClosestElement == &Element)
+			{
+				continue; // already added
+			}
+
+			// See if the edge endpoints lie in the plane
+			int32 Eid = TopologyProvider->GetGroupEdgeEdges(Element.ID)[Element.PolySegmentIdx];
+			FDynamicMesh3::FEdge Edge = Mesh->GetEdge(Eid);
+			if (abs(PlaneThroughClosestEdge.DistanceTo(Mesh->GetVertex(Edge.Vert.A))) < KINDA_SMALL_NUMBER
+				&& abs(PlaneThroughClosestEdge.DistanceTo(Mesh->GetVertex(Edge.Vert.B))) < KINDA_SMALL_NUMBER)
+			{
+				DownRayElements.Add(FIndex2i(Element.ID, Element.PolySegmentIdx));
+			}
+		}
+	}
+
+	// Try to select faces that project to the closest edge
+	if (Settings.bPreferProjectedElement && Settings.bEnableFaceHits)
+	{
+		TSet<int32> AddedGroups;
+
+		// Start with the closest edge
+		int32 Eid = TopologyProvider->GetGroupEdgeEdges(ClosestElement->ID)[ClosestElement->PolySegmentIdx];
+		FDynamicMesh3::FEdge Edge = Mesh->GetEdge(Eid);
+		FVector3d VertA = Mesh->GetVertex(Edge.Vert.A);
+		FVector3d VertB = Mesh->GetVertex(Edge.Vert.B);
+
+
+		// Grab a plane through the two verts that contains the ray direction (again, this is assuming that we'd
+		// only do this in ortho mode, otherwise the plane would go through the ray origin.
+		UE::Geometry::FPlane3d PlaneThroughClosestEdge(VertA, VertB, VertA + Ray.Direction);
+
+		// Checking that the face is coplanar simply entails checking that the opposite vert is in the plane. However,
+		// it is possible even for the closest edge to have multiple coplanar faces if a group going towards
+		// the camera curves away. We need to be able to select just one, and we want to select the one that
+		// extends down the view ray. For that, we'll make sure that a vector to the opposite vertex lies on the same
+		// side of the edge as the view ray.
+		FVector3d EdgeVector = VertB - VertA;
+		FVector3d EdgeVecCrossDirection = EdgeVector.Cross(Ray.Direction);
+		FIndex2i OppositeVids = Mesh->GetEdgeOpposingV(Eid);
+
+		FVector3d OppositeVert = Mesh->GetVertex(OppositeVids.A);
+		if (abs(PlaneThroughClosestEdge.DistanceTo(OppositeVert)) < KINDA_SMALL_NUMBER
+			&& EdgeVector.Cross(OppositeVert - VertA).Dot(EdgeVecCrossDirection) > 0)
+		{
+			int GroupId = TopologyProvider->GetGroupIDForTriangle(Edge.Tri.A);
+			ResultOut.SelectedGroupIDs.Add(GroupId);
+			AddedGroups.Add(GroupId);
+		}
+		else if (OppositeVids.B != FDynamicMesh3::InvalidID)
+		{
+			OppositeVert = Mesh->GetVertex(OppositeVids.B);
+
+			if (abs(PlaneThroughClosestEdge.DistanceTo(OppositeVert)) < KINDA_SMALL_NUMBER
+				&& EdgeVector.Cross(OppositeVert - VertA).Dot(EdgeVecCrossDirection) > 0)
+			{
+				int GroupId = TopologyProvider->GetGroupIDForTriangle(Edge.Tri.B);
+				ResultOut.SelectedGroupIDs.Add(GroupId);
+				AddedGroups.Add(GroupId);
+			}
+		}
+
+		// If relevant, get all the other coplanar faces
+		if (Settings.bSelectDownRay && AddedGroups.Num() > 0)
+		{
+			for (int i = 1; i < DownRayElements.Num(); ++i) // skip 0 because it is closest and done
+			{
+				// We already made sure that all these edges are coplanar, so we'll just be checking opposite verts.
+				Eid = TopologyProvider->GetGroupEdgeEdges(DownRayElements[i].A)[DownRayElements[i].B];
+				Edge = Mesh->GetEdge(Eid);
+				OppositeVids = Mesh->GetEdgeOpposingV(Eid);
+
+				// No need to check directionality of faces here, just need them to be coplanar.
+				if (abs(PlaneThroughClosestEdge.DistanceTo(Mesh->GetVertex(OppositeVids.A))) < KINDA_SMALL_NUMBER)
+				{
+					int GroupId = TopologyProvider->GetGroupIDForTriangle(Edge.Tri.A);
+					if (!AddedGroups.Contains(GroupId))
+					{
+						ResultOut.SelectedGroupIDs.Add(GroupId);
+						AddedGroups.Add(GroupId);
+					}
+				}
+				if (OppositeVids.B != FDynamicMesh3::InvalidID
+					&& abs(PlaneThroughClosestEdge.DistanceTo(Mesh->GetVertex(OppositeVids.B))) < KINDA_SMALL_NUMBER)
+				{
+					int GroupId = TopologyProvider->GetGroupIDForTriangle(Edge.Tri.B);
+					if (!AddedGroups.Contains(GroupId))
+					{
+						ResultOut.SelectedGroupIDs.Add(GroupId);
+						AddedGroups.Add(GroupId);
+					}
+				}
+			}//end going through other edges
+		}
+
+		// If we selected faces, we're done
+		if (AddedGroups.Num() > 0)
+		{
+			return true;
+		}
+	}
+
+	// If we didn't end up selecting projected faces, and we have edges to select, select them
+	if (Settings.bEnableEdgeHits)
+	{
+		if (Settings.bSelectDownRay)
+		{
+			for (const FIndex2i& ElementTuple : DownRayElements)
+			{
+				ResultOut.SelectedEdgeIDs.Add(ElementTuple.A);
+			}
+		}
+		else
+		{
+			ResultOut.SelectedEdgeIDs.Add(ClosestElement->ID);
+		}
+
+		if (EdgeSegmentIdOut)
+		{
+			*EdgeSegmentIdOut = ClosestElement->PolySegmentIdx;
+		}
+
+		return true;
+	}
+	return false;
+}
+
+bool IsOccluded(const FGeometrySet3::FNearest& ClosestElement, const FVector3d& ViewOrigin, 
+	const FDynamicMeshAABBTree3* Spatial, bool bBackFacesOcclude)
+{
+	return IsOccluded(ClosestElement.NearestGeoPoint, ViewOrigin, Spatial, bBackFacesOcclude);
+}
+
+bool IsOccluded(const FVector3d& Point, const FVector3d& ViewOrigin, 
+	const FDynamicMeshAABBTree3* Spatial, bool bBackFacesOcclude)
+{
+	// We will shoot a ray backwards to see if we hit something. 
+	FRay3d ToEyeRay(Point, UE::Geometry::Normalized(ViewOrigin - Point), true);
+	ToEyeRay.Origin += (double)(100 * FMathf::ZeroTolerance) * ToEyeRay.Direction;
+
+	IMeshSpatial::FQueryOptions QueryOptions;
+	if (!bBackFacesOcclude)
+	{
+		// If back faces don't occlude (ie we weren't hitting back faces on the way forward), 
+		// then the direction toward the camera should be same as normal direction.
+		QueryOptions.TriangleFilterF = [Spatial, &ToEyeRay](int32 Tid)
+		{
+			return Spatial->GetMesh()->GetTriNormal(Tid).Dot(ToEyeRay.Direction) > 0;
+		};
+	}
+
+	if (Spatial->FindNearestHitTriangle(ToEyeRay, QueryOptions) >= 0)
+	{
+		return true;
+	}
+	return false;
+}
+
+bool FMeshTopologySelector::FindSelectedElement(const FSelectionSettings& Settings, 
+	const FCameraRectangle& CameraRectangle, FTransform3d TargetTransform, FGroupTopologySelection& ResultOut,
+	TMap<int32, bool>* TriIsOccludedCache)
+{
+	// One minor easy thing we can do to speed up the below is to detect cases where transform does not have a non-uniform
+	// scale, and in those cases we transform the rectangle once rather than transforming all the query points/curves. 
+	// This isn't worth doing for non-uniform scale transforms because the rectangle basis stops being orthonormal and 
+	// becomes a problem to work with.
+
+	ResultOut.Clear();
+	FDynamicMeshAABBTree3* Spatial = GetSpatial();
+
+	// Needed for occlusion test, which happens in local space.
+	FVector3d LocalCameraOrigin = TargetTransform.InverseTransformPosition((FVector3d)CameraRectangle.CameraState.Position);
+	
+	// Corner selection takes priority over edges.
+	if (Settings.bEnableCornerHits)
+	{
+		GeometrySet.ParallelFindAllPointsSatisfying(
+			[&CameraRectangle, &TargetTransform, &Settings, &LocalCameraOrigin, Spatial]
+			(const FVector3d& PointPosition)
+			{
+				FVector TransformedPoint = TargetTransform.TransformPosition(PointPosition);
+				return CameraRectangle.IsProjectedPointInRectangle(TransformedPoint)
+					&& (Settings.bIgnoreOcclusion || !IsOccluded(PointPosition, LocalCameraOrigin, 
+						Spatial, Settings.bHitBackFaces));
+			},
+			ResultOut.SelectedCornerIDs);
+	}
+
+	// If we didn't get corners, look for edges.
+	if (ResultOut.SelectedCornerIDs.IsEmpty() && Settings.bEnableEdgeHits)
+	{
+		GeometrySet.ParallelFindAllCurvesSatisfying(
+			[&CameraRectangle, &TargetTransform, &Settings, &LocalCameraOrigin, Spatial](const FPolyline3d& Curve)
+			{
+				// Testing occlusion properly seems like it would be a pain, so we'll just consider something occluded
+				// if one of the endpoints is occluded. This will handle the common case of not wanting to select a
+				// hidden edge that is connected to a visible corner.
+				if (!Settings.bIgnoreOcclusion 
+					&& (IsOccluded(Curve.Start(), LocalCameraOrigin, Spatial, Settings.bHitBackFaces)
+						|| IsOccluded(Curve.End(), LocalCameraOrigin, Spatial, Settings.bHitBackFaces)))
+				{
+					return false;
+				}
+
+				// Check whether any of the component segments intersect the rectangle
+				const TArray<FVector3d>& Verts = Curve.GetVertices();
+				FVector3d CurrentVert = TargetTransform.TransformPosition(Verts[0]);
+				for (int32 i = 1; i < Verts.Num(); ++i)
+				{
+					FVector3d NextVert = TargetTransform.TransformPosition(Verts[i]);
+					if (CameraRectangle.IsProjectedSegmentIntersectingRectangle((FVector)CurrentVert, (FVector)NextVert))
+					{
+						return true;
+					}
+					CurrentVert = NextVert;
+				}
+				return false;
+			},
+			ResultOut.SelectedEdgeIDs);
+	}
+	
+	if (Settings.bEnableFaceHits && ResultOut.IsEmpty())
+	{
+		// Get a frustum volume in local space, then use it to figure out if it contains or intersects
+		// the aabb tree boxes in the traversal. 
+		// TODO: Frustum transformation is done the same way in FractureEditorMode.cpp- should this be placed
+		// somewhere common? Not sure where to put it though.
+		FConvexVolume WorldSpaceFrustum = CameraRectangle.FrustumAsConvexVolume();
+		FMatrix InverseTargetTransform(FTransform(TargetTransform).ToInverseMatrixWithScale());
+
+		FConvexVolume LocalFrustum;
+		LocalFrustum.Planes.Empty(6);
+		for (const FPlane& Plane : WorldSpaceFrustum.Planes)
+		{
+			LocalFrustum.Planes.Add(Plane.TransformBy(InverseTargetTransform));
+		}
+		LocalFrustum.Init();
+
+		FDynamicMeshAABBTree3::FTreeTraversal Traversal;
+		// Used as additional state for the traversal, relying on the fact that our traversal is depth-first
+		// and serial. When we descend to a box that is fully contained, we use CurrentSelectAllDepth to mark
+		// all deeper boxes as contained, and reset it once we return back out of that subtree (done similarly
+		// in UVEditorMeshSelectionMechanic).
+		int CurrentDepth = -1;
+		int CurrentSelectAllDepth = TNumericLimits<int>::Max();
+
+		Traversal.NextBoxF = [&LocalFrustum, &CurrentDepth, &CurrentSelectAllDepth](const FAxisAlignedBox3d& Box, int Depth)
+		{
+			CurrentDepth = Depth;
+			if (Depth > CurrentSelectAllDepth)
+			{
+				// We're currently inside a fully-selected box
+				return true;
+			}
+			// Reset CurrentSelectAllDepth because we must have left the cotained subtree (if it was set to begin with)
+			CurrentSelectAllDepth = TNumericLimits<int>::Max();
+
+			bool bFullyContained = false;
+			bool bIntersects = LocalFrustum.IntersectBox(Box.Center(), Box.Extents(), bFullyContained);
+			
+			if (bFullyContained)
+			{
+				CurrentSelectAllDepth = Depth;
+			}
+			return bIntersects;
+		};
+
+		Traversal.NextTriangleF = 
+			[this, &ResultOut, &CurrentDepth, &CurrentSelectAllDepth, &LocalFrustum, 
+			&LocalCameraOrigin, &Spatial, &Settings, TriIsOccludedCache](int Tid)
+		{
+			int32 TriGroupID = TopologyProvider->GetGroupIDForTriangle(Tid);
+
+			// If this group already got selected, no need to do anything else. Helps avoid some occlusion tests.
+			if (ResultOut.SelectedGroupIDs.Contains(TriGroupID))
+			{
+				return;
+			}
+
+			// Apply frustum check. Skip if box was fully in frustum
+			if (!(CurrentDepth >= CurrentSelectAllDepth))
+			{
+				// TODO: This isn't an ideal way to check whether the triangle is inside the frustum because FPoly is
+				// an old class that we probably want to avoid using. It's fine enough here, but it should be replaced.
+				// Related note: It may be tempting to just check for containment of the centroid, but we probably don't
+				// want to do that because it makes selection behavior dependent on the tesselation of the underlying
+				// group. I.e. a remeshed cube face ends up being much more reliable to select than one made of two triangles 
+				// where a marquee select of even the group center could fail if it's between the tri centroids.
+				FIndex3i TriVids = Mesh->GetTriangle(Tid);
+				FPoly TrianglePolygon;
+				TrianglePolygon.Init();
+				for (int i = 0; i < 3; ++i)
+				{
+					TrianglePolygon.Vertices.Add((FVector3f)Mesh->GetVertex(TriVids[i]));
+				}
+
+				if (!LocalFrustum.ClipPolygon(TrianglePolygon))
+				{
+					// Triangle was fully outside the frustum
+					return;
+				}
+			}
+
+			// Apply occlusion filter. Note that the back face filter is applied via TriangleFilterF in the traversal.
+			if (!Settings.bIgnoreOcclusion)
+			{
+				bool bTriangleIsOccluded = false;
+
+				// See if we've already calculated occlusion for this tri.
+				bool* CachedOcclusion = TriIsOccludedCache->Find(Tid);
+				if (CachedOcclusion)
+				{
+					bTriangleIsOccluded = *CachedOcclusion;
+				}
+				else
+				{
+					bTriangleIsOccluded = IsOccluded(Mesh->GetTriCentroid(Tid), LocalCameraOrigin, Spatial, Settings.bHitBackFaces);
+					TriIsOccludedCache->Add(Tid, bTriangleIsOccluded);
+				}
+
+				if (bTriangleIsOccluded)
+				{
+					return;
+				}
+			}
+
+			ResultOut.SelectedGroupIDs.Add(TriGroupID);
+		};
+
+		// Add the back face filter, if applicable
+		IMeshSpatial::FQueryOptions Options;
+		if (!Settings.bHitBackFaces)
+		{
+			Options.TriangleFilterF = [this, &LocalCameraOrigin](int32 Tid)
+			{
+				return Mesh->GetTriNormal(Tid).Dot(Mesh->GetVertex(Mesh->GetTriangle(Tid)[0]) - LocalCameraOrigin) < 0;
+			};
+		}
+		
+		Spatial->DoTraversal(Traversal, Options);
+	}
+
+	return ResultOut.IsEmpty();
+}
+
+
+#undef LOCTEXT_NAMESPACE
diff --git a/Engine/Plugins/Runtime/MeshModelingToolset/Source/ModelingComponents/Private/Selection/PolygonSelectionMechanic.cpp b/Engine/Plugins/Runtime/MeshModelingToolset/Source/ModelingComponents/Private/Selection/PolygonSelectionMechanic.cpp
--- a/Engine/Plugins/Runtime/MeshModelingToolset/Source/ModelingComponents/Private/Selection/PolygonSelectionMechanic.cpp
+++ b/Engine/Plugins/Runtime/MeshModelingToolset/Source/ModelingComponents/Private/Selection/PolygonSelectionMechanic.cpp
@@ -1,187 +1,31 @@
 // Copyright Epic Games, Inc. All Rights Reserved.
 
 #include "Selection/PolygonSelectionMechanic.h"
-
-#include "BaseBehaviors/SingleClickOrDragBehavior.h"
-#include "BaseBehaviors/MouseHoverBehavior.h"
-#include "InteractiveToolManager.h"
-#include "Util/ColorConstants.h"
+#include "Selection/GroupTopologySelector.h"
 #include "Selection/PersistentMeshSelection.h"
-#include "ToolSceneQueriesUtil.h"
-#include "ToolSetupUtil.h"
 
 using namespace UE::Geometry;
 
 #define LOCTEXT_NAMESPACE "UPolygonSelectionMechanic"
 
-void UPolygonSelectionMechanicProperties::InvertSelection()
-{
-	if (Mechanic.IsValid())
-	{
-		Mechanic->InvertSelection();
-	}
-}
-
-void UPolygonSelectionMechanicProperties::SelectAll()
-{
-	if (Mechanic.IsValid())
-	{
-		Mechanic->SelectAll();
-	}
-}
-
-UPolygonSelectionMechanic::~UPolygonSelectionMechanic()
-{
-	checkf(PreviewGeometryActor == nullptr, TEXT("Shutdown() should be called before UPolygonSelectionMechanic is destroyed."));
-}
-
-void UPolygonSelectionMechanic::Setup(UInteractiveTool* ParentToolIn)
-{
-	UInteractionMechanic::Setup(ParentToolIn);
-
-	TopoSelector = MakeShared<FGroupTopologySelector, ESPMode::ThreadSafe>();
-
-	HoverBehavior = NewObject<UMouseHoverBehavior>();
-	// We use modifiers on hover to change the highlighting according to what can be selected  	
-	HoverBehavior->Modifiers.RegisterModifier(ShiftModifierID, FInputDeviceState::IsShiftKeyDown);
-	HoverBehavior->Modifiers.RegisterModifier(CtrlModifierID, FInputDeviceState::IsCtrlKeyDown);
-	HoverBehavior->Initialize(this);
-	HoverBehavior->SetDefaultPriority(BasePriority);
-	ParentToolIn->AddInputBehavior(HoverBehavior, this);
-
-	MarqueeMechanic = NewObject<URectangleMarqueeMechanic>();
-	MarqueeMechanic->bUseExternalClickDragBehavior = true;
-	MarqueeMechanic->Setup(ParentToolIn);
-	MarqueeMechanic->OnDragRectangleStarted.AddUObject(this, &UPolygonSelectionMechanic::OnDragRectangleStarted);
-	MarqueeMechanic->OnDragRectangleChanged.AddUObject(this, &UPolygonSelectionMechanic::OnDragRectangleChanged);
-	MarqueeMechanic->OnDragRectangleFinished.AddUObject(this, &UPolygonSelectionMechanic::OnDragRectangleFinished);
-	MarqueeMechanic->SetBasePriority(BasePriority.MakeLower());
-
-	ClickOrDragBehavior = NewObject<USingleClickOrDragInputBehavior>();
-	ClickOrDragBehavior->Initialize(this, MarqueeMechanic);
-	ClickOrDragBehavior->Modifiers.RegisterModifier(ShiftModifierID, FInputDeviceState::IsShiftKeyDown);
-	ClickOrDragBehavior->Modifiers.RegisterModifier(CtrlModifierID, FInputDeviceState::IsCtrlKeyDown);
-	ClickOrDragBehavior->SetDefaultPriority(BasePriority);
-	ParentTool->AddInputBehavior(ClickOrDragBehavior, this);
-
-	Properties = NewObject<UPolygonSelectionMechanicProperties>(this);
-	Properties->Initialize(this);
-	if (bAddSelectionFilterPropertiesToParentTool)
-	{
-		AddToolPropertySource(Properties);
-	}
-	Properties->WatchProperty(Properties->bSelectVertices, [this](bool bSelectVertices) {
-		UpdateMarqueeEnabled(); 
-	});
-	Properties->WatchProperty(Properties->bSelectEdges, [this](bool bSelectVertices) { 
-		UpdateMarqueeEnabled(); 
-	});
-	Properties->WatchProperty(Properties->bSelectFaces, [this](bool bSelectFaces) {
-		UpdateMarqueeEnabled();
-	});
-	Properties->WatchProperty(Properties->bEnableMarquee, [this](bool bEnableMarquee) { 
-		UpdateMarqueeEnabled(); 
-	});
-
-	// set up visualizers
-	PolyEdgesRenderer.LineColor = FLinearColor::Red;
-	PolyEdgesRenderer.LineThickness = 2.0;
-	PolyEdgesRenderer.PointColor = FLinearColor::Red;
-	PolyEdgesRenderer.PointSize = 8.0f;
-	HilightRenderer.LineColor = FLinearColor::Green;
-	HilightRenderer.LineThickness = 4.0f;
-	HilightRenderer.PointColor = FLinearColor::Green;
-	HilightRenderer.PointSize = 10.0f;
-	SelectionRenderer.LineColor = LinearColors::Gold3f();
-	SelectionRenderer.LineThickness = 4.0f;
-	SelectionRenderer.PointColor = LinearColors::Gold3f();
-	SelectionRenderer.PointSize = 10.0f;
-
-	float HighlightedFacePercentDepthOffset = 0.5f;
-	HighlightedFaceMaterial = ToolSetupUtil::GetSelectionMaterial(FLinearColor::Green, ParentToolIn->GetToolManager(), HighlightedFacePercentDepthOffset);
-	// The rest of the highlighting setup has to be done in Initialize(), since we need the world to set up our drawing component.
-}
-
-void UPolygonSelectionMechanic::Shutdown()
-{
-	if (PreviewGeometryActor)
-	{
-		PreviewGeometryActor->Destroy();
-		PreviewGeometryActor = nullptr;
-	}
-}
-
-void UPolygonSelectionMechanic::GetSelection(UPersistentMeshSelection& SelectionOut, const FCompactMaps* CompactMapsToApply) const
-{
-	SelectionOut.SetSelection(*Topology, PersistentSelection, CompactMapsToApply);
-}
-
-void UPolygonSelectionMechanic::LoadSelection(const UPersistentMeshSelection& SelectionIn)
-{
-	SelectionIn.ExtractIntoSelectionObject(*Topology, PersistentSelection);
-}
-
 void UPolygonSelectionMechanic::Initialize(
 	const FDynamicMesh3* MeshIn,
 	FTransform3d TargetTransformIn,
 	UWorld* WorldIn,
 	const FGroupTopology* TopologyIn,
-	TFunction<FDynamicMeshAABBTree3 * ()> GetSpatialSourceFuncIn)
+	TFunction<FDynamicMeshAABBTree3* ()> GetSpatialSourceFuncIn)
 {
-	this->Mesh = MeshIn;
-	this->Topology = TopologyIn;
-	this->TargetTransform = TargetTransformIn;
-
-	TopoSelector->Initialize(Mesh, Topology);
-	this->GetSpatialFunc = GetSpatialSourceFuncIn;
-	TopoSelector->SetSpatialSource(GetSpatialFunc);
-	TopoSelector->PointsWithinToleranceTest = [this](const FVector3d& Position1, const FVector3d& Position2, double TolScale) {
-		if (CameraState.bIsOrthographic)
-		{
-			// We could just always use ToolSceneQueriesUtil::PointSnapQuery. But in ortho viewports, we happen to know
-			// that the only points that we will ever give this function will be the closest points between a ray and
-			// some geometry, meaning that the vector between them will be orthogonal to the view ray. With this knowledge,
-			// we can do the tolerance computation more efficiently than PointSnapQuery can, since we don't need to project
-			// down to the view plane.
-			// As in PointSnapQuery, we convert our angle-based tolerance to one we can use in an ortho viewport (instead of
-			// dividing our field of view into 90 visual angle degrees, we divide the plane into 90 units).
-			float OrthoTolerance = ToolSceneQueriesUtil::GetDefaultVisualAngleSnapThreshD() * CameraState.OrthoWorldCoordinateWidth / 90.0;
-			OrthoTolerance *= TolScale;
-			return DistanceSquared( TargetTransform.TransformPosition(Position1), TargetTransform.TransformPosition(Position2) ) < OrthoTolerance * OrthoTolerance;
-		}
-		else
-		{
-			return ToolSceneQueriesUtil::PointSnapQuery(CameraState,
-				TargetTransform.TransformPosition(Position1), TargetTransform.TransformPosition(Position2),
-				ToolSceneQueriesUtil::GetDefaultVisualAngleSnapThreshD() * TolScale);
-		}
-	};
+	Topology = TopologyIn;
+	TopoSelector = MakeShared<FGroupTopologySelector>(MeshIn, TopologyIn);
 
-	// Set up the component we use to draw highlighted triangles. Only needs to be done once, not when the mesh
-	// changes (we are assuming that we won't swap worlds without creating a new mechanic).
-	if (PreviewGeometryActor == nullptr)
-	{
-		FRotator Rotation(0.0f, 0.0f, 0.0f);
-		FActorSpawnParameters SpawnInfo;;
-		PreviewGeometryActor = WorldIn->SpawnActor<APreviewGeometryActor>(FVector::ZeroVector, Rotation, SpawnInfo);
-
-		DrawnTriangleSetComponent = NewObject<UTriangleSetComponent>(PreviewGeometryActor);
-		PreviewGeometryActor->SetRootComponent(DrawnTriangleSetComponent);
-		DrawnTriangleSetComponent->RegisterComponent();
-	}
-
-	PreviewGeometryActor->SetActorTransform((FTransform)TargetTransformIn);
-
-	DrawnTriangleSetComponent->Clear();
-	CurrentlyHighlightedGroups.Empty();
+	UMeshTopologySelectionMechanic::Initialize(MeshIn, TargetTransformIn, WorldIn, GetSpatialSourceFuncIn);
 }
 
 void UPolygonSelectionMechanic::Initialize(
 	UDynamicMeshComponent* MeshComponentIn,
 	const FGroupTopology* TopologyIn,
-	TFunction<FDynamicMeshAABBTree3 * ()> GetSpatialSourceFuncIn)
+	TFunction<FDynamicMeshAABBTree3* ()> GetSpatialSourceFuncIn)
 {
-
 	Initialize(MeshComponentIn->GetMesh(),
 		(FTransform3d)MeshComponentIn->GetComponentTransform(),
 		MeshComponentIn->GetWorld(),
@@ -189,217 +33,43 @@ void UPolygonSelectionMechanic::Initialize(
 		GetSpatialSourceFuncIn);
 }
 
-void UPolygonSelectionMechanic::DisableBehaviors(UInteractiveTool* ParentToolIn)
-{
-	ParentToolIn->RemoveInputBehaviorsBySource(this);
-	ParentToolIn->RemoveInputBehaviorsBySource(MarqueeMechanic);
-
-	// TODO: Is it worth adding a way to remove the property watchers for marquee?
-}
-
-void UPolygonSelectionMechanic::SetIsEnabled(bool bOn)
-{
-	bIsEnabled = bOn;
-	UpdateMarqueeEnabled();
-}
-
-void UPolygonSelectionMechanic::SetBasePriority(const FInputCapturePriority &Priority)
-{
-	BasePriority = Priority;
-	if (ClickOrDragBehavior)
-	{
-		ClickOrDragBehavior->SetDefaultPriority(Priority);
-	}
-	if (HoverBehavior)
-	{
-		HoverBehavior->SetDefaultPriority(Priority);
-	}
-	if (MarqueeMechanic)
-	{
-		MarqueeMechanic->SetBasePriority(Priority.MakeLower());
-	}
-}
-
-TPair<FInputCapturePriority, FInputCapturePriority> UPolygonSelectionMechanic::GetPriorityRange() const
-{
-	TPair<FInputCapturePriority, FInputCapturePriority> Result;
-	Result.Key = BasePriority;
-	Result.Value = MarqueeMechanic->GetPriorityRange().Value;
-	return Result;
-}
-
-void UPolygonSelectionMechanic::Render(IToolsContextRenderAPI* RenderAPI)
-{
-	MarqueeMechanic->Render(RenderAPI);
-
-	// Cache the view camera state so we can use for snapping/etc.
-	// This should not happen in Render() though...
-	GetParentTool()->GetToolManager()->GetContextQueriesAPI()->GetCurrentViewState(CameraState);
-
-	FViewCameraState RenderCameraState = RenderAPI->GetCameraState();
-
-	const FDynamicMesh3* TargetMesh = this->Mesh;
-	FTransform Transform = (FTransform)TargetTransform;
-
-	PolyEdgesRenderer.BeginFrame(RenderAPI, RenderCameraState);
-	PolyEdgesRenderer.SetTransform(Transform);
-	for (const FGroupTopology::FGroupEdge& Edge : Topology->Edges)
-	{
-		FVector3d A, B;
-		for (int32 eid : Edge.Span.Edges)
-		{
-			TargetMesh->GetEdgeV(eid, A, B);
-			PolyEdgesRenderer.DrawLine(A, B);
-		}
-	}
-	if (bShowSelectableCorners)
-	{
-		for (const FGroupTopology::FCorner& Corner : Topology->Corners)
-		{
-			FVector3d A = TargetMesh->GetVertex(Corner.VertexID);
-			PolyEdgesRenderer.DrawPoint(A);
-		}
-	}
-	PolyEdgesRenderer.EndFrame();
-
-	if (PersistentSelection.IsEmpty() == false)
-	{
-		SelectionRenderer.BeginFrame(RenderAPI, RenderCameraState);
-		SelectionRenderer.SetTransform(Transform);
-		TopoSelector->DrawSelection(PersistentSelection, &SelectionRenderer, &RenderCameraState);
-		SelectionRenderer.EndFrame();
-	}
-
-	HilightRenderer.BeginFrame(RenderAPI, RenderCameraState);
-	HilightRenderer.SetTransform(Transform);
-	TopoSelector->DrawSelection(HilightSelection, &HilightRenderer, &RenderCameraState, FGroupTopologySelector::ECornerDrawStyle::Circle);
-	HilightRenderer.EndFrame();
-}
-
-void UPolygonSelectionMechanic::DrawHUD(FCanvas* Canvas, IToolsContextRenderAPI* RenderAPI)
-{
-	MarqueeMechanic->DrawHUD(Canvas, RenderAPI);
-}
-
-
-
-
-void UPolygonSelectionMechanic::ClearHighlight()
+// TODO: Remove this function when Properties is removed after deprecation
+void UPolygonSelectionMechanic::Setup(UInteractiveTool* ParentToolIn)
 {
-	checkf(DrawnTriangleSetComponent != nullptr, TEXT("Initialize() not called on UPolygonSelectionMechanic."));
+	UMeshTopologySelectionMechanic::Setup(ParentToolIn);
 
-	HilightSelection.Clear();
-	DrawnTriangleSetComponent->Clear();
-	CurrentlyHighlightedGroups.Empty();
-}
-
-
-void UPolygonSelectionMechanic::NotifyMeshChanged(bool bTopologyModified)
-{
-	ClearHighlight();
-	TopoSelector->Invalidate(true, bTopologyModified);
-	if (bTopologyModified)
-	{
-		PersistentSelection.Clear();
-		SelectionTimestamp++;
-		OnSelectionChanged.Broadcast();
-	}
+	PRAGMA_DISABLE_DEPRECATION_WARNINGS
+	Properties_DEPRECATED = NewObject<UDEPRECATED_PolygonSelectionMechanicProperties>(this);
+	Properties_DEPRECATED->Initialize(this);
+	Properties_DEPRECATED->WatchProperty(Properties_DEPRECATED->bSelectVertices, [this](bool bSelectVertices) {
+		UpdateMarqueeEnabled();
+	});
+	Properties_DEPRECATED->WatchProperty(Properties_DEPRECATED->bSelectEdges, [this](bool bSelectVertices) {
+		UpdateMarqueeEnabled();
+	});
+	Properties_DEPRECATED->WatchProperty(Properties_DEPRECATED->bSelectFaces, [this](bool bSelectFaces) {
+		UpdateMarqueeEnabled();
+	});
+	Properties_DEPRECATED->WatchProperty(Properties_DEPRECATED->bEnableMarquee, [this](bool bEnableMarquee) {
+		UpdateMarqueeEnabled();
+	});
+	PRAGMA_ENABLE_DEPRECATION_WARNINGS
 }
 
 
-bool UPolygonSelectionMechanic::TopologyHitTest(const FRay& WorldRay, FHitResult& OutHit, bool bUseOrthoSettings)
-{
-	FGroupTopologySelection Selection;
-	return TopologyHitTest(WorldRay, OutHit, Selection, bUseOrthoSettings);
-}
-
-bool UPolygonSelectionMechanic::TopologyHitTest(const FRay& WorldRay, FHitResult& OutHit, FGroupTopologySelection& OutSelection, bool bUseOrthoSettings)
+void UPolygonSelectionMechanic::GetSelection(UPersistentMeshSelection& SelectionOut, const FCompactMaps* CompactMapsToApply) const
 {
-	FRay3d LocalRay(TargetTransform.InverseTransformPosition((FVector3d)WorldRay.Origin),
-		TargetTransform.InverseTransformVector((FVector3d)WorldRay.Direction));
-	UE::Geometry::Normalize(LocalRay.Direction);
-
-	FVector3d LocalPosition, LocalNormal;
-	int32 EdgeSegmentId; // Only used if hit is an edge
-	FGroupTopologySelector::FSelectionSettings TopoSelectorSettings = GetTopoSelectorSettings(bUseOrthoSettings);
-	if (TopoSelector->FindSelectedElement(TopoSelectorSettings, LocalRay, OutSelection, LocalPosition, LocalNormal, &EdgeSegmentId) == false)
-	{
-		return false;
-	}
-
-	if (OutSelection.SelectedCornerIDs.Num() > 0)
-	{
-		OutHit.FaceIndex = OutSelection.GetASelectedCornerID();
-		OutHit.Distance = LocalRay.GetParameter(LocalPosition);
-		OutHit.ImpactPoint = (FVector)TargetTransform.TransformPosition(LocalRay.PointAt(OutHit.Distance));
-	}
-	else if (OutSelection.SelectedEdgeIDs.Num() > 0)
-	{
-		OutHit.FaceIndex = OutSelection.GetASelectedEdgeID();
-		OutHit.Distance = LocalRay.GetParameter(LocalPosition);
-		OutHit.ImpactPoint = (FVector)TargetTransform.TransformPosition(LocalRay.PointAt(OutHit.Distance));
-		OutHit.Item = EdgeSegmentId;
-	}
-	else
-	{
-		FDynamicMeshAABBTree3* Spatial = GetSpatialFunc();
-		int HitTID = Spatial->FindNearestHitTriangle(LocalRay);
-		if (HitTID != IndexConstants::InvalidID)
-		{
-			FTriangle3d Triangle;
-			Spatial->GetMesh()->GetTriVertices(HitTID, Triangle.V[0], Triangle.V[1], Triangle.V[2]);
-			FIntrRay3Triangle3d Query(LocalRay, Triangle);
-			if (Query.Find())
-			{
-				OutHit.FaceIndex = HitTID;
-				OutHit.Distance = (float)Query.RayParameter;
-				OutHit.Normal = (FVector)TargetTransform.TransformVectorNoScale(Spatial->GetMesh()->GetTriNormal(HitTID));
-				OutHit.ImpactPoint = (FVector)TargetTransform.TransformPosition(LocalRay.PointAt(Query.RayParameter));
-			}
-			else
-			{
-				return false;
-			}
-		}
-	}
-	return true;
+	SelectionOut.SetSelection(*Topology, PersistentSelection, CompactMapsToApply);
 }
 
-
-
-FGroupTopologySelector::FSelectionSettings UPolygonSelectionMechanic::GetTopoSelectorSettings(bool bUseOrthoSettings)
+void UPolygonSelectionMechanic::LoadSelection(const UPersistentMeshSelection& SelectionIn)
 {
-	FGroupTopologySelector::FSelectionSettings Settings;
-
-	Settings.bEnableFaceHits = Properties->bSelectFaces;
-	Settings.bEnableEdgeHits = Properties->bSelectEdges;
-	Settings.bEnableCornerHits = Properties->bSelectVertices;
-	Settings.bHitBackFaces = Properties->bHitBackFaces;
-
-	if (!PersistentSelection.IsEmpty() && (ShouldAddToSelectionFunc() || ShouldRemoveFromSelectionFunc()))
-	{
-		// If we have a selection and we're adding/removing/toggling elements make sure we only hit elements with compatible types
-		Settings.bEnableFaceHits = Settings.bEnableFaceHits && PersistentSelection.SelectedGroupIDs.Num() > 0;
-		Settings.bEnableEdgeHits = Settings.bEnableEdgeHits && PersistentSelection.SelectedEdgeIDs.Num() > 0;
-		Settings.bEnableCornerHits = Settings.bEnableCornerHits && PersistentSelection.SelectedCornerIDs.Num() > 0;
-	}
-
-	if (bUseOrthoSettings)
-	{
-		Settings.bPreferProjectedElement = Properties->bPreferProjectedElement;
-		Settings.bSelectDownRay = Properties->bSelectDownRay;
-		Settings.bIgnoreOcclusion = Properties->bIgnoreOcclusion;
-	}
-
-	return Settings;
+	SelectionIn.ExtractIntoSelectionObject(*Topology, PersistentSelection);
 }
 
-
-
-
 bool UPolygonSelectionMechanic::UpdateHighlight(const FRay& WorldRay)
 {
-	checkf(DrawnTriangleSetComponent != nullptr, TEXT("Initialize() not called on UPolygonSelectionMechanic."));
+	checkf(DrawnTriangleSetComponent != nullptr, TEXT("Initialize() not called on UMeshTopologySelectionMechanic."));
 
 	FRay3d LocalRay(TargetTransform.InverseTransformPosition((FVector3d)WorldRay.Origin),
 		TargetTransform.InverseTransformVector((FVector3d)WorldRay.Direction));
@@ -410,13 +80,15 @@ bool UPolygonSelectionMechanic::UpdateHighlight(const FRay& WorldRay)
 	FGroupTopologySelector::FSelectionSettings TopoSelectorSettings = GetTopoSelectorSettings(CameraState.bIsOrthographic);
 	bool bHit = TopoSelector->FindSelectedElement(TopoSelectorSettings, LocalRay, HilightSelection, LocalPosition, LocalNormal);
 
+	TSharedPtr<FGroupTopologySelector> GroupTopoSelector = StaticCastSharedPtr<FGroupTopologySelector>(TopoSelector);
+
 	if (HilightSelection.SelectedEdgeIDs.Num() > 0 && Properties->bSelectEdgeRings && ShouldSelectEdgeRingsFunc())
 	{
-		TopoSelector->ExpandSelectionByEdgeRings(HilightSelection);
+		GroupTopoSelector->ExpandSelectionByEdgeRings(HilightSelection);
 	}
 	if (HilightSelection.SelectedEdgeIDs.Num() > 0 && Properties->bSelectEdgeLoops && ShouldSelectEdgeLoopsFunc())
 	{
-		TopoSelector->ExpandSelectionByEdgeLoops(HilightSelection);
+		GroupTopoSelector->ExpandSelectionByEdgeLoops(HilightSelection);
 	}
 
 	// Don't hover highlight a selection that we already selected, because people didn't like that
@@ -459,7 +131,7 @@ bool UPolygonSelectionMechanic::UpdateHighlight(const FRay& WorldRay)
 					DrawnTriangleSetComponent->AddTriangle(FRenderableTriangle(HighlightedFaceMaterial,
 						FRenderableTriangleVertex((FVector)Mesh->GetVertex(VertIndices.A), (FVector2D)Mesh->GetVertexUV(VertIndices.A), (FVector)TriangleNormal, FLinearColor(Mesh->GetVertexColor(VertIndices.A)).ToFColor(true)),
 						FRenderableTriangleVertex((FVector)Mesh->GetVertex(VertIndices.B), (FVector2D)Mesh->GetVertexUV(VertIndices.B), (FVector)TriangleNormal, FLinearColor(Mesh->GetVertexColor(VertIndices.B)).ToFColor(true)),
-						FRenderableTriangleVertex((FVector)Mesh->GetVertex(VertIndices.C), (FVector2D)Mesh->GetVertexUV(VertIndices.C), (FVector)TriangleNormal, FLinearColor(Mesh->GetVertexColor(VertIndices.C)).ToFColor(true)) ));
+						FRenderableTriangleVertex((FVector)Mesh->GetVertex(VertIndices.C), (FVector2D)Mesh->GetVertexUV(VertIndices.C), (FVector)TriangleNormal, FLinearColor(Mesh->GetVertexColor(VertIndices.C)).ToFColor(true))));
 				}
 
 				CurrentlyHighlightedGroups.Add(Gid);
@@ -471,14 +143,6 @@ bool UPolygonSelectionMechanic::UpdateHighlight(const FRay& WorldRay)
 }
 
 
-
-
-bool UPolygonSelectionMechanic::HasSelection() const
-{
-	return PersistentSelection.IsEmpty() == false;
-}
-
-
 bool UPolygonSelectionMechanic::UpdateSelection(const FRay& WorldRay, FVector3d& LocalHitPositionOut, FVector3d& LocalHitNormalOut)
 {
 	FRay3d LocalRay(TargetTransform.InverseTransformPosition((FVector3d)WorldRay.Origin),
@@ -495,13 +159,15 @@ bool UPolygonSelectionMechanic::UpdateSelection(const FRay& WorldRay, FVector3d&
 		LocalHitPositionOut = LocalPosition;
 		LocalHitNormalOut = LocalNormal;
 
+		TSharedPtr<FGroupTopologySelector> GroupTopoSelector = StaticCastSharedPtr<FGroupTopologySelector>(TopoSelector);
+
 		if (Selection.SelectedEdgeIDs.Num() > 0 && Properties->bSelectEdgeRings && ShouldSelectEdgeRingsFunc())
 		{
-			TopoSelector->ExpandSelectionByEdgeRings(Selection);
+			GroupTopoSelector->ExpandSelectionByEdgeRings(Selection);
 		}
 		if (Selection.SelectedEdgeIDs.Num() > 0 && Properties->bSelectEdgeLoops && ShouldSelectEdgeLoopsFunc())
 		{
-			TopoSelector->ExpandSelectionByEdgeLoops(Selection);
+			GroupTopoSelector->ExpandSelectionByEdgeLoops(Selection);
 		}
 	}
 
@@ -536,387 +202,4 @@ bool UPolygonSelectionMechanic::UpdateSelection(const FRay& WorldRay, FVector3d&
 }
 
 
-void UPolygonSelectionMechanic::SetSelection(const FGroupTopologySelection& Selection, bool bBroadcast)
-{
-	PersistentSelection = Selection;
-	SelectionTimestamp++;
-	if (bBroadcast)
-	{
-		OnSelectionChanged.Broadcast();
-	}
-}
-
-
-void UPolygonSelectionMechanic::ClearSelection()
-{
-	PersistentSelection.Clear();
-	SelectionTimestamp++;
-	OnSelectionChanged.Broadcast();
-}
-
-void UPolygonSelectionMechanic::InvertSelection()
-{
-	if (PersistentSelection.IsEmpty())
-	{
-		SelectAll();
-		return;
-	}
-
-	ParentTool->GetToolManager()->BeginUndoTransaction(LOCTEXT("SelectionChange", "Selection"));
-	BeginChange();
-
-	const FGroupTopologySelection PreviousSelection = PersistentSelection;
-	PersistentSelection.Clear();
-
-	if (!PreviousSelection.SelectedCornerIDs.IsEmpty())
-	{
-		for (int32 CornerID = 0; CornerID < Topology->Corners.Num(); ++CornerID)
-		{
-			if (!PreviousSelection.SelectedCornerIDs.Contains(CornerID))
-			{
-				PersistentSelection.SelectedCornerIDs.Add(CornerID);
-			}
-		}
-	}
-	else if (!PreviousSelection.SelectedEdgeIDs.IsEmpty())
-	{
-		for (int32 EdgeID = 0; EdgeID < Topology->Edges.Num(); ++EdgeID)
-		{
-			if (!PreviousSelection.SelectedEdgeIDs.Contains(EdgeID))
-			{
-				PersistentSelection.SelectedEdgeIDs.Add(EdgeID);
-			}
-		}
-	}
-	else if (!PreviousSelection.SelectedGroupIDs.IsEmpty())
-	{
-		for (const FGroupTopology::FGroup& Group : Topology->Groups)
-		{
-			if (!PreviousSelection.SelectedGroupIDs.Contains(Group.GroupID))
-			{
-				PersistentSelection.SelectedGroupIDs.Add(Group.GroupID);
-			}
-		}
-	}
-
-	SelectionTimestamp++;
-	OnSelectionChanged.Broadcast();
-	EndChangeAndEmitIfModified();
-	ParentTool->GetToolManager()->EndUndoTransaction();
-}
-
-void UPolygonSelectionMechanic::SelectAll()
-{
-	const FGroupTopologySelection PreviousSelection = PersistentSelection;
-
-	ParentTool->GetToolManager()->BeginUndoTransaction(LOCTEXT("SelectionChange", "Selection"));
-	BeginChange();
-
-	auto SelectAllIndices = [](int32 MaxExclusiveIndex, TSet<int32>& ContainerOut)
-	{
-		for (int32 i = 0; i < MaxExclusiveIndex; ++i)
-		{
-			ContainerOut.Add(i);
-		}
-	};
-
-	PersistentSelection.Clear();
-
-	// Select based on settings, prefering corners to edges to groups (since this is the preference we have
-	// elsewhere, eg in marquee).
-	if (Properties->bSelectVertices)
-	{
-		SelectAllIndices(Topology->Corners.Num(), PersistentSelection.SelectedCornerIDs);
-	}
-	else if (Properties->bSelectEdges)
-	{
-		SelectAllIndices(Topology->Edges.Num(), PersistentSelection.SelectedEdgeIDs);
-	}
-	else if (Properties->bSelectFaces)
-	{
-		for (const FGroupTopology::FGroup& Group : Topology->Groups)
-		{
-			PersistentSelection.SelectedGroupIDs.Add(Group.GroupID);
-		}
-	}
-
-	SelectionTimestamp++;
-	OnSelectionChanged.Broadcast();
-	
-	if (PreviousSelection != PersistentSelection)
-	{
-		SelectionTimestamp++;
-		OnSelectionChanged.Broadcast();
-	}
-	
-	EndChangeAndEmitIfModified();
-	ParentTool->GetToolManager()->EndUndoTransaction();
-}
-
-FInputRayHit UPolygonSelectionMechanic::IsHitByClick(const FInputDeviceRay& ClickPos)
-{
-	if (!bIsEnabled)
-	{
-		return FInputRayHit(); // bHit is false
-	}
-
-	FHitResult OutHit;
-	FGroupTopologySelection Selection;
-	if (TopologyHitTest(ClickPos.WorldRay, OutHit, Selection, CameraState.bIsOrthographic))
-	{
-		return FInputRayHit(OutHit.Distance);
-	}
-
-	// Return a hit so we always capture and can clear the selection
-	return FInputRayHit(TNumericLimits<float>::Max());
-}
-
-void UPolygonSelectionMechanic::OnClicked(const FInputDeviceRay& ClickPos)
-{
-	// update selection
-	ParentTool->GetToolManager()->BeginUndoTransaction(LOCTEXT("SelectionChange", "Selection"));
-	BeginChange();
-
-	// This will fire off a OnSelectionChanged delegate.
-	UpdateSelection(ClickPos.WorldRay, LastClickedHitPosition, LastClickedHitNormal);
-
-	EndChangeAndEmitIfModified();
-	ParentTool->GetToolManager()->EndUndoTransaction();
-}
-
-FInputRayHit UPolygonSelectionMechanic::BeginHoverSequenceHitTest(const FInputDeviceRay& PressPos)
-{
-	FHitResult OutHit;
-	if (bIsEnabled && TopologyHitTest(PressPos.WorldRay, OutHit))
-	{
-		return FInputRayHit(OutHit.Distance);
-	}
-	return FInputRayHit(); // bHit is false
-}
-
-void UPolygonSelectionMechanic::OnBeginHover(const FInputDeviceRay& DevicePos)
-{
-}
-
-bool UPolygonSelectionMechanic::OnUpdateHover(const FInputDeviceRay& DevicePos)
-{
-	UpdateHighlight(DevicePos.WorldRay);
-	return true;
-}
-
-void UPolygonSelectionMechanic::OnEndHover()
-{
-	ClearHighlight();
-}
-
-void UPolygonSelectionMechanic::OnUpdateModifierState(int ModifierID, bool bIsOn)
-{
-	switch (ModifierID)
-	{
-	case ShiftModifierID:
-		bShiftToggle = bIsOn;
-		break;
-	case CtrlModifierID:
-		bCtrlToggle = bIsOn;
-		break;
-	default:
-		break;
-	}
-}
-
-void UPolygonSelectionMechanic::OnDragRectangleStarted()
-{
-	bCurrentlyMarqueeDragging = true;
-
-	ParentTool->GetToolManager()->BeginUndoTransaction(LOCTEXT("SelectionChange", "Selection"));
-	BeginChange();
-
-	PreDragPersistentSelection = PersistentSelection;
-	LastUpdateRectangleSelection = PersistentSelection;
-	PreDragTopoSelectorSettings = GetTopoSelectorSettings(false);
-	PreDragTopoSelectorSettings.bIgnoreOcclusion = Properties->bMarqueeIgnoreOcclusion; // uses a separate setting for marquee
-}
-
-void UPolygonSelectionMechanic::OnDragRectangleChanged(const FCameraRectangle& CurrentRectangle)
-{
-	FGroupTopologySelection RectangleSelection;
-
-	TopoSelector->FindSelectedElement(PreDragTopoSelectorSettings, CurrentRectangle, TargetTransform,
-		RectangleSelection, &TriIsOccludedCache);
-
-	if (ShouldAddToSelectionFunc())
-	{
-		PersistentSelection = PreDragPersistentSelection;
-		if (ShouldRemoveFromSelectionFunc())
-		{
-			PersistentSelection.Toggle(RectangleSelection);
-		}
-		else
-		{
-			PersistentSelection.Append(RectangleSelection);
-		}
-	}
-	else if (ShouldRemoveFromSelectionFunc())
-	{
-		PersistentSelection = PreDragPersistentSelection;
-		PersistentSelection.Remove(RectangleSelection);
-	}
-	else
-	{
-		// Neither key pressed.
-		PersistentSelection = RectangleSelection;
-	}
-
-	// If we modified the currently selected edges/vertices, they will be properly displayed in our
-	// Render() call. However, the mechanic is not responsible for face highlighting, so if we modified
-	// that, we need to notify the user so that they can update the highlighting (since OnSelectionChanged
-	// only gets broadcast at rectangle end).
-	if ((!PersistentSelection.SelectedGroupIDs.IsEmpty() || !LastUpdateRectangleSelection.SelectedGroupIDs.IsEmpty()) // if groups are involved
-		&& PersistentSelection != LastUpdateRectangleSelection)
-	{
-		LastUpdateRectangleSelection = PersistentSelection;
-		OnFaceSelectionPreviewChanged.Broadcast();
-	}
-}
-
-void UPolygonSelectionMechanic::OnDragRectangleFinished(const FCameraRectangle& Rectangle, bool bCancelled)
-{
-	bCurrentlyMarqueeDragging = false;
-
-	TriIsOccludedCache.Reset();
-
-	if (PersistentSelection != PreDragPersistentSelection)
-	{
-		SelectionTimestamp++;
-		OnSelectionChanged.Broadcast();
-	}
-
-	EndChangeAndEmitIfModified();
-	ParentTool->GetToolManager()->EndUndoTransaction();
-}
-
-void UPolygonSelectionMechanic::UpdateMarqueeEnabled()
-{
-	MarqueeMechanic->SetIsEnabled(
-		bIsEnabled 
-		&& Properties->bEnableMarquee
-		&& (Properties->bSelectVertices || Properties->bSelectEdges || Properties->bSelectFaces));
-}
-
-
-void UPolygonSelectionMechanic::BeginChange()
-{
-	// If you hit this ensure, either you didn't match a BeginChange with an EndChange(), or you
-	// interleaved two actions that both needed a BeginChange(). For instance, are you sure you
-	// are not performing actions while a marquee rectangle is active?
-	ensure(ActiveChange.IsValid() == false);
-	ActiveChange = MakeUnique<FPolygonSelectionMechanicSelectionChange>();
-	ActiveChange->Before = PersistentSelection;
-	ActiveChange->Timestamp = SelectionTimestamp;
-}
-
-TUniquePtr<FToolCommandChange> UPolygonSelectionMechanic::EndChange()
-{
-	if (ensure(ActiveChange.IsValid()) && SelectionTimestamp != ActiveChange->Timestamp)
-	{
-		ActiveChange->After = PersistentSelection;
-		return MoveTemp(ActiveChange);
-	}
-	ActiveChange = TUniquePtr<FPolygonSelectionMechanicSelectionChange>();
-	return TUniquePtr<FToolCommandChange>();
-}
-
-bool UPolygonSelectionMechanic::EndChangeAndEmitIfModified()
-{
-	if (ensure(ActiveChange.IsValid()) && SelectionTimestamp != ActiveChange->Timestamp)
-	{
-		ActiveChange->After = PersistentSelection;
-		GetParentTool()->GetToolManager()->EmitObjectChange(this, MoveTemp(ActiveChange),
-			LOCTEXT("SelectionChangeMessage", "Selection Change"));
-		return true;
-	}
-	ActiveChange = TUniquePtr<FPolygonSelectionMechanicSelectionChange>();
-	return false;
-}
-
-void UPolygonSelectionMechanic::GetClickedHitPosition(FVector3d& PositionOut, FVector3d& NormalOut) const
-{
-	PositionOut = LastClickedHitPosition;
-	NormalOut = LastClickedHitNormal;
-}
-
-FFrame3d UPolygonSelectionMechanic::GetSelectionFrame(bool bWorld, FFrame3d* InitialLocalFrame) const
-{
-	FFrame3d UseFrame;
-	if (PersistentSelection.IsEmpty() == false)
-	{
-		UseFrame = Topology->GetSelectionFrame(PersistentSelection, InitialLocalFrame);
-	}
-
-	if (bWorld)
-	{
-		UseFrame.Transform(TargetTransform);
-	}
-
-	return UseFrame;
-}
-
-FAxisAlignedBox3d UPolygonSelectionMechanic::GetSelectionBounds(bool bWorld) const
-{
-	if ( ! PersistentSelection.IsEmpty() )
-	{
-		if (bWorld)
-		{
-			return Topology->GetSelectionBounds(PersistentSelection, [this](const FVector3d& Pos) { return TargetTransform.TransformPosition(Pos); });
-		}
-		else
-		{
-			return Topology->GetSelectionBounds(PersistentSelection, [this](const FVector3d& Pos) { return Pos; });
-		}
-	}
-	else 
-	{
-		if (bWorld)
-		{
-			return FAxisAlignedBox3d(Mesh->GetBounds(), TargetTransform);
-		}
-		else
-		{
-			return Mesh->GetBounds();
-		}
-	}
-}
-
-void UPolygonSelectionMechanic::SetShowSelectableCorners(bool bShowCorners)
-{
-	bShowSelectableCorners = bShowCorners;
-}
-
-
-
-void FPolygonSelectionMechanicSelectionChange::Apply(UObject* Object)
-{
-	UPolygonSelectionMechanic* Mechanic = Cast<UPolygonSelectionMechanic>(Object);
-	if (Mechanic)
-	{
-		Mechanic->PersistentSelection = After;
-		Mechanic->OnSelectionChanged.Broadcast();
-	}
-}
-void FPolygonSelectionMechanicSelectionChange::Revert(UObject* Object)
-{
-	UPolygonSelectionMechanic* Mechanic = Cast<UPolygonSelectionMechanic>(Object);
-	if (Mechanic)
-	{
-		Mechanic->PersistentSelection = Before;
-		Mechanic->OnSelectionChanged.Broadcast();
-	}
-}
-FString FPolygonSelectionMechanicSelectionChange::ToString() const
-{
-	return TEXT("FPolygonSelectionMechanicSelectionChange");
-}
-
-
-
-#undef LOCTEXT_NAMESPACE
\ No newline at end of file
+#undef LOCTEXT_NAMESPACE
diff --git a/Engine/Plugins/Runtime/MeshModelingToolset/Source/ModelingComponents/Public/Selection/BoundarySelectionMechanic.h b/Engine/Plugins/Runtime/MeshModelingToolset/Source/ModelingComponents/Public/Selection/BoundarySelectionMechanic.h
new file mode 100644
--- /dev/null
+++ b/Engine/Plugins/Runtime/MeshModelingToolset/Source/ModelingComponents/Public/Selection/BoundarySelectionMechanic.h
@@ -0,0 +1,29 @@
+// Copyright Epic Games, Inc. All Rights Reserved.
+
+#pragma once
+
+#include "PolygonSelectionMechanic.h"
+#include "BoundarySelectionMechanic.generated.h"
+
+namespace UE::Geometry { class FMeshBoundaryLoops; }
+class FBoundarySelector;
+
+UCLASS()
+class MODELINGCOMPONENTS_API UBoundarySelectionMechanic : public UMeshTopologySelectionMechanic
+{
+	GENERATED_BODY()
+
+public:
+
+	void Initialize(
+		const FDynamicMesh3* MeshIn,
+		FTransform3d TargetTransformIn,
+		UWorld* WorldIn,
+		const UE::Geometry::FMeshBoundaryLoops* BoundaryLoopsIn,
+		TFunction<FDynamicMeshAABBTree3* ()> GetSpatialSourceFuncIn);
+
+	virtual bool UpdateHighlight(const FRay& WorldRay) override;
+
+	virtual bool UpdateSelection(const FRay& WorldRay, FVector3d& LocalHitPositionOut, FVector3d& LocalHitNormalOut) override;
+
+};
diff --git a/Engine/Plugins/Runtime/MeshModelingToolset/Source/ModelingComponents/Public/Selection/BoundarySelector.h b/Engine/Plugins/Runtime/MeshModelingToolset/Source/ModelingComponents/Public/Selection/BoundarySelector.h
new file mode 100644
--- /dev/null
+++ b/Engine/Plugins/Runtime/MeshModelingToolset/Source/ModelingComponents/Public/Selection/BoundarySelector.h
@@ -0,0 +1,127 @@
+// Copyright Epic Games, Inc. All Rights Reserved.
+
+#pragma once
+
+#include "GroupTopologySelector.h"
+#include "MeshBoundaryLoops.h"
+
+using UE::Geometry::FMeshBoundaryLoops;
+using UE::Geometry::FEdgeLoop;
+
+class FBoundaryTopologyProvider : public FTopologyProvider
+{
+public:
+
+	FBoundaryTopologyProvider(const FMeshBoundaryLoops* BoundaryLoops) :
+		BoundaryLoops(BoundaryLoops)
+	{}
+
+	virtual ~FBoundaryTopologyProvider() = default;
+
+	virtual int GetNumEdges() const override
+	{
+		return BoundaryLoops->Loops.Num();
+	}
+
+	virtual void GetEdgePolyline(int EdgeID, FPolyline3d& OutPolyline) const override
+	{
+		OutPolyline.Clear();
+		const TArray<int>& Vertices = BoundaryLoops->Loops[EdgeID].Vertices;
+		for (int i = 0; i < Vertices.Num(); ++i)
+		{
+			OutPolyline.AppendVertex(BoundaryLoops->Mesh->GetVertex(Vertices[i]));
+		}
+	}
+
+	virtual int FindGroupEdgeID(int MeshEdgeID) const override
+	{
+		return BoundaryLoops->FindLoopContainingEdge(MeshEdgeID);
+	}
+
+	virtual const TArray<int>& GetGroupEdgeEdges(int GroupEdgeID) const override
+	{
+		return BoundaryLoops->Loops[GroupEdgeID].Edges;
+	}
+
+	virtual const TArray<int>& GetGroupEdgeVertices(int GroupEdgeID) const override
+	{
+		return BoundaryLoops->Loops[GroupEdgeID].Vertices;
+	}
+
+	virtual void ForGroupSetEdges(const TSet<int32>& GroupIDs, const TFunction<void(int EdgeID)>& EdgeFunc) const override
+	{
+		// We should only have at most one "group" for MeshBoundaryLoops
+		check(GroupIDs.Num() < 2);
+		for (int32 GroupID : GroupIDs)
+		{
+			for (const FEdgeLoop& Loop : BoundaryLoops->Loops)
+			{
+				for (const int32 EdgeID : Loop.Edges)
+				{
+					EdgeFunc(EdgeID);
+				}
+			}
+		}
+	}
+
+	virtual FAxisAlignedBox3d GetSelectionBounds(const FGroupTopologySelection& Selection, TFunctionRef<FVector3d(const FVector3d&)> TransformFunc) const override;
+	virtual FFrame3d GetSelectionFrame(const FGroupTopologySelection& Selection, FFrame3d* InitialLocalFrame) const override;
+
+public:
+
+	//
+	// Functions that don't make sense if we don't have Groups and Corners but are still part of the interface:
+	//
+
+	virtual int GetNumCorners() const override
+	{
+		return 0;
+	}
+
+	virtual int GetCornerVertexID(int CornerID) const override
+	{
+		return IndexConstants::InvalidID;
+	}
+
+	virtual int GetNumGroups() const override
+	{
+		return 1;
+	}
+
+	virtual int GetGroupIDAt(int GroupIndex) const override
+	{
+		return 0;
+	}
+
+	virtual int GetGroupIDForTriangle(int TriangleID) const override
+	{
+		return 0;
+	}
+
+private:
+
+	bool GetLoopTangent(int LoopID, FVector3d& TangentOut) const;
+	FVector3d GetLoopMidpoint(int32 LoopID, double* ArcLengthOut = nullptr, TArray<double>* PerVertexLengthsOut = nullptr) const;
+	double GetLoopArcLength(int32 LoopID, TArray<double>* PerVertexLengthsOut = nullptr) const;
+
+	const FMeshBoundaryLoops* BoundaryLoops;
+
+};
+
+/*
+* This class allows selection of mesh boundary loops. It inherits from FGroupTopologySelector to leverage a lot of the selection functionality, but it uses
+* a FMeshBoundaryLoops object rather than a FGroupTopology object to determine which loops are selectable.
+*/
+class MODELINGCOMPONENTS_API FBoundarySelector : public FMeshTopologySelector
+{
+public:
+
+	/**
+	 * Initialize the selector with the given Mesh and Topology.
+	 * This does not create the internal data structures, this happens lazily on GetGeometrySet()
+	 */
+	FBoundarySelector(const FDynamicMesh3* Mesh, const FMeshBoundaryLoops* BoundaryLoops);
+
+	virtual void DrawSelection(const FGroupTopologySelection& Selection, FToolDataVisualizer* Renderer, const FViewCameraState* CameraState, ECornerDrawStyle CornerDrawStyle = ECornerDrawStyle::Point) override;
+};
+
diff --git a/Engine/Plugins/Runtime/MeshModelingToolset/Source/ModelingComponents/Public/Selection/GroupTopologySelector.h b/Engine/Plugins/Runtime/MeshModelingToolset/Source/ModelingComponents/Public/Selection/GroupTopologySelector.h
--- a/Engine/Plugins/Runtime/MeshModelingToolset/Source/ModelingComponents/Public/Selection/GroupTopologySelector.h
+++ b/Engine/Plugins/Runtime/MeshModelingToolset/Source/ModelingComponents/Public/Selection/GroupTopologySelector.h
@@ -2,149 +2,143 @@
 
 #pragma once
 
+#include "MeshTopologySelector.h"
 #include "GroupTopology.h"
-#include "Spatial/GeometrySet3.h"
-#include "DynamicMesh/DynamicMeshAABBTree3.h"
+#include "UObject/ObjectPtr.h"
 
-struct FCameraRectangle;
-class FToolDataVisualizer;
-struct FViewCameraState;
-
-using UE::Geometry::FGroupTopologySelection;
 using UE::Geometry::FGroupTopology;
-using UE::Geometry::FDynamicMeshAABBTree3;
-using UE::Geometry::FDynamicMesh3;
-using UE::Geometry::FGeometrySet3;
 
 /**
- * FGroupTopologySelector implements selection behavior for a FGroupTopology mesh.
- * Groups, Group Edges, and Corners can be selected dependings on the settings passed in.
- *
- * Internally a FGeometrySet3 is constructed to support ray-hit testing against the edges and corners.
- * 
- * Note that to hit-test against the mesh you have to provide a your own FDynamicMeshAABBTree3.
- * You do this by providing a callback via SetSpatialSource(). The reason for this is that
- * (1) owners of FGroupTopologySelector likely already have a BVTree and (2) if the use case
- * is deformation, we need to make sure the owner has recomputed the BVTree before we call functions 
- * on it. The callback you provide should do that!
- * 
- * DrawSelection() can be used to visualize a selection via line/circle drawing.
- * 
- * @todo optionally have an internal mesh AABBTree that can be used when owner does not provide one?
+ * TopologyProvider using FGroupTopology 
  */
-class MODELINGCOMPONENTS_API FGroupTopologySelector
+
+class FGroupTopologyProvider : public FTopologyProvider
 {
 public:
 
-	//
-	// Configuration variables
-	// 
+	FGroupTopologyProvider(const FGroupTopology* GroupTopology) :
+		GroupTopology(GroupTopology)
+	{}
+	
+	virtual ~FGroupTopologyProvider() = default;
 
-	/**
-	 * Determines the behavior of a FindSelectedElement() call.
-	 */
-	struct FSelectionSettings
+	virtual int GetNumCorners() const override
 	{
-		bool bEnableFaceHits = true;
-		bool bEnableEdgeHits = true;
-		bool bEnableCornerHits = true;
-		
-		// When false, hits of a triangles that are facing away from the camera are not considered. This can
-		// be useful when working with inside-out meshes, where we usually want to select the farther, visible
-		// wall.
-		bool bHitBackFaces = true;
-
-		// The following are mainly useful for ortho viewport selection:
-
-		// Prefer an edge projected to a point rather than the point, and a face projected to an edge
-		// rather than the edge.
-		bool bPreferProjectedElement = false;
-
-		// If the first element is valid, select all elements behind it that are aligned with it.
-		bool bSelectDownRay = false;
-
-		//Do not check whether the closest element is occluded.
-		bool bIgnoreOcclusion = false;
-	};
-
-	/** 
-	 * This is the function we use to determine if a point on a corner/edge is close enough
-	 * to the hit-test ray to treat as a "hit". By default this is Euclidean distance with
-	 * a tolerance of 1.0. You probably need to replace this with your own function.
-	 */
-	TFunction<bool(const FVector3d&, const FVector3d&, double ToleranceScale)> PointsWithinToleranceTest;
+		return GroupTopology->Corners.Num();
+	}
 
-public:
-	FGroupTopologySelector();
+	virtual int GetCornerVertexID(int CornerID) const override
+	{
+		return GroupTopology->Corners[CornerID].VertexID;
+	}
 
-	/**
-	 * Initialize the selector with the given Mesh and Topology.
-	 * This does not create the internal data structures, this happens lazily on GetGeometrySet()
-	 */
-	void Initialize(const FDynamicMesh3* Mesh, const FGroupTopology* Topology);
+	virtual int GetNumEdges() const override
+	{
+		return GroupTopology->Edges.Num();
+	}
 
-	/**
-	 * Provide a function that will return an AABBTree for the Mesh.
-	 * See class comment for why this is necessary.
-	 */
-	void SetSpatialSource(TFunction<FDynamicMeshAABBTree3*(void)> GetSpatialFunc)
+	virtual void GetEdgePolyline(int EdgeID, FPolyline3d& OutPolyline) const override
 	{
-		GetSpatial = GetSpatialFunc;
+		GroupTopology->Edges[EdgeID].Span.GetPolyline(OutPolyline);
 	}
 
-	/**
-	 * Notify the Selector that the mesh has changed. 
-	 * @param bTopologyDeformed if this is true, the mesh vertices have been moved so we need to update bounding boxes/etc
-	 * @param bTopologyModified if this is true, topology has changed and we need to rebuild spatial data structures from scratch
-	 */
-	void Invalidate(bool bTopologyDeformed, bool bTopologyModified);
+	virtual int FindGroupEdgeID(int MeshEdgeID) const override
+	{
+		return GroupTopology->FindGroupEdgeID(MeshEdgeID);
+	}
 
-	/**
-	 * @return the internal GeometrySet. This does lazy updating of the GeometrySet, so this function may take some time.
-	 */
-	const FGeometrySet3& GetGeometrySet();
+	const TArray<int>& GetGroupEdgeEdges(int GroupEdgeID) const override
+	{
+		return GroupTopology->GetGroupEdgeEdges(GroupEdgeID);
+	}
+
+	const TArray<int>& GetGroupEdgeVertices(int GroupEdgeID) const override
+	{
+		return GroupTopology->GetGroupEdgeVertices(GroupEdgeID);
+	}
+
+	virtual int GetGroupIDForTriangle(int TriangleID) const override
+	{
+		return GroupTopology->GetGroupID(TriangleID);
+	}
+
+	virtual void ForGroupSetEdges(const TSet<int32>& GroupIDs, const TFunction<void (int EdgeID)>& EdgeFunc) const
+	{
+		GroupTopology->ForGroupSetEdges(GroupIDs, [EdgeFunc](FGroupTopology::FGroupEdge Edge, int EdgeID)
+		{
+			EdgeFunc(EdgeID);
+		});
+	}
+
+	virtual int GetNumGroups() const override
+	{
+		return GroupTopology->Groups.Num();
+	}
+
+	virtual int GetGroupIDAt(int GroupIndex) const override
+	{
+		return GroupTopology->Groups[GroupIndex].GroupID;
+	}
+
+	virtual FFrame3d GetSelectionFrame(const FGroupTopologySelection& Selection, FFrame3d* InitialLocalFrame = nullptr) const override
+	{
+		return GroupTopology->GetSelectionFrame(Selection, InitialLocalFrame);
+	}
+
+	virtual FAxisAlignedBox3d GetSelectionBounds(const FGroupTopologySelection& Selection, TFunctionRef<FVector3d(const FVector3d&)> TransformFunc) const override
+	{
+		return GroupTopology->GetSelectionBounds(Selection, TransformFunc);
+	}
+
+private:
+
+	const FGroupTopology* GroupTopology;
+
+};
+
+/**
+ * Additional functionality for operating on FGroupTopology
+ */
+struct FGroupTopologyUtils
+{
+	bool GetNextEdgeLoopEdge(int32 IncomingEdgeID, int32 CornerID, int32& NextEdgeIDOut) const;
+	void AddNewEdgeLoopEdgesFromCorner(int32 EdgeID, int32 CornerID, TSet<int32>& EdgeSet) const;
+	void AddNewEdgeRingEdges(int32 StartEdgeID, int32 ForwardGroupID, TSet<int32>& EdgeSet) const;
+	bool GetQuadOppositeEdge(int32 EdgeIDIn, int32 GroupID, int32& OppositeEdgeIDOut) const;
+	FIndex2i GetEdgeEndpointCorners(int EdgeID) const;
+	FIndex2i GetEdgeGroups(int EdgeID) const;
+
+	const FGroupTopology* GroupTopology;
+};
+
+
+/**
+* FGroupTopologySelector is a MeshTopologySelector subclass using an FGroupTopology to define groups, edges, and vertices
+*/
+
+class MODELINGCOMPONENTS_API FGroupTopologySelector : public FMeshTopologySelector
+{
+
+public:
+
+	FGroupTopologySelector() = default;
 
 	/**
-	 * Find which element was selected for a given ray.
-	 *
-	 * @param Settings settings that determine what can be selected.
-	 * @param Ray hit-test ray
-	 * @param ResultOut resulting selection. Will not be cleared before use. At most one of the Groups/Corners/Edges
-	 *  members will be added to.
-	 * @param SelectedPositionOut The point on the ray nearest to the selected element
-	 * @param SelectedNormalOut the normal at that point if ResultOut contains a selected face, otherwise uninitialized
-	 * @param EdgeSegmentIdOut When not null, if the selected element is a group edge, the segment id of the segment
-	 *   that was actually clicked within the edge polyline will be stored here.
-	 *
-	 * @return true if something was selected
+	 * Create the selector with the given Mesh and Topology.
+	 * This does not create the internal data structures, this happens lazily on GetGeometrySet()
 	 */
-	bool FindSelectedElement(const FSelectionSettings& Settings, const FRay3d& Ray, FGroupTopologySelection& ResultOut,
-		FVector3d& SelectedPositionOut, FVector3d& SelectedNormalOut, int32* EdgeSegmentIdOut = nullptr);
+	FGroupTopologySelector(const FDynamicMesh3* Mesh, const FGroupTopology* Topology);
 
 	/**
-	 * Given a camera rectangle (from a marquee selection) gives the corners/edges/groups that are contained in
-	 * the rectangle. Only selects one of the types, and prefers corners to edges to groups.
-	 * 
-	 * For group edges, occlusion is checked at endpoints, and if either is occluded, the edge is considered occluded.
-	 * For faces, occlusion is checked at triangle centroids.
-	 *
-	 * @param Settings Settings that determine what is selected and whether occlusion is tested.
-	 * @param CameraRectangle The camera rectangle.
-	 * @param TargetTransform Transform from points to world space, since the rectangle is in world space.
-	 * @param ResultOut Output selection (cleared before use).
-	 * @param TriIsOccludedCache Cached values of whether a given Tid has an occluded centroid, since occlusion 
-	 *  values will stay the same if the rectangle is resized as long as the camera is not moved. This should
-	 *  be used when doing occlusion-based marquee face selection on meshes with many triangles, because the 
-	 *  occlusion checks are expensive.
-	 *
-	 * @return true if something was selected.
+	 * Initialize the selector with the given Mesh and Topology.
+	 * This does not create the internal data structures, this happens lazily on GetGeometrySet()
 	 */
-	bool FindSelectedElement(const FSelectionSettings& Settings, const FCameraRectangle& CameraRectangle, 
-		FTransform3d TargetTransform, FGroupTopologySelection& ResultOut, 
-		TMap<int32, bool>* TriIsOccludedCache = nullptr);
+	void Initialize(const FDynamicMesh3* Mesh, const FGroupTopology* Topology);
+
+	virtual void DrawSelection(const FGroupTopologySelection& Selection, FToolDataVisualizer* Renderer, const FViewCameraState* CameraState, ECornerDrawStyle CornerDrawStyle = ECornerDrawStyle::Point) override;
 
 	/**
-	 * Using the edges in the given selection as starting points, add any "edge loops" containing the edges. An 
+	 * Using the edges in the given selection as starting points, add any "edge loops" containing the edges. An
 	 * edge loop is a sequence of edges that passes through valence-4 corners through the opposite edge, and may
 	 * not actually form a complete loop if they hit a non-valence-4 corner.
 	 *
@@ -165,39 +159,14 @@ public:
 	 */
 	bool ExpandSelectionByEdgeRings(FGroupTopologySelection& Selection);
 
-	enum ECornerDrawStyle
-	{
-		Circle,
-		Point
-	};
-
-	/**
-	 * Render the given selection with the default settings of the FToolDataVisualizer.
-	 * Selected edges are drawn as lines, and selected corners are drawn as points or small, view-facing circles, based on the CornerDrawStyle parameter.
-	 * (Currently seleced faces are not draw)
-	 */
-	void DrawSelection(const FGroupTopologySelection& Selection, FToolDataVisualizer* Renderer, const FViewCameraState* CameraState, ECornerDrawStyle CornerDrawStyle = ECornerDrawStyle::Point);
-
-
-public:
-	// internal rendering parameters
-	float VisualAngleSnapThreshold = 0.5;
-
-
 protected:
-	const FDynamicMesh3* Mesh = nullptr;
+
+	UE_DEPRECATED(5.2, "Use TopologyProvider in the base class")
 	const FGroupTopology* Topology = nullptr;
 
-	TFunction<FDynamicMeshAABBTree3*(void)> GetSpatial = nullptr;
+private:
+
+	FGroupTopologyUtils GroupTopologyUtils;
 
-	bool bGeometryInitialized = false;
-	bool bGeometryUpToDate = false;
-	FGeometrySet3 GeometrySet;
+};
 
-	bool DoCornerBasedSelection(const FSelectionSettings& Settings, const FRay3d& Ray,
-		FDynamicMeshAABBTree3* Spatial, const FGeometrySet3& TopoSpatial, 
-		FGroupTopologySelection& ResultOut, FVector3d& SelectedPositionOut, int32 *EdgeSegmentIdOut) const;
-	bool DoEdgeBasedSelection(const FSelectionSettings& Settings, const FRay3d& Ray,
-		FDynamicMeshAABBTree3* Spatial, const FGeometrySet3& TopoSpatial,
-		FGroupTopologySelection& ResultOut, FVector3d& SelectedPositionOut, int32 *EdgeSegmentIdOut) const;
-};
\ No newline at end of file
diff --git a/Engine/Plugins/Runtime/MeshModelingToolset/Source/ModelingComponents/Public/Selection/MeshTopologySelectionMechanic.h b/Engine/Plugins/Runtime/MeshModelingToolset/Source/ModelingComponents/Public/Selection/MeshTopologySelectionMechanic.h
new file mode 100644
--- /dev/null
+++ b/Engine/Plugins/Runtime/MeshModelingToolset/Source/ModelingComponents/Public/Selection/MeshTopologySelectionMechanic.h
@@ -0,0 +1,448 @@
+// Copyright Epic Games, Inc. All Rights Reserved.
+
+#pragma once
+
+#include "CoreMinimal.h"
+#include "Drawing/PreviewGeometryActor.h"
+#include "Drawing/TriangleSetComponent.h"
+#include "InputBehavior.h"
+#include "InteractiveTool.h"
+#include "Mechanics/RectangleMarqueeMechanic.h"
+#include "Components/DynamicMeshComponent.h"
+#include "Selection/GroupTopologySelector.h"
+#include "TransformTypes.h"
+#include "ToolDataVisualizer.h"
+#include "InteractionMechanic.h"
+#include "MeshTopologySelectionMechanic.generated.h"
+
+class FMeshTopologySelectionMechanicSelectionChange;
+class UPersistentMeshSelection;
+class UMouseHoverBehavior;
+class UMeshTopologySelectionMechanic;
+class URectangleMarqueeMechanic;
+class USingleClickOrDragInputBehavior;
+
+using UE::Geometry::FDynamicMeshAABBTree3;
+PREDECLARE_USE_GEOMETRY_CLASS(FCompactMaps);
+
+UCLASS()
+class MODELINGCOMPONENTS_API UMeshTopologySelectionMechanicProperties : public UInteractiveToolPropertySet
+{
+	GENERATED_BODY()
+
+public:
+	UPROPERTY(EditAnywhere, Category = SelectionFilter)
+	bool bSelectVertices = true;
+
+	UPROPERTY(EditAnywhere, Category = SelectionFilter)
+	bool bSelectEdges = true;
+
+	UPROPERTY(EditAnywhere, Category = SelectionFilter)
+	bool bSelectFaces = true;
+
+	/** When true, will select edge loops. Edge loops are paths along a string of valence-4 vertices. */
+	UPROPERTY(EditAnywhere, Category = SelectionFilter, meta = (EditCondition = "bSelectEdges"))
+	bool bSelectEdgeLoops = false;
+
+	/** When true, will select rings of edges that are opposite each other across a quad face. */
+	UPROPERTY(EditAnywhere, Category = SelectionFilter, meta = (EditCondition = "bSelectEdges"))
+	bool bSelectEdgeRings = false;
+
+	/** When false, faces that face away from the camera are ignored in selection and occlusion. Useful for working with inside-out meshes. */
+	UPROPERTY(EditAnywhere, Category = SelectionFilter, AdvancedDisplay)
+	bool bHitBackFaces = true;
+
+	UPROPERTY(EditAnywhere, Category = SelectionFilter, AdvancedDisplay)
+	bool bEnableMarquee = true;
+
+	/** Determines whether vertices should be checked for occlusion in marquee select (Note: marquee select currently only works with edges and vertices) */
+	UPROPERTY(EditAnywhere, Category = SelectionFilter, meta = (EditCondition = "bEnableMarquee", EditConditionHides))
+	bool bMarqueeIgnoreOcclusion = true;
+
+	// The following were originally in their own category, all marked as AdvancedDisplay. However, since there wasn't a non-AdvancedDisplay
+	// property in the category, they started out as expanded and could not be collapsed.
+	// The alternative approach, used below, is to have them in a nested category, which starts out as collapsed. This works nicely.
+
+	/** Prefer to select an edge projected to a point rather than the point, or a face projected to an edge rather than the edge. */
+	UPROPERTY(EditAnywhere, Category = "SelectionFilter|Ortho Viewport Behavior")
+	bool bPreferProjectedElement = true;
+
+	/** If the closest element is valid, select other elements behind it that are aligned with it. */
+	UPROPERTY(EditAnywhere, Category = "SelectionFilter|Ortho Viewport Behavior")
+	bool bSelectDownRay = true;
+
+	/** Do not check whether the closest element is occluded from the current view. */
+	UPROPERTY(EditAnywhere, Category = "SelectionFilter|Ortho Viewport Behavior")
+	bool bIgnoreOcclusion = false;
+
+	/** Invert current selection. If selection is empty, has same effect as Select All, and is similarly dependent on selection filter. */
+	UFUNCTION(CallInEditor, Category = SelectionActions)
+	void InvertSelection();
+
+	/** Select all elements. Depends on selection filter, where vertices are preferred to edges to faces. */
+	UFUNCTION(CallInEditor, Category = SelectionActions)
+	void SelectAll();
+
+	void Initialize(UMeshTopologySelectionMechanic* MechanicIn)
+	{
+		Mechanic = MechanicIn;
+	}
+
+protected:
+
+	TWeakObjectPtr<UMeshTopologySelectionMechanic> Mechanic;
+};
+
+
+
+/**
+ * Base class mechanic for selecting a subset of mesh elements (edge loops, groups, corners, etc.)
+ * Internally it relies on an FMeshTopologySelector to define which type of mesh topology is selectable.
+ * 
+ * NOTE: Users should not use this class directly, but rather subclass it and specify a particular FMeshTopologySelector to use.
+ */
+
+UCLASS(Abstract)
+class MODELINGCOMPONENTS_API UMeshTopologySelectionMechanic : public UInteractionMechanic, public IClickBehaviorTarget, public IHoverBehaviorTarget
+{
+	GENERATED_BODY()
+	using FFrame3d = UE::Geometry::FFrame3d;
+	using FAxisAlignedBox3d = UE::Geometry::FAxisAlignedBox3d;
+public:
+
+	virtual ~UMeshTopologySelectionMechanic();
+
+	// configuration variables that must be set before bSetup is called
+	bool bAddSelectionFilterPropertiesToParentTool = true;
+
+	void Initialize(const FDynamicMesh3* MeshIn,
+		FTransform3d TargetTransformIn,
+		UWorld* WorldIn,
+		TFunction<FDynamicMeshAABBTree3* ()> GetSpatialSourceFuncIn);
+
+	virtual void Setup(UInteractiveTool* ParentTool) override;
+	virtual void Shutdown() override;
+
+	virtual void Render(IToolsContextRenderAPI* RenderAPI) override;
+	virtual void DrawHUD(FCanvas* Canvas, IToolsContextRenderAPI* RenderAPI);
+
+	/**
+	 * Removes the mechanic's own click/hover handlers, which means that the parent tool
+	 * will need to call UpdateSelection(), UpdateHighlight(), ClearHighlight(), and 
+	 * ClearSelection() from its own hover/click handlers.
+	 *
+	 * Must be called during tool Setup() after calling Setup() on the mechanic to have an effect.
+	 *
+	 * @param ParentToolIn The parent tool, needed to be able to remove the behaviors.
+	 */
+	void DisableBehaviors(UInteractiveTool* ParentToolIn);
+
+	/**
+	 * Enable/disable the mechanic without permanently removing behaviors or shutting it down.
+	 */
+	void SetIsEnabled(bool bOn);
+
+	/**
+	 * Sets the base priority so that tools can make sure that their own behaviors are higher
+	 * priority. The mechanic will not use any priority value higher than this, but it may use
+	 * lower if it needs to stagger the priorities of behaviors it uses.
+	 * Can be called before or after Setup().
+	 */
+	void SetBasePriority(const FInputCapturePriority& Priority);
+
+	/**
+	 * Gets the current priority range used by behaviors in the mechanic. The returned pair will
+	 * have the base (highest) priority as the key, and the lowest priority as the value.
+	 */
+	TPair<FInputCapturePriority, FInputCapturePriority> GetPriorityRange() const;
+
+
+	void SetShouldSelectEdgeLoopsFunc(TFunction<bool(void)> Func)
+	{
+		ShouldSelectEdgeLoopsFunc = Func;
+	}
+
+	void SetShouldSelectEdgeRingsFunc(TFunction<bool(void)> Func)
+	{
+		ShouldSelectEdgeRingsFunc = Func;
+	}
+
+	/**
+	 * By default, the shift key will cause new clicks to add to the selection. However, this
+	 * can be changed by supplying a different function to check here.
+	 */
+	void SetShouldAddToSelectionFunc(TFunction<bool(void)> Func)
+	{
+		ShouldAddToSelectionFunc = Func;
+	}
+
+	/**
+	 * By default, the Ctrl key will cause new clicks to remove from the existing selection.
+	 * However, this can be changed by supplying a different function to check here.
+	 */
+	void SetShouldRemoveFromSelectionFunc(TFunction<bool(void)> Func)
+	{
+		ShouldRemoveFromSelectionFunc = Func;
+	}
+
+	/**
+	 * Notify internal data structures that the associated MeshComponent has been modified.
+	 * @param bTopologyModified if true, the underlying mesh topology has been changed. This clears the current selection.
+	 */
+	void NotifyMeshChanged(bool bTopologyModified);
+
+	/**
+	 * Perform a hit test on the topology using the current selection settings. In cases of hitting edges and
+	 * corners, OutHit contains the following:
+	 *   OutHit.FaceIndex: edge or corner id in the topology
+	 *   OutHit.ImpactPoint: closest point on the ray to the hit element (Note: not a point on the element!)
+	 *   OutHit.Distance: distance along the ray to ImpactPoint
+	 *   OutHit.Item: if hit item was an edge, index of the segment within the edge polyline. Otherwise undefined.
+	 *
+	 * @param bUseOrthoSettings If true, the ortho-relevant settings for selection are used (selecting down the view ray, etc)
+	 */
+	bool TopologyHitTest(const FRay& WorldRay, FHitResult& OutHit, FGroupTopologySelection& OutSelection, bool bUseOrthoSettings = false);
+	bool TopologyHitTest(const FRay& WorldRay, FHitResult& OutHit, bool bUseOrthoSettings = false);
+
+	TSharedPtr<FMeshTopologySelector, ESPMode::ThreadSafe> GetTopologySelector() { return TopoSelector; }
+
+	//
+	// Hover API
+	//
+
+	/**
+	 * Update the hover highlight based on the hit elements at the given World Ray
+	 * @return true if something was hit and is now being hovered
+	 */
+	virtual bool UpdateHighlight(const FRay& WorldRay) PURE_VIRTUAL(UMeshTopologySelectionMechanic::UpdateHighlight, return false; );
+
+	/**
+	 * Clear current hover-highlight
+	 */
+	void ClearHighlight();
+
+
+	//
+	// Selection API
+	//
+
+	/** 
+	 * Intersect the ray with the mesh and update the selection based on the hit element, modifier states, etc
+	 * @return true if selection was modified
+	 */
+	virtual bool UpdateSelection(const FRay& WorldRay, FVector3d& LocalHitPositionOut, FVector3d& LocalHitNormalOut) PURE_VIRTUAL(UMeshTopologySelectionMechanic::UpdateSelection, return false; );
+
+	/**
+	 * Replace the current selection with an external selection. 
+	 * @warning does not check that the selection is valid!
+	 *
+	 * @param bBroadcast If true, issues an OnSelectionChanged delegate broadcast.
+	 */
+	void SetSelection(const FGroupTopologySelection& Selection, bool bBroadcast = true);
+
+	/**
+	 * Clear the current selection.
+	 */
+	void ClearSelection();
+
+	void InvertSelection();
+	void SelectAll();
+
+	/** 
+	 * @return true if the current selection is non-empty 
+	 */
+	bool HasSelection() const;
+
+	/**
+	 * @return the current selection
+	 */
+	const FGroupTopologySelection& GetActiveSelection() const { return PersistentSelection; }
+
+	/**
+	 * Can be used by in an OnSelectionChanged event to inspect the clicked location (i.e., the
+	 * values returned by the UpdateSelection() function when the click happened).
+	 */
+	void GetClickedHitPosition(FVector3d& PositionOut, FVector3d& NormalOut) const;
+
+	/**
+	 * @return The best-guess 3D frame for the current select
+	 * @param bWorld if true, local-to-world transform of the target MeshComponent is applied to the frame
+	 */
+	FFrame3d GetSelectionFrame(bool bWorld, FFrame3d* InitialLocalFrame = nullptr) const;
+
+	/**
+	 * @return Bounding box for the current selection
+	 * @param bWorld if true, the box is in world space, otherwise it is in local space of the target MeshComponent
+	 */
+	FAxisAlignedBox3d GetSelectionBounds(bool bWorld) const;
+
+	void SetShowSelectableCorners(bool bShowCorners);
+
+	//
+	// Change Tracking
+	//
+
+	/**
+	 * Begin a change record. Internally creates a FCommandChange and initializes it with current state
+	 */
+	void BeginChange();
+
+	/**
+	 * End the active change and return it. Returns empty change if the selection was not modified!
+	 */
+	TUniquePtr<FToolCommandChange> EndChange();
+
+	/**
+	 * Ends the active change and emits it via the parent tool, if the selection has been modified.
+	 */
+	bool EndChangeAndEmitIfModified();
+
+	// IClickBehaviorTarget implementation
+	virtual FInputRayHit IsHitByClick(const FInputDeviceRay& ClickPos) override;
+	virtual void OnClicked(const FInputDeviceRay& ClickPos) override;
+
+	// IHoverBehaviorTarget implementation
+	virtual FInputRayHit BeginHoverSequenceHitTest(const FInputDeviceRay& PressPos) override;
+	virtual void OnBeginHover(const FInputDeviceRay& DevicePos) override;
+	virtual bool OnUpdateHover(const FInputDeviceRay& DevicePos) override;
+	virtual void OnEndHover() override;
+	virtual void OnUpdateModifierState(int ModifierID, bool bIsOn) override;
+
+protected:
+	// These get bound to marquee mechanic delegates.
+	virtual void OnDragRectangleStarted();
+	virtual void OnDragRectangleChanged(const FCameraRectangle& CurrentRectangle);
+	virtual void OnDragRectangleFinished(const FCameraRectangle& Rectangle, bool bCancelled);
+
+	virtual void UpdateMarqueeEnabled();
+
+public:
+	/** 
+	 * OnSelectionChanged is broadcast whenever the selection is modified (including by FChanges, which
+	 * means that called functions should not issue undo transactions. 
+	 */
+	FSimpleMulticastDelegate OnSelectionChanged;
+
+	/**
+	 * This is broadcast during marquee rectangle dragging if selected faces change, to allow user to
+	 * dynamically update face highlighting if desired (needed because OnSelectionChanged is only 
+	 * fired once the rectangle is completed, not while it is being updated).
+	 */
+	FSimpleMulticastDelegate OnFaceSelectionPreviewChanged;
+
+	/**
+	 * When true, the selection mechanic is currently tracking a marquee rectangle drag, and acting
+	 * on the selection may be unwise until it is over (and an OnSelectionChanged event is fired).
+	 */
+	bool IsCurrentlyMarqueeDragging() { return bCurrentlyMarqueeDragging; }
+
+	// TODO: Is it worth issuing separate callbacks in normal selection changes and in FChange ones, to
+	// allow the user to bundle in some FChanges into the normal callback?
+
+	UPROPERTY()
+	TObjectPtr<UMeshTopologySelectionMechanicProperties> Properties;
+
+protected:
+
+	//
+	// Subclass should initialize this with a concrete subclass of FMeshTopologySelector
+	//
+	TSharedPtr<FMeshTopologySelector, ESPMode::ThreadSafe> TopoSelector;
+
+
+	bool bIsEnabled = true;
+
+	const FDynamicMesh3* Mesh;
+	TFunction<FDynamicMeshAABBTree3*()> GetSpatialFunc;
+
+	UPROPERTY()
+	TObjectPtr<UMouseHoverBehavior> HoverBehavior;
+
+	UPROPERTY()
+	TObjectPtr<USingleClickOrDragInputBehavior> ClickOrDragBehavior;
+
+	UPROPERTY()
+	TObjectPtr<URectangleMarqueeMechanic> MarqueeMechanic;
+
+	FInputCapturePriority BasePriority = FInputCapturePriority(FInputCapturePriority::DEFAULT_TOOL_PRIORITY);
+
+	// When bSelectEdgeLoops is true, this function is tested to see if we should select edge loops,
+	// to allow edge loop selection to be toggled with some key (setting bSelectEdgeLoops to
+	// false overrides this function).
+	TFunction<bool(void)> ShouldSelectEdgeLoopsFunc = []() {return true; };
+
+	// When bSelectEdgeRings is true, this function is tested to see if we should select edge rings,
+	// to allow edge ring selection to be toggled with some key (setting bSelectEdgeRings to
+	// false overrides this function).
+	TFunction<bool(void)> ShouldSelectEdgeRingsFunc = []() {return true; };
+
+	TFunction<bool(void)> ShouldAddToSelectionFunc = [this]() {return bShiftToggle; };
+	TFunction<bool(void)> ShouldRemoveFromSelectionFunc = [this]() {return bCtrlToggle; };;
+
+	FTransform3d TargetTransform;
+
+	/**
+	 * Get the topology selector settings to use given the current selection settings.
+	 * 
+	 * @param bUseOrthoSettings If true, the topology selector will be configured to use ortho settings,
+	 *  which are generally different to allow for selection of projected elements, etc.
+	 */
+	FGroupTopologySelector::FSelectionSettings GetTopoSelectorSettings(bool bUseOrthoSettings = false);
+
+	FGroupTopologySelection HilightSelection;
+	FGroupTopologySelection PersistentSelection;
+	int32 SelectionTimestamp = 0;
+	TUniquePtr<FMeshTopologySelectionMechanicSelectionChange> ActiveChange;
+
+	// Used for box selection
+	FGroupTopologySelection PreDragPersistentSelection;
+	FGroupTopologySelection LastUpdateRectangleSelection;
+	FGroupTopologySelector::FSelectionSettings PreDragTopoSelectorSettings;
+	TMap<int32, bool> TriIsOccludedCache;
+	bool bCurrentlyMarqueeDragging = false;
+
+	FVector3d LastClickedHitPosition;
+	FVector3d LastClickedHitNormal;
+
+	/** The actor we create internally to own the DrawnTriangleSetComponent */
+	UPROPERTY()
+	TObjectPtr<APreviewGeometryActor> PreviewGeometryActor;
+
+	UPROPERTY()
+	TObjectPtr<UTriangleSetComponent> DrawnTriangleSetComponent;
+
+	TSet<int> CurrentlyHighlightedGroups;
+
+	UPROPERTY()
+	TObjectPtr<UMaterialInterface> HighlightedFaceMaterial;
+
+	FViewCameraState CameraState;
+
+	bool bShiftToggle = false;
+	bool bCtrlToggle = false;
+	static const int32 ShiftModifierID = 1;
+	static const int32 CtrlModifierID = 2;
+
+	bool bShowSelectableCorners = true;
+public:
+	FToolDataVisualizer PolyEdgesRenderer;
+	FToolDataVisualizer HilightRenderer;
+	FToolDataVisualizer SelectionRenderer;
+
+	friend class FMeshTopologySelectionMechanicSelectionChange;
+
+};
+
+
+class MODELINGCOMPONENTS_API FMeshTopologySelectionMechanicSelectionChange : public FToolCommandChange
+{
+public:
+	FGroupTopologySelection Before;
+	FGroupTopologySelection After;
+	int32 Timestamp = 0;
+
+	virtual void Apply(UObject* Object) override;
+	virtual void Revert(UObject* Object) override;
+	virtual FString ToString() const override;
+};
+
diff --git a/Engine/Plugins/Runtime/MeshModelingToolset/Source/ModelingComponents/Public/Selection/MeshTopologySelector.h b/Engine/Plugins/Runtime/MeshModelingToolset/Source/ModelingComponents/Public/Selection/MeshTopologySelector.h
new file mode 100644
--- /dev/null
+++ b/Engine/Plugins/Runtime/MeshModelingToolset/Source/ModelingComponents/Public/Selection/MeshTopologySelector.h
@@ -0,0 +1,225 @@
+// Copyright Epic Games, Inc. All Rights Reserved.
+
+#pragma once
+
+#include "Spatial/GeometrySet3.h"
+#include "DynamicMesh/DynamicMeshAABBTree3.h"
+#include "Polyline3.h"
+#include "GroupTopology.h"		// FGroupTopologySelection
+
+struct FCameraRectangle;
+class FToolDataVisualizer;
+struct FViewCameraState;
+
+namespace UE::Geometry
+{
+	class FDynamicMesh3;
+}
+
+using UE::Geometry::FGroupTopologySelection;
+using UE::Geometry::FDynamicMeshAABBTree3;
+using UE::Geometry::FDynamicMesh3;
+using UE::Geometry::FGeometrySet3;
+using UE::Geometry::FPolyline3d;
+using UE::Geometry::FIndex2i;
+using UE::Geometry::FFrame3d;
+using UE::Geometry::FAxisAlignedBox3d;
+
+/**
+* FTopologyProvider provides selectable edge/group/corner information for FMeshTopologySelector to use. Users should implement a concrete subclass
+* based on the type of mesh topology they wish to make selectable (e.g. group topology, boundary curves, etc.). Then creating a new TopologySelector
+* subclass should be fairly trivial.
+*/
+class FTopologyProvider
+{
+public:
+	
+	virtual ~FTopologyProvider() = default;
+
+	virtual int GetNumCorners() const = 0;
+	virtual int GetCornerVertexID(int CornerID) const = 0;
+	
+	virtual int GetNumEdges() const = 0;
+	virtual void GetEdgePolyline(int EdgeID, FPolyline3d& OutPolyline) const = 0;
+	virtual int FindGroupEdgeID(int MeshEdgeID) const = 0;
+	
+	virtual const TArray<int>& GetGroupEdgeEdges(int GroupEdgeID) const = 0;
+	virtual const TArray<int>& GetGroupEdgeVertices(int GroupEdgeID) const = 0;
+
+	virtual int GetNumGroups() const = 0;
+	virtual int GetGroupIDAt(int GroupIndex) const = 0;
+	virtual int GetGroupIDForTriangle(int TriangleID) const = 0;
+
+	virtual void ForGroupSetEdges(const TSet<int32>& GroupIDs, const TFunction<void(int EdgeID)>& EdgeFunc) const = 0;
+
+	virtual FFrame3d GetSelectionFrame(const FGroupTopologySelection& Selection, FFrame3d* InitialLocalFrame = nullptr) const = 0;
+	virtual FAxisAlignedBox3d GetSelectionBounds(const FGroupTopologySelection& Selection, TFunctionRef<FVector3d(const FVector3d&)> TransformFunc) const = 0;
+};
+
+/**
+ * FMeshTopologySelector implements selection behavior for a subset of mesh elements.
+ * Topological Groups, Edges, and Corners can be selected depending on the subclass of FTopologyProvider used.
+ *
+ * Internally a FGeometrySet3 is constructed to support ray-hit testing against the edges and corners.
+ *
+ * Note that to hit-test against the mesh you have to provide a your own FDynamicMeshAABBTree3.
+ * You do this by providing a callback via SetSpatialSource(). The reason for this is that
+ * (1) owners of an FMeshTopologySelector likely already have a BVTree and (2) if the use case
+ * is deformation, we need to make sure the owner has recomputed the BVTree before we call functions
+ * on it. The callback you provide should do that!
+ *
+ * DrawSelection() can be used to visualize a selection via line/circle drawing.
+ *
+ * @todo optionally have an internal mesh AABBTree that can be used when owner does not provide one?
+ */
+
+class MODELINGCOMPONENTS_API FMeshTopologySelector
+{
+public:
+
+	virtual ~FMeshTopologySelector() = default;
+
+	//
+	// Configuration variables
+	// 
+
+	/**
+	 * Determines the behavior of a FindSelectedElement() call.
+	 */
+	struct FSelectionSettings
+	{
+		bool bEnableFaceHits = true;
+		bool bEnableEdgeHits = true;
+		bool bEnableCornerHits = true;
+		
+		// When false, hits of a triangles that are facing away from the camera are not considered. This can
+		// be useful when working with inside-out meshes, where we usually want to select the farther, visible
+		// wall.
+		bool bHitBackFaces = true;
+
+		// The following are mainly useful for ortho viewport selection:
+
+		// Prefer an edge projected to a point rather than the point, and a face projected to an edge
+		// rather than the edge.
+		bool bPreferProjectedElement = false;
+
+		// If the first element is valid, select all elements behind it that are aligned with it.
+		bool bSelectDownRay = false;
+
+		//Do not check whether the closest element is occluded.
+		bool bIgnoreOcclusion = false;
+	};
+
+	/** 
+	 * This is the function we use to determine if a point on a corner/edge is close enough
+	 * to the hit-test ray to treat as a "hit". By default this is Euclidean distance with
+	 * a tolerance of 1.0. You probably need to replace this with your own function.
+	 */
+	TFunction<bool(const FVector3d&, const FVector3d&, double ToleranceScale)> PointsWithinToleranceTest;
+
+public:
+
+	FMeshTopologySelector();
+
+
+	/**
+	 * Provide a function that will return an AABBTree for the Mesh.
+	 * See class comment for why this is necessary.
+	 */
+	void SetSpatialSource(TFunction<FDynamicMeshAABBTree3*(void)> GetSpatialFunc)
+	{
+		GetSpatial = GetSpatialFunc;
+	}
+
+	/**
+	 * Notify the Selector that the mesh has changed. 
+	 * @param bTopologyDeformed if this is true, the mesh vertices have been moved so we need to update bounding boxes/etc
+	 * @param bTopologyModified if this is true, topology has changed and we need to rebuild spatial data structures from scratch
+	 */
+	void Invalidate(bool bTopologyDeformed, bool bTopologyModified);
+
+	/**
+	 * @return the internal GeometrySet. This does lazy updating of the GeometrySet, so this function may take some time.
+	 */
+	const FGeometrySet3& GetGeometrySet();
+
+	/**
+	 * Find which element was selected for a given ray.
+	 *
+	 * @param Settings settings that determine what can be selected.
+	 * @param Ray hit-test ray
+	 * @param ResultOut resulting selection. Will not be cleared before use. At most one of the Groups/Corners/Edges
+	 *  members will be added to.
+	 * @param SelectedPositionOut The point on the ray nearest to the selected element
+	 * @param SelectedNormalOut the normal at that point if ResultOut contains a selected face, otherwise uninitialized
+	 * @param EdgeSegmentIdOut When not null, if the selected element is a group edge, the segment id of the segment
+	 *   that was actually clicked within the edge polyline will be stored here.
+	 *
+	 * @return true if something was selected
+	 */
+	bool FindSelectedElement(const FSelectionSettings& Settings, const FRay3d& Ray, FGroupTopologySelection& ResultOut,
+		FVector3d& SelectedPositionOut, FVector3d& SelectedNormalOut, int32* EdgeSegmentIdOut = nullptr);
+
+	/**
+	 * Given a camera rectangle (from a marquee selection) gives the corners/edges/groups that are contained in
+	 * the rectangle. Only selects one of the types, and prefers corners to edges to groups.
+	 * 
+	 * For group edges, occlusion is checked at endpoints, and if either is occluded, the edge is considered occluded.
+	 * For faces, occlusion is checked at triangle centroids.
+	 *
+	 * @param Settings Settings that determine what is selected and whether occlusion is tested.
+	 * @param CameraRectangle The camera rectangle.
+	 * @param TargetTransform Transform from points to world space, since the rectangle is in world space.
+	 * @param ResultOut Output selection (cleared before use).
+	 * @param TriIsOccludedCache Cached values of whether a given Tid has an occluded centroid, since occlusion 
+	 *  values will stay the same if the rectangle is resized as long as the camera is not moved. This should
+	 *  be used when doing occlusion-based marquee face selection on meshes with many triangles, because the 
+	 *  occlusion checks are expensive.
+	 *
+	 * @return true if something was selected.
+	 */
+	bool FindSelectedElement(const FSelectionSettings& Settings, const FCameraRectangle& CameraRectangle, 
+		FTransform3d TargetTransform, FGroupTopologySelection& ResultOut, 
+		TMap<int32, bool>* TriIsOccludedCache = nullptr);
+
+	enum ECornerDrawStyle
+	{
+		Circle,
+		Point
+	};
+
+	/**
+	 * Render the given selection with the default settings of the FToolDataVisualizer.
+	 * Selected edges are drawn as lines, and selected corners are drawn as points or small, view-facing circles, based on the CornerDrawStyle parameter.
+	 * (Currently selected faces are not drawn)
+	 */
+	virtual void DrawSelection(const FGroupTopologySelection& Selection, FToolDataVisualizer* Renderer, const FViewCameraState* CameraState, ECornerDrawStyle CornerDrawStyle = ECornerDrawStyle::Point) = 0;
+
+	const FTopologyProvider* GetTopologyProvider()
+	{
+		return TopologyProvider.Get();
+	}
+
+public:
+
+	// internal rendering parameters
+	float VisualAngleSnapThreshold = 0.5;
+
+
+protected:
+
+	const FDynamicMesh3* Mesh = nullptr;
+	TUniquePtr<FTopologyProvider> TopologyProvider = nullptr;
+	TFunction<FDynamicMeshAABBTree3*(void)> GetSpatial = nullptr;
+
+	bool bGeometryInitialized = false;
+	bool bGeometryUpToDate = false;
+	FGeometrySet3 GeometrySet;
+
+	bool DoCornerBasedSelection(const FSelectionSettings& Settings, const FRay3d& Ray,
+		FDynamicMeshAABBTree3* Spatial, const FGeometrySet3& TopoSpatial, 
+		FGroupTopologySelection& ResultOut, FVector3d& SelectedPositionOut, int32 *EdgeSegmentIdOut) const;
+	bool DoEdgeBasedSelection(const FSelectionSettings& Settings, const FRay3d& Ray,
+		FDynamicMeshAABBTree3* Spatial, const FGeometrySet3& TopoSpatial,
+		FGroupTopologySelection& ResultOut, FVector3d& SelectedPositionOut, int32 *EdgeSegmentIdOut) const;
+};
diff --git a/Engine/Plugins/Runtime/MeshModelingToolset/Source/ModelingComponents/Public/Selection/PolygonSelectionMechanic.h b/Engine/Plugins/Runtime/MeshModelingToolset/Source/ModelingComponents/Public/Selection/PolygonSelectionMechanic.h
--- a/Engine/Plugins/Runtime/MeshModelingToolset/Source/ModelingComponents/Public/Selection/PolygonSelectionMechanic.h
+++ b/Engine/Plugins/Runtime/MeshModelingToolset/Source/ModelingComponents/Public/Selection/PolygonSelectionMechanic.h
@@ -2,272 +2,50 @@
 
 #pragma once
 
-#include "CoreMinimal.h"
-#include "Drawing/PreviewGeometryActor.h"
-#include "Drawing/TriangleSetComponent.h"
-#include "InputBehavior.h"
-#include "InteractiveTool.h"
-#include "Mechanics/RectangleMarqueeMechanic.h"
-#include "Components/DynamicMeshComponent.h"
-#include "Selection/GroupTopologySelector.h"
-#include "TransformTypes.h"
-#include "ToolDataVisualizer.h"
-#include "InteractionMechanic.h"
+#include "MeshTopologySelectionMechanic.h"
 #include "PolygonSelectionMechanic.generated.h"
 
-class FPolygonSelectionMechanicSelectionChange;
-class UPersistentMeshSelection;
-class UMouseHoverBehavior;
-class UPolygonSelectionMechanic;
-class URectangleMarqueeMechanic;
-class USingleClickOrDragInputBehavior;
-
-using UE::Geometry::FDynamicMeshAABBTree3;
-PREDECLARE_USE_GEOMETRY_CLASS(FCompactMaps);
-
-UCLASS()
-class MODELINGCOMPONENTS_API UPolygonSelectionMechanicProperties : public UInteractiveToolPropertySet
+// DEPRECATED: Use UMeshTopologySelectionMechanicProperties
+UCLASS(Deprecated)
+class MODELINGCOMPONENTS_API UDEPRECATED_PolygonSelectionMechanicProperties : public UMeshTopologySelectionMechanicProperties
 {
 	GENERATED_BODY()
 
 public:
-	UPROPERTY(EditAnywhere, Category = SelectionFilter)
-	bool bSelectVertices = true;
-
-	UPROPERTY(EditAnywhere, Category = SelectionFilter)
-	bool bSelectEdges = true;
-
-	UPROPERTY(EditAnywhere, Category = SelectionFilter)
-	bool bSelectFaces = true;
-
-	/** When true, will select edge loops. Edge loops are paths along a string of valence-4 vertices. */
-	UPROPERTY(EditAnywhere, Category = SelectionFilter, meta = (EditCondition = "bSelectEdges"))
-	bool bSelectEdgeLoops = false;
-
-	/** When true, will select rings of edges that are opposite each other across a quad face. */
-	UPROPERTY(EditAnywhere, Category = SelectionFilter, meta = (EditCondition = "bSelectEdges"))
-	bool bSelectEdgeRings = false;
-
-	/** When false, faces that face away from the camera are ignored in selection and occlusion. Useful for working with inside-out meshes. */
-	UPROPERTY(EditAnywhere, Category = SelectionFilter, AdvancedDisplay)
-	bool bHitBackFaces = true;
-
-	UPROPERTY(EditAnywhere, Category = SelectionFilter, AdvancedDisplay)
-	bool bEnableMarquee = true;
-
-	/** Determines whether vertices should be checked for occlusion in marquee select (Note: marquee select currently only works with edges and vertices) */
-	UPROPERTY(EditAnywhere, Category = SelectionFilter, meta = (EditCondition = "bEnableMarquee", EditConditionHides))
-	bool bMarqueeIgnoreOcclusion = true;
-
-	// The following were originally in their own category, all marked as AdvancedDisplay. However, since there wasn't a non-AdvancedDisplay
-	// property in the category, they started out as expanded and could not be collapsed.
-	// The alternative approach, used below, is to have them in a nested category, which starts out as collapsed. This works nicely.
-
-	/** Prefer to select an edge projected to a point rather than the point, or a face projected to an edge rather than the edge. */
-	UPROPERTY(EditAnywhere, Category = "SelectionFilter|Ortho Viewport Behavior")
-	bool bPreferProjectedElement = true;
-
-	/** If the closest element is valid, select other elements behind it that are aligned with it. */
-	UPROPERTY(EditAnywhere, Category = "SelectionFilter|Ortho Viewport Behavior")
-	bool bSelectDownRay = true;
-
-	/** Do not check whether the closest element is occluded from the current view. */
-	UPROPERTY(EditAnywhere, Category = "SelectionFilter|Ortho Viewport Behavior")
-	bool bIgnoreOcclusion = false;
-
-	/** Invert current selection. If selection is empty, has same effect as Select All, and is similarly dependent on selection filter. */
-	UFUNCTION(CallInEditor, Category = SelectionActions)
-	void InvertSelection();
-
-	/** Select all elements. Depends on selection filter, where vertices are preferred to edges to faces. */
-	UFUNCTION(CallInEditor, Category = SelectionActions)
-	void SelectAll();
 
 	void Initialize(UPolygonSelectionMechanic* MechanicIn)
 	{
 		Mechanic = MechanicIn;
 	}
-
-protected:
-	TWeakObjectPtr<UPolygonSelectionMechanic> Mechanic;
 };
 
 
-
 /**
  * UPolygonSelectionMechanic implements the interaction for selecting a set of faces/vertices/edges
- * from a FGroupTopology on a UDynamicMeshComponent. 
+ * from a FGroupTopology.
  */
 UCLASS()
-class MODELINGCOMPONENTS_API UPolygonSelectionMechanic : public UInteractionMechanic, public IClickBehaviorTarget, public IHoverBehaviorTarget
+class MODELINGCOMPONENTS_API UPolygonSelectionMechanic : public UMeshTopologySelectionMechanic
 {
 	GENERATED_BODY()
-	using FFrame3d = UE::Geometry::FFrame3d;
-	using FAxisAlignedBox3d = UE::Geometry::FAxisAlignedBox3d;
-public:
-
-	virtual ~UPolygonSelectionMechanic();
 
-	// configuration variables that must be set before bSetup is called
-	bool bAddSelectionFilterPropertiesToParentTool = true;
-
-	virtual void Setup(UInteractiveTool* ParentTool) override;
-	virtual void Shutdown() override;
-
-	virtual void Render(IToolsContextRenderAPI* RenderAPI) override;
-	virtual void DrawHUD(FCanvas* Canvas, IToolsContextRenderAPI* RenderAPI);
+public:
 
-	/**
-	 * Initializes the mechanic.
-	 *
-	 * @param Mesh Mesh that we are operating on.
-	 * @param TargetTransform Transform of the mesh.
-	 * @param World World in which we are operating, used to add drawing components that draw highlighted edges.
-	 * @param Topology Group topology of the mesh.
-	 * @param GetSpatialSourceFunc Function that returns an AABB tree for the mesh.
-	 */
 	void Initialize(const FDynamicMesh3* Mesh,
 		FTransform3d TargetTransform,
 		UWorld* World,
 		const FGroupTopology* Topology,
-		TFunction<FDynamicMeshAABBTree3*()> GetSpatialSourceFunc
-		);
+		TFunction<FDynamicMeshAABBTree3* ()> GetSpatialSourceFunc
+	);
 
 	void Initialize(UDynamicMeshComponent* MeshComponent, const FGroupTopology* Topology,
-		TFunction<FDynamicMeshAABBTree3 * ()> GetSpatialSourceFunc
+		TFunction<FDynamicMeshAABBTree3* ()> GetSpatialSourceFunc
 	);
 
-	/**
-	 * Removes the mechanic's own click/hover handlers, which means that the parent tool
-	 * will need to call UpdateSelection(), UpdateHighlight(), ClearHighlight(), and 
-	 * ClearSelection() from its own hover/click handlers.
-	 *
-	 * Must be called during tool Setup() after calling Setup() on the mechanic to have an effect.
-	 *
-	 * @param ParentToolIn The parent tool, needed to be able to remove the behaviors.
-	 */
-	void DisableBehaviors(UInteractiveTool* ParentToolIn);
-
-	/**
-	 * Enable/disable the mechanic without permanently removing behaviors or shutting it down.
-	 */
-	void SetIsEnabled(bool bOn);
-
-	/**
-	 * Sets the base priority so that tools can make sure that their own behaviors are higher
-	 * priority. The mechanic will not use any priority value higher than this, but it may use
-	 * lower if it needs to stagger the priorities of behaviors it uses.
-	 * Can be called before or after Setup().
-	 */
-	void SetBasePriority(const FInputCapturePriority& Priority);
-
-	/**
-	 * Gets the current priority range used by behaviors in the mechanic. The returned pair will
-	 * have the base (highest) priority as the key, and the lowest priority as the value.
-	 */
-	TPair<FInputCapturePriority, FInputCapturePriority> GetPriorityRange() const;
-
-
-	void SetShouldSelectEdgeLoopsFunc(TFunction<bool(void)> Func)
-	{
-		ShouldSelectEdgeLoopsFunc = Func;
-	}
-
-	void SetShouldSelectEdgeRingsFunc(TFunction<bool(void)> Func)
-	{
-		ShouldSelectEdgeRingsFunc = Func;
-	}
-
-	/**
-	 * By default, the shift key will cause new clicks to add to the selection. However, this
-	 * can be changed by supplying a different function to check here.
-	 */
-	void SetShouldAddToSelectionFunc(TFunction<bool(void)> Func)
-	{
-		ShouldAddToSelectionFunc = Func;
-	}
-
-	/**
-	 * By default, the Ctrl key will cause new clicks to remove from the existing selection.
-	 * However, this can be changed by supplying a different function to check here.
-	 */
-	void SetShouldRemoveFromSelectionFunc(TFunction<bool(void)> Func)
-	{
-		ShouldRemoveFromSelectionFunc = Func;
-	}
-
-	/**
-	 * Notify internal data structures that the associated MeshComponent has been modified.
-	 * @param bTopologyModified if true, the underlying mesh topology has been changed. This clears the current selection.
-	 */
-	void NotifyMeshChanged(bool bTopologyModified);
-
-	/**
-	 * Perform a hit test on the topology using the current selection settings. In cases of hitting edges and
-	 * corners, OutHit contains the following:
-	 *   OutHit.FaceIndex: edge or corner id in the topology
-	 *   OutHit.ImpactPoint: closest point on the ray to the hit element (Note: not a point on the element!)
-	 *   OutHit.Distance: distance along the ray to ImpactPoint
-	 *   OutHit.Item: if hit item was an edge, index of the segment within the edge polyline. Otherwise undefined.
-	 *
-	 * @param bUseOrthoSettings If true, the ortho-relevant settings for selection are used (selecting down the view ray, etc)
-	 */
-	bool TopologyHitTest(const FRay& WorldRay, FHitResult& OutHit, FGroupTopologySelection& OutSelection, bool bUseOrthoSettings = false);
-	bool TopologyHitTest(const FRay& WorldRay, FHitResult& OutHit, bool bUseOrthoSettings = false);
-
-	TSharedPtr<FGroupTopologySelector, ESPMode::ThreadSafe> GetTopologySelector() { return TopoSelector; }
-
-	//
-	// Hover API
-	//
-
-	/**
-	 * Update the hover highlight based on the hit elements at the given World Ray
-	 * @return true if something was hit and is now being hovered
-	 */
-	bool UpdateHighlight(const FRay& WorldRay);
-
-	/**
-	 * Clear current hover-highlight
-	 */
-	void ClearHighlight();
-
-
-	//
-	// Selection API
-	//
-
-	/** 
-	 * Intersect the ray with the mesh and update the selection based on the hit element, modifier states, etc
-	 * @return true if selection was modified
-	 */
-	bool UpdateSelection(const FRay& WorldRay, FVector3d& LocalHitPositionOut, FVector3d& LocalHitNormalOut);
-
-	/**
-	 * Replace the current selection with an external selection. 
-	 * @warning does not check that the selection is valid!
-	 *
-	 * @param bBroadcast If true, issues an OnSelectionChanged delegate broadcast.
-	 */
-	void SetSelection(const FGroupTopologySelection& Selection, bool bBroadcast = true);
-
-	/**
-	 * Clear the current selection.
-	 */
-	void ClearSelection();
-
-	void InvertSelection();
-	void SelectAll();
-
-	/** 
-	 * @return true if the current selection is non-empty 
-	 */
-	bool HasSelection() const;
-
-	/**
-	 * @return the current selection
-	 */
-	const FGroupTopologySelection& GetActiveSelection() const { return PersistentSelection; }
+	// UMeshTopologySelectionMechanic
+	virtual void Setup(UInteractiveTool* ParentTool) override;
+	virtual bool UpdateHighlight(const FRay& WorldRay) override;
+	virtual bool UpdateSelection(const FRay& WorldRay, FVector3d& LocalHitPositionOut, FVector3d& LocalHitNormalOut) override;
 
 	/**
 	 * Gives the current selection as a storable selection object. Can optionally apply the passed-in
@@ -281,186 +59,20 @@ public:
 	 */
 	void LoadSelection(const UPersistentMeshSelection& Selection);
 
-	/**
-	 * Can be used by in an OnSelectionChanged event to inspect the clicked location (i.e., the
-	 * values returned by the UpdateSelection() function when the click happened).
-	 */
-	void GetClickedHitPosition(FVector3d& PositionOut, FVector3d& NormalOut) const;
-
-	/**
-	 * @return The best-guess 3D frame for the current select
-	 * @param bWorld if true, local-to-world transform of the target MeshComponent is applied to the frame
-	 */
-	FFrame3d GetSelectionFrame(bool bWorld, FFrame3d* InitialLocalFrame = nullptr) const;
-
-	/**
-	 * @return Bounding box for the current selection
-	 * @param bWorld if true, the box is in world space, otherwise it is in local space of the target MeshComponent
-	 */
-	FAxisAlignedBox3d GetSelectionBounds(bool bWorld) const;
-
-	void SetShowSelectableCorners(bool bShowCorners);
 
-	//
-	// Change Tracking
-	//
-
-	/**
-	 * Begin a change record. Internally creates a FCommandChange and initializes it with current state
-	 */
-	void BeginChange();
-
-	/**
-	 * End the active change and return it. Returns empty change if the selection was not modified!
-	 */
-	TUniquePtr<FToolCommandChange> EndChange();
-
-	/**
-	 * Ends the active change and emits it via the parent tool, if the selection has been modified.
-	 */
-	bool EndChangeAndEmitIfModified();
-
-	// IClickBehaviorTarget implementation
-	virtual FInputRayHit IsHitByClick(const FInputDeviceRay& ClickPos) override;
-	virtual void OnClicked(const FInputDeviceRay& ClickPos) override;
-
-	// IHoverBehaviorTarget implementation
-	virtual FInputRayHit BeginHoverSequenceHitTest(const FInputDeviceRay& PressPos) override;
-	virtual void OnBeginHover(const FInputDeviceRay& DevicePos) override;
-	virtual bool OnUpdateHover(const FInputDeviceRay& DevicePos) override;
-	virtual void OnEndHover() override;
-	virtual void OnUpdateModifierState(int ModifierID, bool bIsOn) override;
-
-protected:
-	// These get bound to marquee mechanic delegates.
-	virtual void OnDragRectangleStarted();
-	virtual void OnDragRectangleChanged(const FCameraRectangle& CurrentRectangle);
-	virtual void OnDragRectangleFinished(const FCameraRectangle& Rectangle, bool bCancelled);
-
-	virtual void UpdateMarqueeEnabled();
-
-public:
-	/** 
-	 * OnSelectionChanged is broadcast whenever the selection is modified (including by FChanges, which
-	 * means that called functions should not issue undo transactions. 
-	 */
-	FSimpleMulticastDelegate OnSelectionChanged;
-
-	/**
-	 * This is broadcast during marquee rectangle dragging if selected faces change, to allow user to
-	 * dynamically update face highlighting if desired (needed because OnSelectionChanged is only 
-	 * fired once the rectangle is completed, not while it is being updated).
-	 */
-	FSimpleMulticastDelegate OnFaceSelectionPreviewChanged;
-
-	/**
-	 * When true, the selection mechanic is currently tracking a marquee rectangle drag, and acting
-	 * on the selection may be unwise until it is over (and an OnSelectionChanged event is fired).
-	 */
-	bool IsCurrentlyMarqueeDragging() { return bCurrentlyMarqueeDragging; }
-
-	// TODO: Is it worth issuing separate callbacks in normal selection changes and in FChange ones, to
-	// allow the user to bundle in some FChanges into the normal callback?
+	UE_DEPRECATED(5.2, "Use GetTopologySelector in base class")
+	TSharedPtr<FGroupTopologySelector, ESPMode::ThreadSafe> GetTopologySelector() 
+	{ 
+		return StaticCastSharedPtr<FGroupTopologySelector>(TopoSelector); 
+	}
 
+	UE_DEPRECATED(5.2, "Use MechanicProperties in base class")
 	UPROPERTY()
-	TObjectPtr<UPolygonSelectionMechanicProperties> Properties;
+	TObjectPtr<UDEPRECATED_PolygonSelectionMechanicProperties> Properties_DEPRECATED;
 
-protected:
-	bool bIsEnabled = true;
+private:
 
-	const FDynamicMesh3* Mesh;
+	// TODO: Would be nice to get rid of this and write everything in terms of TopologySelector and TopologyProvider
 	const FGroupTopology* Topology;
-	TFunction<FDynamicMeshAABBTree3*()> GetSpatialFunc;
-
-	UPROPERTY()
-	TObjectPtr<UMouseHoverBehavior> HoverBehavior;
-
-	UPROPERTY()
-	TObjectPtr<USingleClickOrDragInputBehavior> ClickOrDragBehavior;
-
-	UPROPERTY()
-	TObjectPtr<URectangleMarqueeMechanic> MarqueeMechanic;
-
-	FInputCapturePriority BasePriority = FInputCapturePriority(FInputCapturePriority::DEFAULT_TOOL_PRIORITY);
-
-	// When bSelectEdgeLoops is true, this function is tested to see if we should select edge loops,
-	// to allow edge loop selection to be toggled with some key (setting bSelectEdgeLoops to
-	// false overrides this function).
-	TFunction<bool(void)> ShouldSelectEdgeLoopsFunc = []() {return true; };
-
-	// When bSelectEdgeRings is true, this function is tested to see if we should select edge rings,
-	// to allow edge ring selection to be toggled with some key (setting bSelectEdgeRings to
-	// false overrides this function).
-	TFunction<bool(void)> ShouldSelectEdgeRingsFunc = []() {return true; };
-
-	TFunction<bool(void)> ShouldAddToSelectionFunc = [this]() {return bShiftToggle; };
-	TFunction<bool(void)> ShouldRemoveFromSelectionFunc = [this]() {return bCtrlToggle; };;
-
-	FTransform3d TargetTransform;
-
-	TSharedPtr<FGroupTopologySelector, ESPMode::ThreadSafe> TopoSelector;
-
-	/**
-	 * Get the topology selector settings to use given the current selection settings.
-	 * 
-	 * @param bUseOrthoSettings If true, the topology selector will be configured to use ortho settings,
-	 *  which are generally different to allow for selection of projected elements, etc.
-	 */
-	FGroupTopologySelector::FSelectionSettings GetTopoSelectorSettings(bool bUseOrthoSettings = false);
-
-	FGroupTopologySelection HilightSelection;
-	FGroupTopologySelection PersistentSelection;
-	int32 SelectionTimestamp = 0;
-	TUniquePtr<FPolygonSelectionMechanicSelectionChange> ActiveChange;
 
-	// Used for box selection
-	FGroupTopologySelection PreDragPersistentSelection;
-	FGroupTopologySelection LastUpdateRectangleSelection;
-	FGroupTopologySelector::FSelectionSettings PreDragTopoSelectorSettings;
-	TMap<int32, bool> TriIsOccludedCache;
-	bool bCurrentlyMarqueeDragging = false;
-
-	FVector3d LastClickedHitPosition;
-	FVector3d LastClickedHitNormal;
-
-	/** The actor we create internally to own the DrawnTriangleSetComponent */
-	UPROPERTY()
-	TObjectPtr<APreviewGeometryActor> PreviewGeometryActor;
-
-	UPROPERTY()
-	TObjectPtr<UTriangleSetComponent> DrawnTriangleSetComponent;
-
-	TSet<int> CurrentlyHighlightedGroups;
-
-	UPROPERTY()
-	TObjectPtr<UMaterialInterface> HighlightedFaceMaterial;
-
-	FViewCameraState CameraState;
-
-	bool bShiftToggle = false;
-	bool bCtrlToggle = false;
-	static const int32 ShiftModifierID = 1;
-	static const int32 CtrlModifierID = 2;
-
-	bool bShowSelectableCorners = true;
-public:
-	FToolDataVisualizer PolyEdgesRenderer;
-	FToolDataVisualizer HilightRenderer;
-	FToolDataVisualizer SelectionRenderer;
-
-	friend class FPolygonSelectionMechanicSelectionChange;
 };
-
-
-
-class MODELINGCOMPONENTS_API FPolygonSelectionMechanicSelectionChange : public FToolCommandChange
-{
-public:
-	FGroupTopologySelection Before;
-	FGroupTopologySelection After;
-	int32 Timestamp = 0;
-
-	virtual void Apply(UObject* Object) override;
-	virtual void Revert(UObject* Object) override;
-	virtual FString ToString() const override;
-};
\ No newline at end of file
