ModelingMode: Use actor factory when creating static mesh assets

diff --git a/Engine/Plugins/Runtime/MeshModelingToolset/Source/ModelingComponents/Public/ModelingObjectsCreationAPI.h b/Engine/Plugins/Runtime/MeshModelingToolset/Source/ModelingComponents/Public/ModelingObjectsCreationAPI.h
--- a/Engine/Plugins/Runtime/MeshModelingToolset/Source/ModelingComponents/Public/ModelingObjectsCreationAPI.h
+++ b/Engine/Plugins/Runtime/MeshModelingToolset/Source/ModelingComponents/Public/ModelingObjectsCreationAPI.h
@@ -32,7 +32,8 @@ enum class ECreateModelingObjectResult : uint8
 	Failed_InvalidWorld,
 	Failed_InvalidMesh,
 	Failed_InvalidTexture,
-	Failed_AssetCreationFailed
+	Failed_AssetCreationFailed,
+    Failed_ActorCreationFailed
 };
 
 /**
diff --git a/Engine/Plugins/Runtime/MeshModelingToolset/Source/ModelingComponentsEditorOnly/Private/EditorModelingObjectsCreationAPI.cpp b/Engine/Plugins/Runtime/MeshModelingToolset/Source/ModelingComponentsEditorOnly/Private/EditorModelingObjectsCreationAPI.cpp
--- a/Engine/Plugins/Runtime/MeshModelingToolset/Source/ModelingComponentsEditorOnly/Private/EditorModelingObjectsCreationAPI.cpp
+++ b/Engine/Plugins/Runtime/MeshModelingToolset/Source/ModelingComponentsEditorOnly/Private/EditorModelingObjectsCreationAPI.cpp
@@ -27,6 +27,8 @@
 #include "DynamicMeshActor.h"
 #include "Components/DynamicMeshComponent.h"
 
+#include "ActorFactories/ActorFactory.h"
+#include "AssetSelection.h"
 
 using namespace UE::Geometry;
 
@@ -318,14 +320,23 @@ FCreateMeshObjectResult UEditorModelingObjectsCreationAPI::CreateStaticMeshAsset
 	UStaticMesh* NewStaticMesh = ResultData.StaticMesh;
 
 	// create new StaticMeshActor
-	FRotator Rotation(0.0f, 0.0f, 0.0f);
-	FActorSpawnParameters SpawnInfo;
-	// @todo nothing here is specific to AStaticMeshActor...could we pass in a CDO and clone it instead of using SpawnActor?
-	AStaticMeshActor* StaticMeshActor = CreateMeshParams.TargetWorld->SpawnActor<AStaticMeshActor>(FVector::ZeroVector, Rotation, SpawnInfo);
-	FActorLabelUtilities::SetActorLabelUnique(StaticMeshActor, CreateMeshParams.BaseName);
-	UStaticMeshComponent* StaticMeshComponent = StaticMeshActor->GetStaticMeshComponent();
+	AStaticMeshActor* const StaticMeshActor = [NewStaticMesh, &CreateMeshParams]() -> AStaticMeshActor*
+	{
+		if (UActorFactory* const StaticMeshFactory = FActorFactoryAssetProxy::GetFactoryForAssetObject(NewStaticMesh))
+		{
+			AActor* const Actor = StaticMeshFactory->CreateActor(NewStaticMesh, CreateMeshParams.TargetWorld->GetCurrentLevel(), FTransform::Identity);
+			FActorLabelUtilities::SetActorLabelUnique(Actor, CreateMeshParams.BaseName);
+			return Cast<AStaticMeshActor>(Actor);
+		}
+		return nullptr;
+	}();
+	if (!StaticMeshActor)
+	{
+		return FCreateMeshObjectResult{ECreateModelingObjectResult::Failed_ActorCreationFailed};
+	}
 
 	// set the mesh
+	UStaticMeshComponent* const StaticMeshComponent = StaticMeshActor->GetStaticMeshComponent();
 
 	// this disconnects the component from various events
 	StaticMeshComponent->UnregisterComponent();
