AddPrimitiveTool -  Disc primitive, Punctured Disc.  Don't let the inner radius exceed the outer

diff --git a/Engine/Plugins/Runtime/MeshModelingToolset/Source/MeshModelingTools/Private/AddPrimitiveTool.cpp b/Engine/Plugins/Runtime/MeshModelingToolset/Source/MeshModelingTools/Private/AddPrimitiveTool.cpp
--- a/Engine/Plugins/Runtime/MeshModelingToolset/Source/MeshModelingTools/Private/AddPrimitiveTool.cpp
+++ b/Engine/Plugins/Runtime/MeshModelingToolset/Source/MeshModelingTools/Private/AddPrimitiveTool.cpp
@@ -537,6 +537,34 @@ UAddDiscPrimitiveTool::UAddDiscPrimitiveTool(const FObjectInitializer& ObjectIni
 	UInteractiveTool::SetToolDisplayName(LOCTEXT("DiscToolName", "Create Disc"));
 }
 
+void UAddDiscPrimitiveTool::Setup()
+{
+	Super::Setup();
+
+
+	// add watchers to ensure the hole radius is never larger than the disc radius when updating a punctured disc
+
+	UProceduralDiscToolProperties* DiscSettings = Cast<UProceduralDiscToolProperties>(ShapeSettings);
+
+	int32 HoleRadiusWatchID = DiscSettings->WatchProperty(DiscSettings->HoleRadius, 
+															[DiscSettings](float HR) 
+															{
+																if (DiscSettings->DiscType == EProceduralDiscType::PuncturedDisc)
+																{
+																	DiscSettings->HoleRadius = FMath::Min(DiscSettings->Radius, HR);
+																}
+															});
+	DiscSettings->WatchProperty(DiscSettings->Radius, 
+								[DiscSettings, HoleRadiusWatchID](float R) 
+								{
+									if (DiscSettings->DiscType == EProceduralDiscType::PuncturedDisc)
+									{
+										DiscSettings->HoleRadius = FMath::Min(DiscSettings->HoleRadius, R);
+										DiscSettings->SilentUpdateWatcherAtIndex(HoleRadiusWatchID);
+									}
+								});
+}
+
 void UAddDiscPrimitiveTool::GenerateMesh(FDynamicMesh3* OutMesh) const
 {
 	const UProceduralDiscToolProperties* DiscSettings = Cast<UProceduralDiscToolProperties>(ShapeSettings);
@@ -569,6 +597,7 @@ void UAddDiscPrimitiveTool::GenerateMesh(FDynamicMesh3* OutMesh) const
 }
 
 
+
 UAddTorusPrimitiveTool::UAddTorusPrimitiveTool(const FObjectInitializer& ObjectInitializer)
 	: Super(ObjectInitializer.SetDefaultSubobjectClass<UProceduralTorusToolProperties>(TEXT("ShapeSettings")))
 {
diff --git a/Engine/Plugins/Runtime/MeshModelingToolset/Source/MeshModelingTools/Public/AddPrimitiveTool.h b/Engine/Plugins/Runtime/MeshModelingToolset/Source/MeshModelingTools/Public/AddPrimitiveTool.h
--- a/Engine/Plugins/Runtime/MeshModelingToolset/Source/MeshModelingTools/Public/AddPrimitiveTool.h
+++ b/Engine/Plugins/Runtime/MeshModelingToolset/Source/MeshModelingTools/Public/AddPrimitiveTool.h
@@ -581,6 +581,7 @@ class UAddDiscPrimitiveTool : public UAddPrimitiveTool
 	GENERATED_BODY()
 public:
 	explicit UAddDiscPrimitiveTool(const FObjectInitializer& ObjectInitializer);
+	virtual void Setup() override;
 protected:
 	virtual void GenerateMesh(FDynamicMesh3* OutMesh) const override;
 };
