Fix VoxBlnd, VoxMrph, and VoxWrap not showing preview at LWC

diff --git a/Engine/Plugins/Runtime/MeshModelingToolset/Source/ModelingOperators/Private/CompositionOps/VoxelBlendMeshesOp.cpp b/Engine/Plugins/Runtime/MeshModelingToolset/Source/ModelingOperators/Private/CompositionOps/VoxelBlendMeshesOp.cpp
--- a/Engine/Plugins/Runtime/MeshModelingToolset/Source/ModelingOperators/Private/CompositionOps/VoxelBlendMeshesOp.cpp
+++ b/Engine/Plugins/Runtime/MeshModelingToolset/Source/ModelingOperators/Private/CompositionOps/VoxelBlendMeshesOp.cpp
@@ -34,6 +34,9 @@ void FVoxelBlendMeshesOp::CalculateResult(FProgressCancel* Progress)
 	TImplicitBlend<FDynamicMesh3> ImplicitBlend;
 	ImplicitBlend.bSubtract = bSubtract;
 
+	FVector3d AverageTranslation = GetAverageTranslation<FTransform>(Transforms);
+	ResultTransform = FTransformSRT3d(AverageTranslation);
+
 	TArray<FDynamicMesh3> TransformedMeshes; TransformedMeshes.Reserve(Meshes.Num());
 	FAxisAlignedBox3d CombinedBounds = FAxisAlignedBox3d::Empty();
 	for (int MeshIdx = 0; MeshIdx < Meshes.Num(); MeshIdx++)
@@ -44,7 +47,9 @@ void FVoxelBlendMeshesOp::CalculateResult(FProgressCancel* Progress)
 			continue;
 		}
 		
-		MeshTransforms::ApplyTransform(TransformedMeshes[MeshIdx], (FTransformSRT3d)Transforms[MeshIdx], true);
+		FTransformSRT3d ToApply = Transforms[MeshIdx];
+		ToApply.SetTranslation(ToApply.GetTranslation() - AverageTranslation);
+		MeshTransforms::ApplyTransform(TransformedMeshes[MeshIdx], ToApply, true);
 
 		if (bVoxWrap)
 		{
diff --git a/Engine/Plugins/Runtime/MeshModelingToolset/Source/ModelingOperators/Private/CompositionOps/VoxelMorphologyMeshesOp.cpp b/Engine/Plugins/Runtime/MeshModelingToolset/Source/ModelingOperators/Private/CompositionOps/VoxelMorphologyMeshesOp.cpp
--- a/Engine/Plugins/Runtime/MeshModelingToolset/Source/ModelingOperators/Private/CompositionOps/VoxelMorphologyMeshesOp.cpp
+++ b/Engine/Plugins/Runtime/MeshModelingToolset/Source/ModelingOperators/Private/CompositionOps/VoxelMorphologyMeshesOp.cpp
@@ -55,6 +55,8 @@ void FVoxelMorphologyMeshesOp::CalculateResult(FProgressCancel* Progress)
 	}
 
 	FDynamicMesh3 CombinedMesh;
+	FVector3d AverageTranslation = GetAverageTranslation<FTransformSRT3d>(Transforms);
+	ResultTransform = FTransformSRT3d(AverageTranslation);
 
 	// append all meshes (transformed but without attributes)
 	FDynamicMeshEditor AppendEditor(&CombinedMesh);
@@ -64,9 +66,9 @@ void FVoxelMorphologyMeshesOp::CalculateResult(FProgressCancel* Progress)
 		bool bReverseOrientation = MeshTransform.GetDeterminant() < 0;
 		FMeshIndexMappings IndexMaps;
 		AppendEditor.AppendMesh(Meshes[MeshIdx].Get(), IndexMaps,
-			[MeshTransform](int VID, const FVector3d& Pos)
+			[MeshTransform, AverageTranslation](int VID, const FVector3d& Pos)
 			{
-				return MeshTransform.TransformPosition(Pos);
+				return MeshTransform.TransformPosition(Pos) - AverageTranslation;
 			}, nullptr
 		);
 		if (bReverseOrientation)
diff --git a/Engine/Plugins/Runtime/MeshModelingToolset/Source/ModelingOperators/Private/CompositionOps/VoxelSolidifyMeshesOp.cpp b/Engine/Plugins/Runtime/MeshModelingToolset/Source/ModelingOperators/Private/CompositionOps/VoxelSolidifyMeshesOp.cpp
--- a/Engine/Plugins/Runtime/MeshModelingToolset/Source/ModelingOperators/Private/CompositionOps/VoxelSolidifyMeshesOp.cpp
+++ b/Engine/Plugins/Runtime/MeshModelingToolset/Source/ModelingOperators/Private/CompositionOps/VoxelSolidifyMeshesOp.cpp
@@ -32,6 +32,9 @@ void FVoxelSolidifyMeshesOp::CalculateResult(FProgressCancel* Progress)
 
 	FDynamicMesh3 CombinedMesh;
 
+	FVector3d AverageTranslation = GetAverageTranslation<FTransformSRT3d>(Transforms);
+	ResultTransform = FTransformSRT3d(AverageTranslation);
+
 	// append all meshes (transformed but without attributes)
 	FDynamicMeshEditor AppendEditor(&CombinedMesh);
 	for (int MeshIdx = 0; MeshIdx < Meshes.Num(); MeshIdx++)
@@ -40,9 +43,9 @@ void FVoxelSolidifyMeshesOp::CalculateResult(FProgressCancel* Progress)
 		bool bReverseOrientation = MeshTransform.GetDeterminant() < 0;
 		FMeshIndexMappings IndexMaps;
 		AppendEditor.AppendMesh(Meshes[MeshIdx].Get(), IndexMaps,
-			[MeshTransform](int VID, const FVector3d& Pos)
+			[MeshTransform, AverageTranslation](int VID, const FVector3d& Pos)
 			{
-				return MeshTransform.TransformPosition(Pos);
+				return MeshTransform.TransformPosition(Pos) - AverageTranslation;
 			}, nullptr
 		);
 		if (bReverseOrientation)
diff --git a/Engine/Plugins/Runtime/MeshModelingToolset/Source/ModelingOperators/Public/BaseOps/VoxelBaseOp.h b/Engine/Plugins/Runtime/MeshModelingToolset/Source/ModelingOperators/Public/BaseOps/VoxelBaseOp.h
--- a/Engine/Plugins/Runtime/MeshModelingToolset/Source/ModelingOperators/Public/BaseOps/VoxelBaseOp.h
+++ b/Engine/Plugins/Runtime/MeshModelingToolset/Source/ModelingOperators/Public/BaseOps/VoxelBaseOp.h
@@ -27,6 +27,22 @@ public:
 
 	bool bRemoveInternalSurfaces = false;
 
+	template<typename TransformType>
+	static FVector3d GetAverageTranslation(TArrayView<const TransformType> Transforms)
+	{
+		if (Transforms.IsEmpty())
+		{
+			return FVector3d::ZeroVector;
+		}
+
+		FVector3d Avg(0, 0, 0);
+		for (const FTransformSRT3d& Transform : Transforms)
+		{
+			Avg += Transform.GetTranslation();
+		}
+		return Avg / double(Transforms.Num());
+	}
+
 	virtual void PostProcessResult(FProgressCancel* Progress, double MeshCellSize);
 };
 
