fix polyedit extrude distance not having a max

diff --git a/Engine/Plugins/Runtime/MeshModelingToolset/Source/MeshModelingTools/Public/ToolActivities/PolyEditExtrudeActivity.h b/Engine/Plugins/Runtime/MeshModelingToolset/Source/MeshModelingTools/Public/ToolActivities/PolyEditExtrudeActivity.h
--- a/Engine/Plugins/Runtime/MeshModelingToolset/Source/MeshModelingTools/Public/ToolActivities/PolyEditExtrudeActivity.h
+++ b/Engine/Plugins/Runtime/MeshModelingToolset/Source/MeshModelingTools/Public/ToolActivities/PolyEditExtrudeActivity.h
@@ -116,8 +116,8 @@ public:
 	EPolyEditExtrudeDistanceMode DistanceMode = EPolyEditExtrudeDistanceMode::ClickInViewport;
 
 	/** Distance to extrude. */
-	UPROPERTY(EditAnywhere, Category = Extrude,
-		meta = (EditConditionHides, EditCondition = "DistanceMode == EPolyEditExtrudeDistanceMode::Fixed"))
+	UPROPERTY(EditAnywhere, Category = Extrude, meta = (UIMin = "-1000", UIMax = "1000", ClampMin = "-10000", ClampMax = "10000",
+		EditConditionHides, EditCondition = "DistanceMode == EPolyEditExtrudeDistanceMode::Fixed"))
 	double Distance = 100;
 
 	/** Direction in which to extrude. */
