Modeling Tools.  Mesh Booleans and CutMeshesWithMeshes. When filling holes, generate normals and UVs for the new patches

diff --git a/Engine/Plugins/Runtime/MeshModelingToolset/Source/ModelingOperators/Private/CompositionOps/BooleanMeshesOp.cpp b/Engine/Plugins/Runtime/MeshModelingToolset/Source/ModelingOperators/Private/CompositionOps/BooleanMeshesOp.cpp
--- a/Engine/Plugins/Runtime/MeshModelingToolset/Source/ModelingOperators/Private/CompositionOps/BooleanMeshesOp.cpp
+++ b/Engine/Plugins/Runtime/MeshModelingToolset/Source/ModelingOperators/Private/CompositionOps/BooleanMeshesOp.cpp
@@ -5,6 +5,7 @@
 #include "Operations/MeshBoolean.h"
 #include "MeshBoundaryLoops.h"
 #include "Operations/MinimalHoleFiller.h"
+#include "CompGeom/PolygonTriangulation.h"
 
 using namespace UE::Geometry;
 
@@ -101,10 +102,39 @@ void FBooleanMeshesOp::CalculateResult(FProgressCancel* Progress)
 			return;
 		}
 
+		const bool bNeedsNormalsAndUVs = MeshBoolean.Result->HasAttributes();
+		// UV scale based on the mesh bounds, used for hole-filled regions.
+		const double MeshUVScaleFactor = (1.0 / MeshBoolean.Result->GetBounds().MaxDim());
 		for (FEdgeLoop& Loop : OpenBoundary.Loops)
 		{
+			const int32 NewGroupID = MeshBoolean.Result->AllocateTriangleGroup();
+			
 			FMinimalHoleFiller Filler(MeshBoolean.Result, Loop);
-			Filler.Fill();
+			const bool bFilled = Filler.Fill(NewGroupID);
+
+			if (!bFilled)
+			{
+				continue;
+			}
+			
+			if (bNeedsNormalsAndUVs)  // Compute normals and UVs
+			{
+				// Compute a best-fit plane of the boundary vertices
+				FVector3d PlaneOrigin, PlaneNormal;
+				{
+					TArray<FVector3d> VertexPositions;
+					Loop.GetVertices(VertexPositions);
+
+					PolygonTriangulation::ComputePolygonPlane<double>(VertexPositions, PlaneNormal, PlaneOrigin);
+					PlaneNormal *= -1.0;	// Previous function seems to orient the normal opposite to what's expected elsewhere
+				}
+
+
+				FDynamicMeshEditor Editor(ResultMesh.Get());
+				FFrame3d ProjectionFrame(PlaneOrigin, PlaneNormal);
+				Editor.SetTriangleNormals(Filler.NewTriangles, (FVector3f)PlaneNormal);
+				Editor.SetTriangleUVsFromProjection(Filler.NewTriangles, ProjectionFrame, MeshUVScaleFactor);
+			}
 		}
 
 		TArray<int32> UpdatedBoundaryEdges;
