UV layout previews in the 3D scene should not cast shadows

diff --git a/Engine/Plugins/Runtime/MeshModelingToolset/Source/ModelingComponents/Private/Drawing/UVLayoutPreview.cpp b/Engine/Plugins/Runtime/MeshModelingToolset/Source/ModelingComponents/Private/Drawing/UVLayoutPreview.cpp
--- a/Engine/Plugins/Runtime/MeshModelingToolset/Source/ModelingComponents/Private/Drawing/UVLayoutPreview.cpp
+++ b/Engine/Plugins/Runtime/MeshModelingToolset/Source/ModelingComponents/Private/Drawing/UVLayoutPreview.cpp
@@ -17,6 +17,7 @@ void UUVLayoutPreview::CreateInWorld(UWorld* World)
 {
 	PreviewMesh = NewObject<UPreviewMesh>(this);
 	PreviewMesh->CreateInWorld(World, FTransform::Identity);
+	PreviewMesh->SetShadowsEnabled(false);
 
 	TriangleComponent = NewObject<UTriangleSetComponent>(PreviewMesh->GetActor());
 	TriangleComponent->SetupAttachment(PreviewMesh->GetRootComponent());
