ModelingComponents: Modify AssetUtils::ReadTexture()'s PlatformData path to catch and error out if being used in non-Editor builds and the texture format is not TC_VectorDisplacementMap. We cannot change texture format at runtime, because UpdateResource() depends on the source data being available. This is old code that never worked and would just crash outside the Editor.

Add messaging in geometry script functions that call ReadTexture, to hint to the user how to configure their textures to use at runtime

diff --git a/Engine/Plugins/Experimental/GeometryScripting/Source/GeometryScriptingCore/Private/MeshDeformFunctions.cpp b/Engine/Plugins/Experimental/GeometryScripting/Source/GeometryScriptingCore/Private/MeshDeformFunctions.cpp
--- a/Engine/Plugins/Experimental/GeometryScripting/Source/GeometryScriptingCore/Private/MeshDeformFunctions.cpp
+++ b/Engine/Plugins/Experimental/GeometryScripting/Source/GeometryScriptingCore/Private/MeshDeformFunctions.cpp
@@ -417,7 +417,7 @@ UDynamicMesh* UGeometryScriptLibrary_MeshDeformFunctions::ApplyDisplaceFromTextu
 	TImageBuilder<FVector4f> ImageData;
 	if (UE::AssetUtils::ReadTexture(TextureAsset, ImageData, false) == false)
 	{
-		UE::Geometry::AppendError(Debug, EGeometryScriptErrorType::OperationFailed, LOCTEXT("ApplyDisplaceFromTextureMap_TexReadFailed", "ApplyDisplaceFromTextureMap: Error reading source texture data"));
+		UE::Geometry::AppendError(Debug, EGeometryScriptErrorType::OperationFailed, LOCTEXT("ApplyDisplaceFromTextureMap_TexReadFailed", "ApplyDisplaceFromTextureMap: Error reading source texture data. If using this function at Runtime, The Compression Settings type on the UTexture2D Asset must be set to VectorDisplacementmap (RGBA8)."));
 		return TargetMesh;
 	}
 	int UseImageChannel = FMath::Clamp(Options.ImageChannel, 0, 3);
diff --git a/Engine/Plugins/Experimental/GeometryScripting/Source/GeometryScriptingCore/Private/TextureMapFunctions.cpp b/Engine/Plugins/Experimental/GeometryScripting/Source/GeometryScriptingCore/Private/TextureMapFunctions.cpp
--- a/Engine/Plugins/Experimental/GeometryScripting/Source/GeometryScriptingCore/Private/TextureMapFunctions.cpp
+++ b/Engine/Plugins/Experimental/GeometryScripting/Source/GeometryScriptingCore/Private/TextureMapFunctions.cpp
@@ -26,7 +26,7 @@ void UGeometryScriptLibrary_TextureMapFunctions::SampleTexture2DAtUVPositions(
 	TImageBuilder<FVector4f> ImageData;
 	if (UE::AssetUtils::ReadTexture(TextureAsset, ImageData, false) == false)
 	{
-		UE::Geometry::AppendError(Debug, EGeometryScriptErrorType::OperationFailed, LOCTEXT("SampleTexture2DAtUVPositions_TexReadFailed", "SampleTexture2DAtUVPositions: Error reading source texture data"));
+		UE::Geometry::AppendError(Debug, EGeometryScriptErrorType::OperationFailed, LOCTEXT("SampleTexture2DAtUVPositions_TexReadFailed", "SampleTexture2DAtUVPositions: Error reading source texture data. If using this function at Runtime, The Compression Settings type on the UTexture2D Asset must be set to VectorDisplacementmap (RGBA8)."));
 		return;
 	}
 
diff --git a/Engine/Plugins/Runtime/MeshModelingToolset/Source/ModelingComponents/Private/AssetUtils/Texture2DUtil.cpp b/Engine/Plugins/Runtime/MeshModelingToolset/Source/ModelingComponents/Private/AssetUtils/Texture2DUtil.cpp
--- a/Engine/Plugins/Runtime/MeshModelingToolset/Source/ModelingComponents/Private/AssetUtils/Texture2DUtil.cpp
+++ b/Engine/Plugins/Runtime/MeshModelingToolset/Source/ModelingComponents/Private/AssetUtils/Texture2DUtil.cpp
@@ -22,7 +22,11 @@ static bool ReadTexture_PlatformData(
 	// Note that the current code cannot run on a background thread, UpdateResource() will call FlushRenderingCommands()
 	// which will check() if it's on the Game Thread
 
-	check(TextureMap->GetPlatformData());
+	if (TextureMap->GetPlatformData() == nullptr)
+	{
+		return false;
+	}
+
 	const int32 Width = TextureMap->GetPlatformData()->Mips[0].SizeX;
 	const int32 Height = TextureMap->GetPlatformData()->Mips[0].SizeY;
 	const FImageDimensions Dimensions = FImageDimensions(Width, Height);
@@ -31,16 +35,29 @@ static bool ReadTexture_PlatformData(
 
 	// convert built platform texture data to uncompressed RGBA8 format
 	const TextureCompressionSettings InitialCompressionSettings = TextureMap->CompressionSettings;
+
+	// the Platform BulkData for most texture formats is compressed and we cannot process it on
+	// the CPU. However TC_VectorDisplacementmap is uncompressed RGBA8.
+	// in the Editor, we can temporarily change the texture type and rebuild the PlatformData.
+	// However at Runtime we cannot, and so we can only error-out if we do not have a readable compression type.
 #if WITH_EDITOR
 	const TextureMipGenSettings InitialMipGenSettings = TextureMap->MipGenSettings;
+	bool bNeedToRevertTextureChanges = false;
+	if (InitialMipGenSettings != TextureMipGenSettings::TMGS_NoMipmaps || InitialCompressionSettings != TextureCompressionSettings::TC_VectorDisplacementmap)
+	{
+		TextureMap->CompressionSettings = TextureCompressionSettings::TC_VectorDisplacementmap;
+		TextureMap->MipGenSettings = TextureMipGenSettings::TMGS_NoMipmaps;
+		TextureMap->UpdateResource();
+		bNeedToRevertTextureChanges = true;
+	}
+#else
+	if (InitialCompressionSettings != TC_VectorDisplacementmap)
+	{
+		return false;
+	}
 #endif
-	TextureMap->CompressionSettings = TextureCompressionSettings::TC_VectorDisplacementmap;
-#if WITH_EDITOR
-	TextureMap->MipGenSettings = TextureMipGenSettings::TMGS_NoMipmaps;
-#endif
-	TextureMap->UpdateResource();
 
-	// lock texture and read as FColor
+	// lock texture and read as FColor (RGBA8)
 	const FColor* FormattedImageData = reinterpret_cast<const FColor*>(TextureMap->GetPlatformData()->Mips[0].BulkData.LockReadOnly());
 
 	// maybe could be done more quickly by row?
@@ -53,13 +70,17 @@ static bool ReadTexture_PlatformData(
 		DestImage.SetPixel(i, ToVector4<float>(FloatColor));
 	}
 
-	// restore built platform texture data to initial state
 	TextureMap->GetPlatformData()->Mips[0].BulkData.Unlock();
-	TextureMap->CompressionSettings = InitialCompressionSettings;
+
 #if WITH_EDITOR
-	TextureMap->MipGenSettings = InitialMipGenSettings;
+	// restore built platform texture data to initial state, if we modified it
+	if (bNeedToRevertTextureChanges)
+	{
+		TextureMap->CompressionSettings = InitialCompressionSettings;
+		TextureMap->MipGenSettings = InitialMipGenSettings;
+		TextureMap->UpdateResource();
+	}
 #endif
-	TextureMap->UpdateResource();
 
 	return true;
 }
