fix lattice dimensions becoming editable on control point movement redo

diff --git a/Engine/Plugins/Runtime/MeshModelingToolset/Source/ModelingComponents/Private/Mechanics/LatticeControlPointsMechanic.cpp b/Engine/Plugins/Runtime/MeshModelingToolset/Source/ModelingComponents/Private/Mechanics/LatticeControlPointsMechanic.cpp
--- a/Engine/Plugins/Runtime/MeshModelingToolset/Source/ModelingComponents/Private/Mechanics/LatticeControlPointsMechanic.cpp
+++ b/Engine/Plugins/Runtime/MeshModelingToolset/Source/ModelingComponents/Private/Mechanics/LatticeControlPointsMechanic.cpp
@@ -801,7 +801,7 @@ void FLatticeControlPointsMechanicMovementChange::Apply(UObject* Object)
 {
 	ULatticeControlPointsMechanic* Mechanic = Cast<ULatticeControlPointsMechanic>(Object);
 	Mechanic->UpdatePointLocations(NewPositions);
-	Mechanic->bHasChanged = false;
+	Mechanic->bHasChanged = true;
 	Mechanic->OnPointsChanged.Broadcast();
 }
 
