Add missing LOCTEXT_NAMESPACE

diff --git a/Engine/Plugins/Runtime/MeshModelingToolset/Source/ModelingComponents/Private/Selection/GroupTopologySelector.cpp b/Engine/Plugins/Runtime/MeshModelingToolset/Source/ModelingComponents/Private/Selection/GroupTopologySelector.cpp
--- a/Engine/Plugins/Runtime/MeshModelingToolset/Source/ModelingComponents/Private/Selection/GroupTopologySelector.cpp
+++ b/Engine/Plugins/Runtime/MeshModelingToolset/Source/ModelingComponents/Private/Selection/GroupTopologySelector.cpp
@@ -5,6 +5,8 @@
 #include "ToolDataVisualizer.h"
 #include "ToolSceneQueriesUtil.h"
 
+#define LOCTEXT_NAMESPACE "FGroupTopologySelector"
+
 void FGroupTopologyUtils::AddNewEdgeLoopEdgesFromCorner(int32 EdgeID, int32 CornerID, TSet<int32>& EdgeSet) const
 {
 	int32 LastCornerID = CornerID;
