ModelingMode: Adjust the UV Project status bar message to include information about Shift+Ctrl+click for repositioning without reorienting.

diff --git a/Engine/Plugins/Runtime/MeshModelingToolset/Source/MeshModelingTools/Private/UVProjectionTool.cpp b/Engine/Plugins/Runtime/MeshModelingToolset/Source/MeshModelingTools/Private/UVProjectionTool.cpp
--- a/Engine/Plugins/Runtime/MeshModelingToolset/Source/MeshModelingTools/Private/UVProjectionTool.cpp
+++ b/Engine/Plugins/Runtime/MeshModelingToolset/Source/MeshModelingTools/Private/UVProjectionTool.cpp
@@ -143,7 +143,7 @@ void UUVProjectionTool::Setup()
 
 	SetToolDisplayName(LOCTEXT("ToolName", "UV Projection"));
 	GetToolManager()->DisplayMessage(
-		LOCTEXT("UVProjectionToolDescription", "Generate UVs for a Mesh by projecting onto simple geometric shapes. Ctrl+click to reposition shape. Face selections can be made in the PolyEdit and TriEdit Tools."),
+		LOCTEXT("UVProjectionToolDescription", "Generate UVs for a Mesh by projecting onto simple geometric shapes. Ctrl+click to reposition shape. Shift+Ctrl+click to reposition shape without reorienting. Face selections can be made in the PolyEdit and TriEdit Tools. "),
 		EToolMessageLevel::UserNotification);
 }
 
