Fixed a bug where the trim tool starts a potentially expensive boolean calculation when only the visual properties change like the opacity color and transparency, show/hide trimming mesh.

#rnx

diff --git a/Engine/Plugins/Runtime/MeshModelingToolset/Source/MeshModelingTools/Private/CSGMeshesTool.cpp b/Engine/Plugins/Runtime/MeshModelingToolset/Source/MeshModelingTools/Private/CSGMeshesTool.cpp
--- a/Engine/Plugins/Runtime/MeshModelingToolset/Source/MeshModelingTools/Private/CSGMeshesTool.cpp
+++ b/Engine/Plugins/Runtime/MeshModelingToolset/Source/MeshModelingTools/Private/CSGMeshesTool.cpp
@@ -344,7 +344,35 @@ void UCSGMeshesTool::OnPropertyModified(UObject* PropertySet, FProperty* Propert
 	}
 	else
 	{
-		Super::OnPropertyModified(PropertySet, Property);
+		//TODO: UBaseCreateFromSelectedTool::OnPropertyModified below calls Preview->InvalidateResult() which triggers
+		//expensive recompute of the boolean operator. We should rethink which code is responsible for invalidating the 
+		//preview and make it consistent since some of the code calls Preview->InvalidateResult() manually.
+		bool bVisualUpdate = false;
+		if (bTrimMode)
+		{
+			bVisualUpdate = Property && (Property->GetFName() == GET_MEMBER_NAME_CHECKED(UTrimMeshesToolProperties, bShowTrimmingMesh) ||
+										 Property->GetFName() == GET_MEMBER_NAME_CHECKED(UTrimMeshesToolProperties, OpacityOfTrimmingMesh) ||
+										 Property->GetFName() == GET_MEMBER_NAME_CHECKED(UTrimMeshesToolProperties, ColorOfTrimmingMesh) ||
+										 Property->GetFName() == FName("R") ||
+										 Property->GetFName() == FName("G") ||
+										 Property->GetFName() == FName("B"));
+		}
+		else 
+		{	
+			bVisualUpdate = Property && (Property->GetFName() == GET_MEMBER_NAME_CHECKED(UCSGMeshesToolProperties, bShowSubtractedMesh) ||
+										 Property->GetFName() == GET_MEMBER_NAME_CHECKED(UCSGMeshesToolProperties, SubtractedMeshColor) ||
+										 Property->GetFName() == GET_MEMBER_NAME_CHECKED(UCSGMeshesToolProperties, SubtractedMeshOpacity) ||
+										 Property->GetFName() == FName("R") ||
+										 Property->GetFName() == FName("G") ||
+										 Property->GetFName() == FName("B"));
+		}
+
+		// If the property that was changed only affects the visuals, we do not need to recalculate the output 
+		// of FBooleanMeshesOp 
+		if (bVisualUpdate == false)
+		{
+			Super::OnPropertyModified(PropertySet, Property);
+		}
 	}
 }
 
