make PathRev tool warn the user if the draw plane / axis is very far from the camera position, with an explanation of how to place the draw plane in the view.

diff --git a/Engine/Plugins/Runtime/MeshModelingToolset/Source/MeshModelingTools/Private/DrawAndRevolveTool.cpp b/Engine/Plugins/Runtime/MeshModelingToolset/Source/MeshModelingTools/Private/DrawAndRevolveTool.cpp
--- a/Engine/Plugins/Runtime/MeshModelingToolset/Source/MeshModelingTools/Private/DrawAndRevolveTool.cpp
+++ b/Engine/Plugins/Runtime/MeshModelingToolset/Source/MeshModelingTools/Private/DrawAndRevolveTool.cpp
@@ -168,6 +168,16 @@ void UDrawAndRevolveTool::Setup()
 
 	UpdateRevolutionAxis();
 
+	FViewCameraState InitialCameraState;
+	GetToolManager()->GetContextQueriesAPI()->GetCurrentViewState(InitialCameraState);
+	if (FVector::DistSquared(InitialCameraState.Position, Settings->DrawPlaneOrigin) > FarDrawPlaneThreshold * FarDrawPlaneThreshold)
+	{
+		bHasFarPlaneWarning = true;
+		GetToolManager()->DisplayMessage(
+			LOCTEXT("FarDrawPlane", "The axis of revolution is far from the camera. Note that you can ctrl-click to place the axis on a visible surface."),
+			EToolMessageLevel::UserWarning);
+	}
+
 	// The plane mechanic lets us update the plane in which we draw the profile curve, as long as we haven't
 	// started adding points to it already.
 	FFrame3d ProfileDrawPlane(Settings->DrawPlaneOrigin, Settings->DrawPlaneOrientation.Quaternion());
@@ -188,6 +198,11 @@ void UDrawAndRevolveTool::Setup()
 			ControlPointsMechanic->SetPlane(PlaneMechanic->Plane);
 		}
 		UpdateRevolutionAxis();
+		if (bHasFarPlaneWarning) // if the user has changed the plane, no longer need a warning
+		{
+			bHasFarPlaneWarning = false;
+			GetToolManager()->DisplayMessage(FText(), EToolMessageLevel::UserWarning);
+		}
 		});
 
 	ControlPointsMechanic->SetPlane(PlaneMechanic->Plane);
@@ -330,6 +345,13 @@ void UDrawAndRevolveTool::Render(IToolsContextRenderAPI* RenderAPI)
 {
 	GetToolManager()->GetContextQueriesAPI()->GetCurrentViewState(CameraState);
 
+	if (bHasFarPlaneWarning && FVector::DistSquared(CameraState.Position, Settings->DrawPlaneOrigin) < FarDrawPlaneThreshold * FarDrawPlaneThreshold)
+	{
+		// if camera is now closer to the axis, no longer need a warning
+		bHasFarPlaneWarning = false;
+		GetToolManager()->DisplayMessage(FText(), EToolMessageLevel::UserWarning);
+	}
+
 	if (PlaneMechanic != nullptr)
 	{
 		PlaneMechanic->Render(RenderAPI);
diff --git a/Engine/Plugins/Runtime/MeshModelingToolset/Source/MeshModelingTools/Public/DrawAndRevolveTool.h b/Engine/Plugins/Runtime/MeshModelingToolset/Source/MeshModelingTools/Public/DrawAndRevolveTool.h
--- a/Engine/Plugins/Runtime/MeshModelingToolset/Source/MeshModelingTools/Public/DrawAndRevolveTool.h
+++ b/Engine/Plugins/Runtime/MeshModelingToolset/Source/MeshModelingTools/Public/DrawAndRevolveTool.h
@@ -147,4 +147,8 @@ protected:
 	void GenerateAsset(const FDynamicMeshOpResult& Result);
 
 	friend class URevolveOperatorFactory;
+
+private:
+	constexpr static double FarDrawPlaneThreshold = 100 * 100;
+	bool bHasFarPlaneWarning = false;
 };
