Fix missing transform back to world space in single-selection case for BaseCreateFromSelected tools (Merge, VoxWrap, VoxMorph)

diff --git a/Engine/Plugins/Runtime/MeshModelingToolset/Source/ModelingComponents/Private/BaseTools/BaseCreateFromSelectedTool.cpp b/Engine/Plugins/Runtime/MeshModelingToolset/Source/ModelingComponents/Private/BaseTools/BaseCreateFromSelectedTool.cpp
--- a/Engine/Plugins/Runtime/MeshModelingToolset/Source/ModelingComponents/Private/BaseTools/BaseCreateFromSelectedTool.cpp
+++ b/Engine/Plugins/Runtime/MeshModelingToolset/Source/ModelingComponents/Private/BaseTools/BaseCreateFromSelectedTool.cpp
@@ -190,6 +190,7 @@ void UBaseCreateFromSelectedTool::GenerateAsset(const FDynamicMeshOpResult& OpRe
 	if (Targets.Num() == 1) // in the single-selection case, shove the result back into the original component space
 	{
 		FTransform3d FromSourceComponentSpace = UE::ToolTarget::GetLocalToWorldTransform(Targets[0]);
+		MeshTransforms::ApplyTransform(*OpResult.Mesh, OpResult.Transform, true);
 		MeshTransforms::ApplyTransformInverse(*OpResult.Mesh, FromSourceComponentSpace, true);
 		NewTransform = UE::ToolTarget::GetLocalToWorldTransform(Targets[0]);
 	}
