ModelingTools: Merge and unify the AutoUV codebases between the UV Editor and Modeling Mode.

diff --git a/Engine/Plugins/Experimental/UVEditor/Source/UVEditorToolsEditorOnly/Private/Operators/UVEditorParameterizeMeshOp.cpp b/Engine/Plugins/Experimental/UVEditor/Source/UVEditorToolsEditorOnly/Private/Operators/UVEditorParameterizeMeshOp.cpp
deleted file mode 100644
diff --git a/Engine/Plugins/Experimental/UVEditor/Source/UVEditorToolsEditorOnly/Private/UVEditorParameterizeMeshTool.cpp b/Engine/Plugins/Experimental/UVEditor/Source/UVEditorToolsEditorOnly/Private/UVEditorParameterizeMeshTool.cpp
--- a/Engine/Plugins/Experimental/UVEditor/Source/UVEditorToolsEditorOnly/Private/UVEditorParameterizeMeshTool.cpp
+++ b/Engine/Plugins/Experimental/UVEditor/Source/UVEditorToolsEditorOnly/Private/UVEditorParameterizeMeshTool.cpp
@@ -4,7 +4,8 @@
 #include "InteractiveToolManager.h"
 #include "ToolBuilderUtil.h"
 #include "DynamicMesh/DynamicMesh3.h"
-#include "Operators/UVEditorParameterizeMeshOp.h"
+#include "ParameterizationOps/ParameterizeMeshOp.h"
+#include "Properties/ParameterizeMeshProperties.h"
 #include "PropertySets/PolygroupLayersProperties.h"
 #include "Polygroups/PolygroupUtil.h"
 #include "ToolTargets/UVEditorToolMeshInput.h"
@@ -53,22 +54,22 @@ void UUVEditorParameterizeMeshTool::Setup()
 	UInteractiveTool::Setup();
 
 	// initialize our properties
-	Settings = NewObject<UUVEditorParameterizeMeshToolProperties>(this);
+	Settings = NewObject<UParameterizeMeshToolProperties>(this);
 	Settings->RestoreProperties(this);
 	AddToolPropertySource(Settings);
-	Settings->WatchProperty(Settings->Method, [&](EUVEditorParameterizeMeshUVMethod) { OnMethodTypeChanged(); });
+	Settings->WatchProperty(Settings->Method, [&](EParameterizeMeshUVMethod) { OnMethodTypeChanged(); });
 
-	UVAtlasProperties = NewObject<UUVEditorParameterizeMeshToolUVAtlasProperties>(this);
+	UVAtlasProperties = NewObject<UParameterizeMeshToolUVAtlasProperties>(this);
 	UVAtlasProperties->RestoreProperties(this);
 	AddToolPropertySource(UVAtlasProperties);
 	SetToolPropertySourceEnabled(UVAtlasProperties, true);
 
-	XAtlasProperties = NewObject<UUVEditorParameterizeMeshToolXAtlasProperties>(this);
+	XAtlasProperties = NewObject<UParameterizeMeshToolXAtlasProperties>(this);
 	XAtlasProperties->RestoreProperties(this);
 	AddToolPropertySource(XAtlasProperties);
 	SetToolPropertySourceEnabled(XAtlasProperties, true);
 	
-	PatchBuilderProperties = NewObject<UUVEditorParameterizeMeshToolPatchBuilderProperties>(this);
+	PatchBuilderProperties = NewObject<UParameterizeMeshToolPatchBuilderProperties>(this);
 	PatchBuilderProperties->RestoreProperties(this);
 	AddToolPropertySource(PatchBuilderProperties);
 	SetToolPropertySourceEnabled(PatchBuilderProperties, true);
@@ -78,18 +79,19 @@ void UUVEditorParameterizeMeshTool::Setup()
 		PolygroupLayerProperties = NewObject<UPolygroupLayersProperties>(this);
 		PolygroupLayerProperties->RestoreProperties(this, TEXT("UVEditorRecomputeUVsTool"));
 		PolygroupLayerProperties->InitializeGroupLayers(Targets[0]->AppliedCanonical.Get());
-		PolygroupLayerProperties->WatchProperty(PolygroupLayerProperties->ActiveGroupLayer, [&](FName) { OnSelectedGroupLayerChanged(); });
 		AddToolPropertySource(PolygroupLayerProperties);
 		PatchBuilderProperties->bPolygroupsEnabled = true;
 		UVAtlasProperties->bPolygroupsEnabled = true;
+		SetToolPropertySourceEnabled(PolygroupLayerProperties, (Settings->Method == EParameterizeMeshUVMethod::PatchBuilder && 
+		                                                    PatchBuilderProperties->bUsePolygroups)
+														   || (Settings->Method == EParameterizeMeshUVMethod::UVAtlas &&
+															   UVAtlasProperties->bUsePolygroups));
 	}
 	else
 	{
-		ActiveGroupSet = nullptr;
 		PatchBuilderProperties->bPolygroupsEnabled = false;
 		UVAtlasProperties->bPolygroupsEnabled = false;
 	}
-	UpdateActiveGroupLayer(false);  /* Don't update factories that don't exist yet. */
 
 	PatchBuilderProperties->bUDIMsEnabled = (FUVEditorUXSettings::CVarEnablePrototypeUDIMSupport.GetValueOnGameThread() > 0);
 	UVAtlasProperties->bUDIMsEnabled = (FUVEditorUXSettings::CVarEnablePrototypeUDIMSupport.GetValueOnGameThread() > 0);
@@ -99,10 +101,29 @@ void UUVEditorParameterizeMeshTool::Setup()
 	for (int32 TargetIndex = 0; TargetIndex < Targets.Num(); ++TargetIndex)
 	{
 		TObjectPtr<UUVEditorToolMeshInput> Target = Targets[TargetIndex];
-		Factories[TargetIndex] = NewObject<UUVEditorParameterizeMeshOperatorFactory>();
+		Factories[TargetIndex] = NewObject<UParameterizeMeshOperatorFactory>();
 		Factories[TargetIndex]->TargetTransform = Target->AppliedPreview->PreviewMesh->GetTransform();
 		Factories[TargetIndex]->Settings = Settings;
-		Factories[TargetIndex]->InputGroups = ActiveGroupSet;
+		Factories[TargetIndex]->GetPolygroups = [this]() -> TSharedPtr<UE::Geometry::FPolygroupSet, ESPMode::ThreadSafe> {
+			if (Targets.Num() == 1)
+			{
+				if (PolygroupLayerProperties->HasSelectedPolygroup() == false)
+				{
+					return MakeShared<UE::Geometry::FPolygroupSet, ESPMode::ThreadSafe>(Targets[0]->AppliedCanonical.Get());
+				}
+				else
+				{
+					FName SelectedName = PolygroupLayerProperties->ActiveGroupLayer;
+					FDynamicMeshPolygroupAttribute* FoundAttrib = UE::Geometry::FindPolygroupLayerByName(*Targets[0]->AppliedCanonical, SelectedName);
+					ensureMsgf(FoundAttrib, TEXT("Selected attribute not found! Falling back to Default group layer."));
+					return MakeShared<UE::Geometry::FPolygroupSet, ESPMode::ThreadSafe>(Targets[0]->AppliedCanonical.Get(), FoundAttrib);
+				}
+			}
+			else
+			{
+				return nullptr;
+			}
+		};
 		Factories[TargetIndex]->UVAtlasProperties = UVAtlasProperties;
 		Factories[TargetIndex]->XAtlasProperties = XAtlasProperties;
 		Factories[TargetIndex]->PatchBuilderProperties = PatchBuilderProperties;
@@ -130,6 +151,11 @@ void UUVEditorParameterizeMeshTool::OnPropertyModified(UObject* PropertySet, FPr
 {
 	TRACE_CPUPROFILER_EVENT_SCOPE(UVEditorParameterizeMeshTool_OnPropertyModified);
 
+	SetToolPropertySourceEnabled(PolygroupLayerProperties, (Settings->Method == EParameterizeMeshUVMethod::PatchBuilder && 
+		                                                    PatchBuilderProperties->bUsePolygroups)
+														   || (Settings->Method == EParameterizeMeshUVMethod::UVAtlas &&
+															   UVAtlasProperties->bUsePolygroups));
+
 	for (TObjectPtr<UUVEditorToolMeshInput> Target : Targets)
 	{
 		Target->AppliedPreview->InvalidateResult();
@@ -142,12 +168,9 @@ void UUVEditorParameterizeMeshTool::OnMethodTypeChanged()
 {
 	TRACE_CPUPROFILER_EVENT_SCOPE(UVEditorParameterizeMeshTool_OnMethodTypeChanged);
 	
-	SetToolPropertySourceEnabled(UVAtlasProperties, Settings->Method == EUVEditorParameterizeMeshUVMethod::UVAtlas);
-	SetToolPropertySourceEnabled(XAtlasProperties, Settings->Method == EUVEditorParameterizeMeshUVMethod::XAtlas);
-	SetToolPropertySourceEnabled(PatchBuilderProperties, Settings->Method == EUVEditorParameterizeMeshUVMethod::PatchBuilder);
-
-	SetToolPropertySourceEnabled(PolygroupLayerProperties, Settings->Method == EUVEditorParameterizeMeshUVMethod::UVAtlas ||
-		                                                   Settings->Method == EUVEditorParameterizeMeshUVMethod::PatchBuilder);
+	SetToolPropertySourceEnabled(UVAtlasProperties, Settings->Method == EParameterizeMeshUVMethod::UVAtlas);
+	SetToolPropertySourceEnabled(XAtlasProperties, Settings->Method == EParameterizeMeshUVMethod::XAtlas);
+	SetToolPropertySourceEnabled(PatchBuilderProperties, Settings->Method == EParameterizeMeshUVMethod::PatchBuilder);
 
 	for (TObjectPtr<UUVEditorToolMeshInput> Target : Targets)
 	{
@@ -221,46 +244,6 @@ void UUVEditorParameterizeMeshTool::Shutdown(EToolShutdownType ShutdownType)
 	Targets.Empty();
 }
 
-void UUVEditorParameterizeMeshTool::OnSelectedGroupLayerChanged()
-{
-	UpdateActiveGroupLayer();
-	for (TObjectPtr<UUVEditorToolMeshInput> Target : Targets)
-	{
-		Target->AppliedPreview->InvalidateResult();
-	}
-}
-
-
-void UUVEditorParameterizeMeshTool::UpdateActiveGroupLayer(bool bUpdateFactories)
-{
-	if (Targets.Num() == 1)
-	{
-		if (PolygroupLayerProperties->HasSelectedPolygroup() == false)
-		{
-			ActiveGroupSet = MakeShared<UE::Geometry::FPolygroupSet, ESPMode::ThreadSafe>(Targets[0]->AppliedCanonical.Get());
-		}
-		else
-		{
-			FName SelectedName = PolygroupLayerProperties->ActiveGroupLayer;
-			FDynamicMeshPolygroupAttribute* FoundAttrib = UE::Geometry::FindPolygroupLayerByName(*Targets[0]->AppliedCanonical, SelectedName);
-			ensureMsgf(FoundAttrib, TEXT("Selected attribute not found! Falling back to Default group layer."));
-			ActiveGroupSet = MakeShared<UE::Geometry::FPolygroupSet, ESPMode::ThreadSafe>(Targets[0]->AppliedCanonical.Get(), FoundAttrib);
-		}
-	}
-	else
-	{
-		ActiveGroupSet = nullptr;
-	}
-
-	if (bUpdateFactories)
-	{
-		for (int32 FactoryIdx = 0; FactoryIdx < Factories.Num(); ++FactoryIdx)
-		{
-			Factories[FactoryIdx]->InputGroups = ActiveGroupSet;
-		}
-	}
-}
-
 void UUVEditorParameterizeMeshTool::OnTick(float DeltaTime)
 {
 	for (TObjectPtr<UUVEditorToolMeshInput> Target : Targets)
@@ -314,11 +297,11 @@ void UUVEditorParameterizeMeshTool::RecordAnalytics()
 	Attributes.Add(FAnalyticsEventAttribute(TEXT("Stats.ToolActiveDuration"), (FDateTime::UtcNow() - ToolStartTimeAnalytics).ToString()));
 
 	// Tool settings chosen by the user (Volatile! Sync with EditCondition meta-tags in *Properties members)
-	const FString MethodName = StaticEnum<EUVEditorParameterizeMeshUVMethod>()->GetNameStringByIndex(static_cast<int>(Settings->Method));
+	const FString MethodName = StaticEnum<EParameterizeMeshUVMethod>()->GetNameStringByIndex(static_cast<int>(Settings->Method));
 	Attributes.Add(FAnalyticsEventAttribute(TEXT("Settings.Method"), MethodName));
 	switch (Settings->Method)
 	{
-		case EUVEditorParameterizeMeshUVMethod::PatchBuilder:
+		case EParameterizeMeshUVMethod::PatchBuilder:
 			Attributes.Add(FAnalyticsEventAttribute(FString::Printf(TEXT("Settings.%s.InitialPatches"), *MethodName), PatchBuilderProperties->InitialPatches));
 			Attributes.Add(FAnalyticsEventAttribute(FString::Printf(TEXT("Settings.%s.CurvatureAlignment"), *MethodName), PatchBuilderProperties->CurvatureAlignment));
 			Attributes.Add(FAnalyticsEventAttribute(FString::Printf(TEXT("Settings.%s.MergingDistortionThreshold"), *MethodName), PatchBuilderProperties->MergingDistortionThreshold));
@@ -331,11 +314,11 @@ void UUVEditorParameterizeMeshTool::RecordAnalytics()
 				Attributes.Add(FAnalyticsEventAttribute(FString::Printf(TEXT("Settings.%s.TextureResolution"), *MethodName), PatchBuilderProperties->TextureResolution));
 			}
 			break;
-		case EUVEditorParameterizeMeshUVMethod::UVAtlas:
+		case EParameterizeMeshUVMethod::UVAtlas:
 			Attributes.Add(FAnalyticsEventAttribute(FString::Printf(TEXT("Settings.%s.IslandStretch"), *MethodName), UVAtlasProperties->IslandStretch));
 			Attributes.Add(FAnalyticsEventAttribute(FString::Printf(TEXT("Settings.%s.NumIslands"), *MethodName), UVAtlasProperties->NumIslands));
 			break;
-		case EUVEditorParameterizeMeshUVMethod::XAtlas:
+		case EParameterizeMeshUVMethod::XAtlas:
 			Attributes.Add(FAnalyticsEventAttribute(FString::Printf(TEXT("Settings.%s.MaxIterations"), *MethodName), XAtlasProperties->MaxIterations));
 			break;
 		default:
diff --git a/Engine/Plugins/Experimental/UVEditor/Source/UVEditorToolsEditorOnly/Public/Operators/UVEditorParameterizeMeshOp.h b/Engine/Plugins/Experimental/UVEditor/Source/UVEditorToolsEditorOnly/Public/Operators/UVEditorParameterizeMeshOp.h
deleted file mode 100644
diff --git a/Engine/Plugins/Experimental/UVEditor/Source/UVEditorToolsEditorOnly/Public/UVEditorParameterizeMeshTool.h b/Engine/Plugins/Experimental/UVEditor/Source/UVEditorToolsEditorOnly/Public/UVEditorParameterizeMeshTool.h
--- a/Engine/Plugins/Experimental/UVEditor/Source/UVEditorToolsEditorOnly/Public/UVEditorParameterizeMeshTool.h
+++ b/Engine/Plugins/Experimental/UVEditor/Source/UVEditorToolsEditorOnly/Public/UVEditorParameterizeMeshTool.h
@@ -21,11 +21,11 @@ class UDynamicMeshComponent;
 class UMaterialInterface;
 class UMaterialInstanceDynamic;
 class UUVEditorToolMeshInput;
-class UUVEditorParameterizeMeshToolProperties;
-class UUVEditorParameterizeMeshToolUVAtlasProperties;
-class UUVEditorParameterizeMeshToolXAtlasProperties;
-class UUVEditorParameterizeMeshToolPatchBuilderProperties;
-class UUVEditorParameterizeMeshOperatorFactory;
+class UParameterizeMeshToolProperties;
+class UParameterizeMeshToolUVAtlasProperties;
+class UParameterizeMeshToolXAtlasProperties;
+class UParameterizeMeshToolPatchBuilderProperties;
+class UParameterizeMeshOperatorFactory;
 class UPolygroupLayersProperties;
 
 UCLASS()
@@ -77,29 +77,25 @@ protected:
 	TArray<TObjectPtr<UUVEditorToolMeshInput>> Targets;
 
 	UPROPERTY()
-	TObjectPtr<UUVEditorParameterizeMeshToolProperties> Settings = nullptr;
+	TObjectPtr<UParameterizeMeshToolProperties> Settings = nullptr;
 
 	UPROPERTY()
-	TObjectPtr<UUVEditorParameterizeMeshToolUVAtlasProperties> UVAtlasProperties = nullptr;
+	TObjectPtr<UParameterizeMeshToolUVAtlasProperties> UVAtlasProperties = nullptr;
 
 	UPROPERTY()
-	TObjectPtr<UUVEditorParameterizeMeshToolXAtlasProperties> XAtlasProperties = nullptr;
+	TObjectPtr<UParameterizeMeshToolXAtlasProperties> XAtlasProperties = nullptr;
 
 	UPROPERTY()
-	TObjectPtr<UUVEditorParameterizeMeshToolPatchBuilderProperties> PatchBuilderProperties = nullptr;
+	TObjectPtr<UParameterizeMeshToolPatchBuilderProperties> PatchBuilderProperties = nullptr;
 
 	UPROPERTY()
 	TObjectPtr<UPolygroupLayersProperties> PolygroupLayerProperties = nullptr;
 
 	UPROPERTY()
-	TArray<TObjectPtr<UUVEditorParameterizeMeshOperatorFactory>> Factories;
+	TArray<TObjectPtr<UParameterizeMeshOperatorFactory>> Factories;
 
 	void OnMethodTypeChanged();
 
-	TSharedPtr<UE::Geometry::FPolygroupSet, ESPMode::ThreadSafe> ActiveGroupSet;
-	void OnSelectedGroupLayerChanged();
-	void UpdateActiveGroupLayer(bool bUpdateFactories = true);
-
 	//
 	// Analytics
 	//
diff --git a/Engine/Plugins/Runtime/MeshModelingToolset/Source/MeshModelingToolsEditorOnly/MeshModelingToolsEditorOnly.Build.cs b/Engine/Plugins/Runtime/MeshModelingToolset/Source/MeshModelingToolsEditorOnly/MeshModelingToolsEditorOnly.Build.cs
--- a/Engine/Plugins/Runtime/MeshModelingToolset/Source/MeshModelingToolsEditorOnly/MeshModelingToolsEditorOnly.Build.cs
+++ b/Engine/Plugins/Runtime/MeshModelingToolset/Source/MeshModelingToolsEditorOnly/MeshModelingToolsEditorOnly.Build.cs
@@ -17,6 +17,7 @@ public class MeshModelingToolsEditorOnly : ModuleRules
 				"MeshModelingTools",
 				"ModelingComponents",
 				"ModelingOperatorsEditorOnly",
+				"DynamicMesh",
 				// ... add other public dependencies that you statically link with here ...
 			}
 		);
diff --git a/Engine/Plugins/Runtime/MeshModelingToolset/Source/MeshModelingToolsEditorOnly/Private/ParameterizeMeshTool.cpp b/Engine/Plugins/Runtime/MeshModelingToolset/Source/MeshModelingToolsEditorOnly/Private/ParameterizeMeshTool.cpp
--- a/Engine/Plugins/Runtime/MeshModelingToolset/Source/MeshModelingToolsEditorOnly/Private/ParameterizeMeshTool.cpp
+++ b/Engine/Plugins/Runtime/MeshModelingToolset/Source/MeshModelingToolsEditorOnly/Private/ParameterizeMeshTool.cpp
@@ -8,6 +8,8 @@
 #include "DynamicMesh/DynamicMesh3.h"
 #include "ParameterizationOps/ParameterizeMeshOp.h"
 #include "Properties/ParameterizeMeshProperties.h"
+#include "Polygroups/PolygroupUtil.h"
+#include "Polygroups/PolygroupSet.h"
 
 #include UE_INLINE_GENERATED_CPP_BY_NAME(ParameterizeMeshTool)
 
@@ -39,19 +41,6 @@ void UParameterizeMeshTool::Setup()
 	FTransform InputTransform = (FTransform)UE::ToolTarget::GetLocalToWorldTransform(Target);
 	FComponentMaterialSet MaterialSet = UE::ToolTarget::GetMaterialSet(Target);
 
-	Preview = NewObject<UMeshOpPreviewWithBackgroundCompute>(this);
-	Preview->Setup(GetTargetWorld(), this);
-	ToolSetupUtil::ApplyRenderingConfigurationToPreview(Preview->PreviewMesh, nullptr);
-	Preview->PreviewMesh->SetTangentsMode(EDynamicMeshComponentTangentsMode::AutoCalculated);
-	Preview->PreviewMesh->ReplaceMesh(*InputMesh);
-	Preview->ConfigureMaterials(MaterialSet.Materials, ToolSetupUtil::GetDefaultWorkingMaterial(GetToolManager()));
-	Preview->PreviewMesh->SetTransform(InputTransform);
-
-	Preview->OnMeshUpdated.AddLambda([this](UMeshOpPreviewWithBackgroundCompute* Op)
-	{
-		OnPreviewMeshUpdated();
-	});
-
 	UE::ToolTarget::HideSourceObject(Target);
 
 	// initialize our properties
@@ -88,20 +77,21 @@ void UParameterizeMeshTool::Setup()
 	AddToolPropertySource(PatchBuilderProperties);
 	SetToolPropertySourceEnabled(PatchBuilderProperties, false);
 
-	PolygroupLayerProperties = NewObject<UParameterizeMeshToolPatchBuilderGroupLayerProperties>(this);
-	PolygroupLayerProperties->RestoreProperties(this, TEXT("ModelingUVTools") );
+	PolygroupLayerProperties = NewObject<UPolygroupLayersProperties>(this);
+	PolygroupLayerProperties->RestoreProperties(this, TEXT("ModelingUVTools"));
 	PolygroupLayerProperties->InitializeGroupLayers(InputMesh.Get());
 	AddToolPropertySource(PolygroupLayerProperties);
-	SetToolPropertySourceEnabled(PolygroupLayerProperties, false);
-
+	PatchBuilderProperties->bPolygroupsEnabled = true;
+	UVAtlasProperties->bPolygroupsEnabled = true;
+	SetToolPropertySourceEnabled(PolygroupLayerProperties, (Settings->Method == EParameterizeMeshUVMethod::PatchBuilder && 
+		                                                    PatchBuilderProperties->bUsePolygroups)
+														   || (Settings->Method == EParameterizeMeshUVMethod::UVAtlas &&
+															   UVAtlasProperties->bUsePolygroups));
 
 	MaterialSettings = NewObject<UExistingMeshMaterialProperties>(this);
 	MaterialSettings->MaterialMode = ESetMeshMaterialMode::Checkerboard;
 	MaterialSettings->RestoreProperties(this, TEXT("ModelingUVTools"));
 	AddToolPropertySource(MaterialSettings);
-	// force update
-	MaterialSettings->UpdateMaterials();
-	Preview->OverrideMaterial = MaterialSettings->GetActiveOverrideMaterial();
 
 	if (bCreateUVLayoutViewOnSetup)
 	{
@@ -115,6 +105,43 @@ void UParameterizeMeshTool::Setup()
 		AddToolPropertySource(UVLayoutView->Settings);
 	}
 
+	Factory = NewObject<UParameterizeMeshOperatorFactory>();
+	Factory->TargetTransform = InputTransform;
+	Factory->Settings = Settings;
+	Factory->GetPolygroups = [this]() {
+		if (PolygroupLayerProperties->HasSelectedPolygroup() == false)
+		{
+			return MakeShared<UE::Geometry::FPolygroupSet, ESPMode::ThreadSafe>(InputMesh.Get());
+		}
+		else
+		{
+			FName SelectedName = PolygroupLayerProperties->ActiveGroupLayer;
+			FDynamicMeshPolygroupAttribute* FoundAttrib = UE::Geometry::FindPolygroupLayerByName(*InputMesh, SelectedName);
+			ensureMsgf(FoundAttrib, TEXT("Selected attribute not found! Falling back to Default group layer."));
+			return MakeShared<UE::Geometry::FPolygroupSet, ESPMode::ThreadSafe>(InputMesh.Get(), FoundAttrib);
+		}
+	};
+	Factory->UVAtlasProperties = UVAtlasProperties;
+	Factory->XAtlasProperties = XAtlasProperties;
+	Factory->PatchBuilderProperties = PatchBuilderProperties;
+	Factory->OriginalMesh = InputMesh;
+	Factory->GetSelectedUVChannel = [this]() { return UVChannelProperties->UVChannelNamesList.IndexOfByKey(UVChannelProperties->UVChannel); };
+
+	Preview = NewObject<UMeshOpPreviewWithBackgroundCompute>(this);
+	Preview->Setup(GetTargetWorld(), Factory);
+	ToolSetupUtil::ApplyRenderingConfigurationToPreview(Preview->PreviewMesh, nullptr);
+	Preview->PreviewMesh->SetTangentsMode(EDynamicMeshComponentTangentsMode::AutoCalculated);
+	Preview->PreviewMesh->ReplaceMesh(*InputMesh);
+	Preview->ConfigureMaterials(MaterialSet.Materials, ToolSetupUtil::GetDefaultWorkingMaterial(GetToolManager()));
+	Preview->PreviewMesh->SetTransform(InputTransform);
+
+	Preview->OnMeshUpdated.AddLambda([this](UMeshOpPreviewWithBackgroundCompute* Op)
+		{
+			OnPreviewMeshUpdated();
+		});
+	// force update
+	MaterialSettings->UpdateMaterials();
+	Preview->OverrideMaterial = MaterialSettings->GetActiveOverrideMaterial();
 	Preview->InvalidateResult();    // start compute
 
 	SetToolDisplayName(LOCTEXT("ToolNameGlobal", "AutoUV"));
@@ -125,6 +152,11 @@ void UParameterizeMeshTool::Setup()
 
 void UParameterizeMeshTool::OnPropertyModified(UObject* PropertySet, FProperty* Property)
 {
+	SetToolPropertySourceEnabled(PolygroupLayerProperties, (Settings->Method == EParameterizeMeshUVMethod::PatchBuilder && 
+		                                                    PatchBuilderProperties->bUsePolygroups)
+														   || (Settings->Method == EParameterizeMeshUVMethod::UVAtlas &&
+															   UVAtlasProperties->bUsePolygroups));
+
 	if (PropertySet == UVChannelProperties 
 		|| PropertySet == Settings 
 		|| PropertySet == UVAtlasProperties  
@@ -145,7 +177,6 @@ void UParameterizeMeshTool::OnMethodTypeChanged()
 	SetToolPropertySourceEnabled(UVAtlasProperties, Settings->Method == EParameterizeMeshUVMethod::UVAtlas);
 	SetToolPropertySourceEnabled(XAtlasProperties, Settings->Method == EParameterizeMeshUVMethod::XAtlas);
 	SetToolPropertySourceEnabled(PatchBuilderProperties, Settings->Method == EParameterizeMeshUVMethod::PatchBuilder);
-	SetToolPropertySourceEnabled(PolygroupLayerProperties, Settings->Method == EParameterizeMeshUVMethod::PatchBuilder);
 
 	Preview->InvalidateResult();
 }
@@ -183,6 +214,7 @@ void UParameterizeMeshTool::OnShutdown(EToolShutdownType ShutdownType)
 		GetToolManager()->EndUndoTransaction();
 	}
 
+	Factory = nullptr;
 }
 
 void UParameterizeMeshTool::Render(IToolsContextRenderAPI* RenderAPI)
@@ -208,50 +240,6 @@ bool UParameterizeMeshTool::CanAccept() const
 	return Super::CanAccept() && Preview->HaveValidResult();
 }
 
-TUniquePtr<FDynamicMeshOperator> UParameterizeMeshTool::MakeNewOperator()
-{
-	TUniquePtr<FParameterizeMeshOp> ParameterizeMeshOp = MakeUnique<FParameterizeMeshOp>();
-
-	ParameterizeMeshOp->InputMesh = InputMesh;
-	ParameterizeMeshOp->UVLayer = UVChannelProperties->GetSelectedChannelIndex(true);
-	ParameterizeMeshOp->Method = (EParamOpBackend)(int)Settings->Method;
-
-	// uvatlas options
-	ParameterizeMeshOp->Stretch   = UVAtlasProperties->IslandStretch;
-	ParameterizeMeshOp->NumCharts = UVAtlasProperties->NumIslands;
-	if (Settings->Method == EParameterizeMeshUVMethod::UVAtlas)
-	{
-		ParameterizeMeshOp->Width = ParameterizeMeshOp->Height = UVAtlasProperties->TextureResolution;
-	}
-	
-	// xatlas options
-	ParameterizeMeshOp->XAtlasMaxIterations = XAtlasProperties->MaxIterations;
-
-	// patchbuilder options
-	ParameterizeMeshOp->InitialPatchCount = PatchBuilderProperties->InitialPatches;
-	ParameterizeMeshOp->bRespectInputGroups = PolygroupLayerProperties->bConstrainToPolygroups;
-	if (PolygroupLayerProperties->bConstrainToPolygroups)
-	{
-		ParameterizeMeshOp->InputGroupLayer = PolygroupLayerProperties->GetSelectedLayer(*InputMesh);
-	}
-	ParameterizeMeshOp->PatchCurvatureAlignmentWeight = PatchBuilderProperties->CurvatureAlignment;
-	ParameterizeMeshOp->PatchMergingMetricThresh = PatchBuilderProperties->MergingDistortionThreshold;
-	ParameterizeMeshOp->PatchMergingAngleThresh = PatchBuilderProperties->MergingAngleThreshold;
-	ParameterizeMeshOp->ExpMapNormalSmoothingSteps = PatchBuilderProperties->SmoothingSteps;
-	ParameterizeMeshOp->ExpMapNormalSmoothingAlpha = PatchBuilderProperties->SmoothingAlpha;
-	ParameterizeMeshOp->bEnablePacking = PatchBuilderProperties->bRepack;
-	if (Settings->Method == EParameterizeMeshUVMethod::PatchBuilder)
-	{
-		ParameterizeMeshOp->Width = ParameterizeMeshOp->Height = PatchBuilderProperties->TextureResolution;
-	}
-
-	ParameterizeMeshOp->SetTransform(UE::ToolTarget::GetLocalToWorldTransform(Target));
-
-	return ParameterizeMeshOp;
-}
-
-
-
 void UParameterizeMeshTool::OnPreviewMeshUpdated()
 {
 	if (UVLayoutView)
diff --git a/Engine/Plugins/Runtime/MeshModelingToolset/Source/MeshModelingToolsEditorOnly/Public/ParameterizeMeshTool.h b/Engine/Plugins/Runtime/MeshModelingToolset/Source/MeshModelingToolsEditorOnly/Public/ParameterizeMeshTool.h
--- a/Engine/Plugins/Runtime/MeshModelingToolset/Source/MeshModelingToolsEditorOnly/Public/ParameterizeMeshTool.h
+++ b/Engine/Plugins/Runtime/MeshModelingToolset/Source/MeshModelingToolsEditorOnly/Public/ParameterizeMeshTool.h
@@ -21,6 +21,7 @@ class UParameterizeMeshToolProperties;
 class UParameterizeMeshToolUVAtlasProperties;
 class UParameterizeMeshToolXAtlasProperties;
 class UParameterizeMeshToolPatchBuilderProperties;
+class UParameterizeMeshOperatorFactory;
 
 
 UCLASS()
@@ -32,23 +33,12 @@ public:
 };
 
 
-UCLASS()
-class MESHMODELINGTOOLSEDITORONLY_API UParameterizeMeshToolPatchBuilderGroupLayerProperties : public UPolygroupLayersProperties
-{
-	GENERATED_BODY()
-public:
-
-	UPROPERTY(EditAnywhere, Category = "PolyGroup Layer", meta = (DisplayName = "Constrain to PolyGroups"))
-	bool bConstrainToPolygroups = false;
-};
-
-
 /**
  * UParameterizeMeshTool automatically decomposes the input mesh into charts, solves for UVs,
  * and then packs the resulting charts
  */
 UCLASS()
-class MESHMODELINGTOOLSEDITORONLY_API UParameterizeMeshTool : public USingleSelectionMeshEditingTool, public UE::Geometry::IDynamicMeshOperatorFactory
+class MESHMODELINGTOOLSEDITORONLY_API UParameterizeMeshTool : public USingleSelectionMeshEditingTool
 {
 	GENERATED_BODY()
 
@@ -65,9 +55,6 @@ public:
 
 	virtual void OnPropertyModified(UObject* PropertySet, FProperty* Property) override;
 
-	// IDynamicMeshOperatorFactory API
-	virtual TUniquePtr<UE::Geometry::FDynamicMeshOperator> MakeNewOperator() override;
-
 protected:
 	UPROPERTY()
 	TObjectPtr<UMeshUVChannelProperties> UVChannelProperties = nullptr;
@@ -85,7 +72,7 @@ protected:
 	TObjectPtr<UParameterizeMeshToolPatchBuilderProperties> PatchBuilderProperties = nullptr;
 
 	UPROPERTY()
-	TObjectPtr<UParameterizeMeshToolPatchBuilderGroupLayerProperties> PolygroupLayerProperties = nullptr;
+	TObjectPtr<UPolygroupLayersProperties> PolygroupLayerProperties = nullptr;
 
 	UPROPERTY()
 	TObjectPtr<UExistingMeshMaterialProperties> MaterialSettings = nullptr;
@@ -99,6 +86,9 @@ protected:
 	UPROPERTY()
 	TObjectPtr<UMeshOpPreviewWithBackgroundCompute> Preview = nullptr;
 
+	UPROPERTY()
+	TObjectPtr<UParameterizeMeshOperatorFactory> Factory;
+
 	TSharedPtr<UE::Geometry::FDynamicMesh3, ESPMode::ThreadSafe> InputMesh;
 
 	void OnMethodTypeChanged();
diff --git a/Engine/Plugins/Runtime/MeshModelingToolset/Source/ModelingOperatorsEditorOnly/Private/ParameterizationOps/ParameterizeMeshOp.cpp b/Engine/Plugins/Runtime/MeshModelingToolset/Source/ModelingOperatorsEditorOnly/Private/ParameterizationOps/ParameterizeMeshOp.cpp
--- a/Engine/Plugins/Runtime/MeshModelingToolset/Source/ModelingOperatorsEditorOnly/Private/ParameterizationOps/ParameterizeMeshOp.cpp
+++ b/Engine/Plugins/Runtime/MeshModelingToolset/Source/ModelingOperatorsEditorOnly/Private/ParameterizationOps/ParameterizeMeshOp.cpp
@@ -27,7 +27,7 @@ using namespace UE::Geometry;
 
 #define LOCTEXT_NAMESPACE "ParameterizeMeshOp"
 
-FParameterizeMeshOp::FLinearMesh::FLinearMesh(const FDynamicMesh3& Mesh, const bool bRespectPolygroups)
+FParameterizeMeshOp::FLinearMesh::FLinearMesh(const FDynamicMesh3& Mesh, UE::Geometry::FPolygroupSet* InputGroups)
 {
 	TArray<FVector3f>& Positions = this->VertexBuffer;
 
@@ -40,11 +40,10 @@ FParameterizeMeshOp::FLinearMesh::FLinearMesh(const FDynamicMesh3& Mesh, const b
 	// Compute the mapping from triangle ID to triangle number
 	{
 		const int32 MaxTriID = Mesh.MaxTriangleID(); // really +1.. all TriID < MaxTriID
-		const int32 NumTris  = Mesh.TriangleCount();
+		const int32 NumTris = Mesh.TriangleCount();
 
 		// reserve space and add elements 
-		TriFromID.Empty(MaxTriID);
-		TriFromID.AddUninitialized(MaxTriID);
+		TriFromID.SetNum(MaxTriID);
 
 		// reserve space
 		TriToID.Empty(NumTris);
@@ -61,7 +60,7 @@ FParameterizeMeshOp::FLinearMesh::FLinearMesh(const FDynamicMesh3& Mesh, const b
 	// Compute the mapping from vertex ID to vertex number
 	{
 		const int32 MaxVertID = Mesh.MaxVertexID();
-		const int32 NumVerts  = Mesh.VertexCount();
+		const int32 NumVerts = Mesh.VertexCount();
 
 		// reserve space and add elements
 		VertFromID.Empty(MaxVertID);
@@ -84,10 +83,9 @@ FParameterizeMeshOp::FLinearMesh::FLinearMesh(const FDynamicMesh3& Mesh, const b
 		int32 NumVerts = Mesh.VertexCount();
 		Positions.Empty(NumVerts);
 
-		for (const auto& Vertex : Mesh.VerticesItr())
+		for (const FVector3d& Vertex : Mesh.VerticesItr())
 		{
-			FVector3f Pos(Vertex.X, Vertex.Y, Vertex.Z);
-			Positions.Add(Pos);
+			Positions.Add(static_cast<FVector3f>(Vertex));
 		}
 	}
 
@@ -96,7 +94,7 @@ FParameterizeMeshOp::FLinearMesh::FLinearMesh(const FDynamicMesh3& Mesh, const b
 	// Fill the index buffer
 	{
 		IndexBuffer.Empty(NumTris * 3);
-		for (const auto& Tri : Mesh.TrianglesItr())
+		for (const FIndex3i& Tri : Mesh.TrianglesItr())
 		{
 			for (int i = 0; i < 3; ++i)
 			{
@@ -113,7 +111,7 @@ FParameterizeMeshOp::FLinearMesh::FLinearMesh(const FDynamicMesh3& Mesh, const b
 	AdjacencyBuffer.Empty(NumTris * 3);
 
 	// Create Adjacency - create boundaries in adjacency at polygroup boundaries if requested.
-	if (bRespectPolygroups && Mesh.HasTriangleGroups())
+	if (InputGroups)
 	{
 		for (int TriID : Mesh.TriangleIndicesItr())
 		{
@@ -125,8 +123,8 @@ FParameterizeMeshOp::FLinearMesh::FLinearMesh(const FDynamicMesh3& Mesh, const b
 				int32 RemapNbrID = FDynamicMesh3::InvalidID;
 				if (NbrID != FDynamicMesh3::InvalidID)
 				{
-					const int32 CurTriGroup = Mesh.GetTriangleGroup(TriID);
-					const int32 NbrTriGroup = Mesh.GetTriangleGroup(NbrID);
+					const int32 CurTriGroup = InputGroups->GetGroup(TriID);
+					const int32 NbrTriGroup = InputGroups->GetGroup(NbrID);
 
 					RemapNbrID = (CurTriGroup == NbrTriGroup) ? TriFromID[NbrID] : FDynamicMesh3::InvalidID;
 				}
@@ -155,8 +153,8 @@ FParameterizeMeshOp::FLinearMesh::FLinearMesh(const FDynamicMesh3& Mesh, const b
 
 
 void FParameterizeMeshOp::CopyNewUVsToMesh(
-	FDynamicMesh3& Mesh, 
-	const FLinearMesh& LinearMesh, 
+	FDynamicMesh3& Mesh,
+	const FLinearMesh& LinearMesh,
 	const FDynamicMesh3& FlippedMesh,
 	const TArray<FVector2D>& UVVertexBuffer,
 	const TArray<int32>& UVIndexBuffer,
@@ -170,7 +168,7 @@ void FParameterizeMeshOp::CopyNewUVsToMesh(
 	FDynamicMeshUVOverlay* UVOverlay = Mesh.Attributes()->GetUVLayer(UVLayer);
 	if (UVOverlay == nullptr)
 	{
-		Mesh.Attributes()->SetNumUVLayers(UVLayer+1);
+		Mesh.Attributes()->SetNumUVLayers(UVLayer + 1);
 		UVOverlay = Mesh.Attributes()->GetUVLayer(UVLayer);
 	}
 
@@ -230,9 +228,35 @@ void FParameterizeMeshOp::CopyNewUVsToMesh(
 	}
 }
 
+void FParameterizeMeshOp::LayoutToUDIMByPolygroup(FDynamicMesh3& InOutMesh, UE::Geometry::FPolygroupSet& PolygroupSet)
+{
+	FDynamicMeshUVEditor UVEditor(&InOutMesh, UVLayer, true);
+	TMap<int32, TArray<int32> > TidsPerTile;
+
+	for (int32 Tid : InOutMesh.TriangleIndicesItr())
+	{
+		if (UVEditor.GetOverlay()->IsSetTriangle(Tid))
+		{
+			int32 TidGroup = PolygroupSet.GetGroup(Tid);
+			TidsPerTile.FindOrAdd(TidGroup).Add(Tid);
+		}
+	}
+
+	TArray<int32> Tiles;
+	TidsPerTile.GetKeys(Tiles);
+	Tiles.Sort();
+	for (int32 FlatGroupID = 0; FlatGroupID < Tiles.Num(); ++FlatGroupID)
+	{
+		FVector2i TilePos;
+		TilePos.X = FlatGroupID % 10;
+		TilePos.Y = FlatGroupID / 10;
+
+		ensure(UVEditor.UDIMPack(Width, Gutter, TilePos, TidsPerTile.Find(Tiles[FlatGroupID])));
+	}
+}
 
 
-bool FParameterizeMeshOp::ComputeUVs_UVAtlas(FDynamicMesh3& Mesh,  TFunction<bool(float)>& Interrupter)
+bool FParameterizeMeshOp::ComputeUVs_UVAtlas(FDynamicMesh3& Mesh, TFunction<bool(float)>& Interrupter)
 {
 #if NO_UVATLAS
 
@@ -240,14 +264,14 @@ bool FParameterizeMeshOp::ComputeUVs_UVAtlas(FDynamicMesh3& Mesh,  TFunction<boo
 	return false;
 
 #else	// NO_UVATLAS
-	
+
 	TRACE_CPUPROFILER_EVENT_SCOPE(ParameterizeMeshOp_ComputeUVs_UVAtlas);
 
 	// the UVAtlas code is unhappy if you feed it a single degenerate triangle
 	bool bNonDegenerate = true;
 	if (Mesh.TriangleCount() == 1)
 	{
-		for (auto TriID : Mesh.TriangleIndicesItr())
+		for (int32 TriID : Mesh.TriangleIndicesItr())
 		{
 			double Area = Mesh.GetTriArea(TriID);
 			bNonDegenerate = bNonDegenerate && FMath::Abs(Area) > 1.e-5;
@@ -272,7 +296,7 @@ bool FParameterizeMeshOp::ComputeUVs_UVAtlas(FDynamicMesh3& Mesh,  TFunction<boo
 	}
 
 	// Convert to a dense form.
-	FLinearMesh LinearMesh(FlippedMesh, false);
+	FLinearMesh LinearMesh(FlippedMesh, (bRespectInputGroups && InputGroups.IsValid()) ? InputGroups.Get() : nullptr);
 
 	// Data to be populated by the UV generation tool
 	TArray<FVector2D> UVVertexBuffer;
@@ -280,14 +304,14 @@ bool FParameterizeMeshOp::ComputeUVs_UVAtlas(FDynamicMesh3& Mesh,  TFunction<boo
 	TArray<int32>     VertexRemapArray; // This maps the UV vertices to the original position vertices.  Note multiple UV vertices might share the same positional vertex (due to UV boundaries)
 
 
-	float MaxStretch     = Stretch;
+	float MaxStretch = Stretch;
 	int32 MaxChartNumber = NumCharts;
 
 	TUniquePtr<IProxyLODParameterization> ParameterizationTool = IProxyLODParameterization::CreateTool();
 	bool bSuccess = ParameterizationTool->GenerateUVs(Width, Height, Gutter, LinearMesh.VertexBuffer,
-	                                                  LinearMesh.IndexBuffer, LinearMesh.AdjacencyBuffer, Interrupter,
-	                                                  UVVertexBuffer, UVIndexBuffer, VertexRemapArray, MaxStretch,
-	                                                  MaxChartNumber);
+		LinearMesh.IndexBuffer, LinearMesh.AdjacencyBuffer, Interrupter,
+		UVVertexBuffer, UVIndexBuffer, VertexRemapArray, MaxStretch,
+		MaxChartNumber);
 
 	// Add the UVs to the FDynamicMesh
 	if (bSuccess)
@@ -299,6 +323,11 @@ bool FParameterizeMeshOp::ComputeUVs_UVAtlas(FDynamicMesh3& Mesh,  TFunction<boo
 		NewResultInfo.AddError(FGeometryError(0, LOCTEXT("UVAtlasFailed", "UVAtlas failed")));
 	}
 
+	if (bPackToUDIMSByOriginPolygroup && bRespectInputGroups && InputGroups.IsValid())
+	{
+		LayoutToUDIMByPolygroup(Mesh, *InputGroups);
+	}
+
 	return bSuccess;
 
 #endif	// NO_UVATLAS
@@ -312,7 +341,7 @@ bool FParameterizeMeshOp::ComputeUVs_UVAtlas(FDynamicMesh3& Mesh,  TFunction<boo
 
 
 
-bool FParameterizeMeshOp::ComputeUVs_XAtlas(FDynamicMesh3& Mesh,  TFunction<bool(float)>& Interrupter)
+bool FParameterizeMeshOp::ComputeUVs_XAtlas(FDynamicMesh3& Mesh, TFunction<bool(float)>& Interrupter)
 {
 	TRACE_CPUPROFILER_EVENT_SCOPE(ParameterizeMeshOp_ComputeUVs_XAtlas);
 
@@ -326,7 +355,7 @@ bool FParameterizeMeshOp::ComputeUVs_XAtlas(FDynamicMesh3& Mesh,  TFunction<bool
 	}
 
 	// Convert to a dense form.
-	FLinearMesh LinearMesh(FlippedMesh, false);
+	FLinearMesh LinearMesh(FlippedMesh, (bRespectInputGroups && InputGroups.IsValid()) ? InputGroups.Get() : nullptr);
 
 	// Data to be populated by the UV generation tool
 	TArray<FVector2D> UVVertexBuffer;
@@ -337,12 +366,12 @@ bool FParameterizeMeshOp::ComputeUVs_XAtlas(FDynamicMesh3& Mesh,  TFunction<bool
 	ChartOptions.MaxIterations = XAtlasMaxIterations;
 	XAtlasWrapper::XAtlasPackOptions PackOptions;
 	bool bSuccess = XAtlasWrapper::ComputeUVs(LinearMesh.IndexBuffer,
-											LinearMesh.VertexBuffer, 
-											ChartOptions,
-											PackOptions,
-											UVVertexBuffer, 
-											UVIndexBuffer, 
-											VertexRemapArray);
+		LinearMesh.VertexBuffer,
+		ChartOptions,
+		PackOptions,
+		UVVertexBuffer,
+		UVIndexBuffer,
+		VertexRemapArray);
 
 	// Add the UVs to the FDynamicMesh
 	if (bSuccess)
@@ -363,17 +392,16 @@ bool FParameterizeMeshOp::ComputeUVs_XAtlas(FDynamicMesh3& Mesh,  TFunction<bool
 bool FParameterizeMeshOp::ComputeUVs_PatchBuilder(FDynamicMesh3& InOutMesh, FProgressCancel* ProgressCancel)
 {
 	TRACE_CPUPROFILER_EVENT_SCOPE(ParameterizeMeshOp_ComputeUVs_PatchBuilder);
-	
+
 	FPatchBasedMeshUVGenerator UVGenerator;
 
 	TUniquePtr<FPolygroupSet> PolygroupConstraint;
-	if (bRespectInputGroups && InputGroupLayer.CheckExists(&InOutMesh))
+	if (bRespectInputGroups && InputGroups.IsValid())
 	{
-		PolygroupConstraint = MakeUnique<FPolygroupSet>(&InOutMesh, InputGroupLayer);
-		UVGenerator.GroupConstraint = PolygroupConstraint.Get();
+		UVGenerator.GroupConstraint = InputGroups.Get();
 	}
 
-	UVGenerator.TargetPatchCount = FMath::Max(1,InitialPatchCount);
+	UVGenerator.TargetPatchCount = FMath::Max(1, InitialPatchCount);
 	UVGenerator.bNormalWeightedPatches = true;
 	UVGenerator.PatchNormalWeight = FMath::Clamp(PatchCurvatureAlignmentWeight, 0.0, 999999.0);
 	UVGenerator.MinPatchSize = 2;
@@ -393,6 +421,11 @@ bool FParameterizeMeshOp::ComputeUVs_PatchBuilder(FDynamicMesh3& InOutMesh, FPro
 	FDynamicMeshUVEditor UVEditor(&InOutMesh, UVLayer, true);
 	FGeometryResult Result = UVGenerator.AutoComputeUVs(*UVEditor.GetMesh(), *UVEditor.GetOverlay(), ProgressCancel);
 
+	if (bPackToUDIMSByOriginPolygroup && bRespectInputGroups && InputGroups.IsValid())
+	{
+		LayoutToUDIMByPolygroup(InOutMesh, *InputGroups);
+	}
+
 	SetResultInfo(NewResultInfo);
 	return (Result.HasFailed() == false);
 }
@@ -461,15 +494,18 @@ TUniquePtr<FDynamicMeshOperator> UParameterizeMeshOperatorFactory::MakeNewOperat
 	TUniquePtr<FParameterizeMeshOp> Op = MakeUnique<FParameterizeMeshOp>();
 
 	Op->InputMesh = OriginalMesh;
+	Op->InputGroups = GetPolygroups();
 
 	switch (Settings->Method)
 	{
 	case EParameterizeMeshUVMethod::PatchBuilder:
 		Op->Method = EParamOpBackend::PatchBuilder;
 		Op->InitialPatchCount = PatchBuilderProperties->InitialPatches;
+		Op->bRespectInputGroups = PatchBuilderProperties->bUsePolygroups;
 		Op->PatchCurvatureAlignmentWeight = PatchBuilderProperties->CurvatureAlignment;
 		Op->PatchMergingMetricThresh = PatchBuilderProperties->MergingDistortionThreshold;
 		Op->PatchMergingAngleThresh = PatchBuilderProperties->MergingAngleThreshold;
+		Op->bPackToUDIMSByOriginPolygroup = PatchBuilderProperties->bLayoutUDIMPerPolygroup;
 		Op->ExpMapNormalSmoothingSteps = PatchBuilderProperties->SmoothingSteps;
 		Op->ExpMapNormalSmoothingAlpha = PatchBuilderProperties->SmoothingAlpha;
 		Op->bEnablePacking = PatchBuilderProperties->bRepack;
@@ -477,8 +513,10 @@ TUniquePtr<FDynamicMeshOperator> UParameterizeMeshOperatorFactory::MakeNewOperat
 		Op->Height = PatchBuilderProperties->TextureResolution;
 		break;
 	case EParameterizeMeshUVMethod::UVAtlas:
-		Op->Method = EParamOpBackend::UVAtlas;		
+		Op->Method = EParamOpBackend::UVAtlas;
 		Op->Stretch = UVAtlasProperties->IslandStretch;
+		Op->bRespectInputGroups = UVAtlasProperties->bUsePolygroups;
+		Op->bPackToUDIMSByOriginPolygroup = UVAtlasProperties->bLayoutUDIMPerPolygroup;
 		Op->NumCharts = UVAtlasProperties->NumIslands;
 		Op->Width = UVAtlasProperties->TextureResolution;
 		Op->Height = UVAtlasProperties->TextureResolution;
@@ -489,7 +527,7 @@ TUniquePtr<FDynamicMeshOperator> UParameterizeMeshOperatorFactory::MakeNewOperat
 		break;
 	}
 
-	Op->UVLayer = GetSelectedUVChannel();	
+	Op->UVLayer = GetSelectedUVChannel();
 
 	FTransformSRT3d Transform(TargetTransform);
 	Op->SetTransform(Transform);
diff --git a/Engine/Plugins/Runtime/MeshModelingToolset/Source/ModelingOperatorsEditorOnly/Public/ParameterizationOps/ParameterizeMeshOp.h b/Engine/Plugins/Runtime/MeshModelingToolset/Source/ModelingOperatorsEditorOnly/Public/ParameterizationOps/ParameterizeMeshOp.h
--- a/Engine/Plugins/Runtime/MeshModelingToolset/Source/ModelingOperatorsEditorOnly/Public/ParameterizationOps/ParameterizeMeshOp.h
+++ b/Engine/Plugins/Runtime/MeshModelingToolset/Source/ModelingOperatorsEditorOnly/Public/ParameterizationOps/ParameterizeMeshOp.h
@@ -38,7 +38,10 @@ public:
 
 	// source mesh
 	TSharedPtr<FDynamicMesh3, ESPMode::ThreadSafe> InputMesh;
-		
+
+	// source groups (optional)
+	TSharedPtr<UE::Geometry::FPolygroupSet, ESPMode::ThreadSafe> InputGroups;
+
 	// UVAtlas generation parameters
 	float Stretch = 0.11f;
 	int32 NumCharts = 0;
@@ -49,21 +52,23 @@ public:
 	// PatchBuilder generation parameters
 	int32 InitialPatchCount = 100;
 	bool bRespectInputGroups = false;
-	FPolygroupLayer InputGroupLayer;
-	double PatchCurvatureAlignmentWeight = 1.0;
-	double PatchMergingMetricThresh = 1.5;
-	double PatchMergingAngleThresh = 45.0;
+	float PatchCurvatureAlignmentWeight = 1.0f;
+	float PatchMergingMetricThresh = 1.5f;
+	float PatchMergingAngleThresh = 45.0f;
 	int ExpMapNormalSmoothingSteps = 5;
-	double ExpMapNormalSmoothingAlpha = 0.25;
+	float ExpMapNormalSmoothingAlpha = 0.25f;
 
 	// UV layer
 	int32 UVLayer = 0;
 
+	// UDIM support
+	bool bPackToUDIMSByOriginPolygroup = false;
+
 	// Atlas Packing parameters
 	bool bEnablePacking = true;
 	int32 Height = 512;
 	int32 Width = 512;
-	float Gutter = 2.5;
+	float Gutter = 2.5f;
 
 	EParamOpBackend Method = EParamOpBackend::UVAtlas;
 
@@ -85,7 +90,7 @@ protected:
 	// dense index/vertex buffer based representation of the data needed for parameterization
 	struct FLinearMesh
 	{
-		FLinearMesh(const FDynamicMesh3& Mesh, const bool bRespectPolygroups);
+		FLinearMesh(const FDynamicMesh3& Mesh, UE::Geometry::FPolygroupSet* InputGroups);
 
 		// Stripped down mesh
 		TArray<int32>   IndexBuffer;
@@ -114,6 +119,8 @@ protected:
 		const TArray<int32>& UVIndexBuffer,
 		const TArray<int32>& VertexRemapArray,
 		bool bReverseOrientation);
+
+	void LayoutToUDIMByPolygroup(FDynamicMesh3& InOutMesh, UE::Geometry::FPolygroupSet& InputGroups);
 };
 
 } // end namespace UE::Geometry
@@ -146,7 +153,9 @@ public:
 	UPROPERTY()
 	TObjectPtr<UParameterizeMeshToolPatchBuilderProperties> PatchBuilderProperties = nullptr;
 
+	//TSharedPtr<UE::Geometry::FPolygroupSet, ESPMode::ThreadSafe> InputGroups;
 	TSharedPtr<UE::Geometry::FDynamicMesh3, ESPMode::ThreadSafe> OriginalMesh;
 	TUniqueFunction<int32()> GetSelectedUVChannel = []() { return 0; };
+	TUniqueFunction<TSharedPtr<UE::Geometry::FPolygroupSet, ESPMode::ThreadSafe>()> GetPolygroups = []() { return nullptr; };
 	FTransform TargetTransform;
 };
\ No newline at end of file
diff --git a/Engine/Plugins/Runtime/MeshModelingToolset/Source/ModelingOperatorsEditorOnly/Public/Properties/ParameterizeMeshProperties.h b/Engine/Plugins/Runtime/MeshModelingToolset/Source/ModelingOperatorsEditorOnly/Public/Properties/ParameterizeMeshProperties.h
--- a/Engine/Plugins/Runtime/MeshModelingToolset/Source/ModelingOperatorsEditorOnly/Public/Properties/ParameterizeMeshProperties.h
+++ b/Engine/Plugins/Runtime/MeshModelingToolset/Source/ModelingOperatorsEditorOnly/Public/Properties/ParameterizeMeshProperties.h
@@ -7,6 +7,7 @@
 
 #include "ParameterizeMeshProperties.generated.h"
 
+
 UENUM()
 enum class EParameterizeMeshUVMethod
 {
@@ -51,6 +52,22 @@ public:
 	/** Expected resolution of the output textures; this controls spacing left between UV islands to avoid interpolation artifacts. */
 	UPROPERTY(EditAnywhere, Category = UVAtlas, meta = (UIMin = "64", UIMax = "2048", ClampMin = "2", ClampMax = "4096"))
 	int TextureResolution = 1024;
+
+	/** Generate new UVs based on polygroups from specified layer. */
+	UPROPERTY(EditAnywhere, Category = UVAtlas, meta = (DisplayName = "Constrain to Polygroups", EditCondition = bPolygroupsEnabled, EditConditionHides, HideEditConditionToggle))
+	bool bUsePolygroups;
+
+	/** Layout resulting islands on UDIMs based on polygroups. */
+	UPROPERTY(EditAnywhere, Category = UVAtlas, meta = (EditCondition = "bUsePolygroups && bUDIMsEnabled", EditConditionHides, HideEditConditionToggle))
+	bool bLayoutUDIMPerPolygroup;
+
+	/** Controls if polygroup options are available to the user. */
+	UPROPERTY()
+	bool bPolygroupsEnabled;
+
+	/** Controls if UDIM options are available to the user. */
+	UPROPERTY()
+	bool bUDIMsEnabled;
 };
 
 
@@ -101,4 +118,20 @@ public:
 	/** Expected resolution of the output textures; this controls spacing left between UV islands to avoid interpolation artifacts. This is only enabled when Repack is enabled. */
 	UPROPERTY(EditAnywhere, Category = PatchBuilder, meta = (UIMin = "64", UIMax = "2048", ClampMin = "2", ClampMax = "4096", EditCondition = bRepack))
 	int TextureResolution = 1024;
+
+	/** Generate new UVs based on polygroups from specified layer. */
+	UPROPERTY(EditAnywhere, Category = PatchBuilder, meta = (DisplayName = "Constrain to Polygroups", EditCondition = bPolygroupsEnabled, EditConditionHides, HideEditConditionToggle))
+	bool bUsePolygroups;
+
+	/** Layout resulting islands on UDIMs based on polygroups. */
+	UPROPERTY(EditAnywhere, Category = PatchBuilder, meta = (EditCondition = "bUsePolygroups && bUDIMsEnabled", EditConditionHides, HideEditConditionToggle))
+	bool bLayoutUDIMPerPolygroup;
+
+	/** Controls if polygroup options are available to the user. */
+	UPROPERTY()
+	bool bPolygroupsEnabled;
+
+	/** Controls if UDIM options are available to the user. */
+	UPROPERTY()
+	bool bUDIMsEnabled;
 };
