ModelingMode: Merge and unify the Layout codebases between the UV Editor and Modeling Mode.

diff --git a/Engine/Plugins/Experimental/UVEditor/Source/UVEditor/Private/UVEditorMode.cpp b/Engine/Plugins/Experimental/UVEditor/Source/UVEditor/Private/UVEditorMode.cpp
--- a/Engine/Plugins/Experimental/UVEditor/Source/UVEditor/Private/UVEditorMode.cpp
+++ b/Engine/Plugins/Experimental/UVEditor/Source/UVEditor/Private/UVEditorMode.cpp
@@ -47,7 +47,7 @@
 #include "UVEditorUXSettings.h"
 #include "UDIMUtilities.h"
 #include "EditorSupportDelegates.h"
-#include "Utilities/MeshUDIMClassifier.h"
+#include "Parameterization/MeshUDIMClassifier.h"
 #include "UVEditorLogging.h"
 #include "UObject/ObjectSaveContext.h"
 
diff --git a/Engine/Plugins/Experimental/UVEditor/Source/UVEditorTools/Private/Operators/UVEditorRecomputeUVsOp.cpp b/Engine/Plugins/Experimental/UVEditor/Source/UVEditorTools/Private/Operators/UVEditorRecomputeUVsOp.cpp
--- a/Engine/Plugins/Experimental/UVEditor/Source/UVEditorTools/Private/Operators/UVEditorRecomputeUVsOp.cpp
+++ b/Engine/Plugins/Experimental/UVEditor/Source/UVEditorTools/Private/Operators/UVEditorRecomputeUVsOp.cpp
@@ -9,7 +9,8 @@
 #include "Parameterization/DynamicMeshUVEditor.h"
 #include "Parameterization/PatchBasedMeshUVGenerator.h"
 #include "Properties/RecomputeUVsProperties.h"
-#include "Utilities/MeshUDIMClassifier.h"
+#include "Parameterization/MeshUDIMClassifier.h"
+
 
 #include "DynamicSubmesh3.h"
 
diff --git a/Engine/Plugins/Experimental/UVEditor/Source/UVEditorTools/Private/Operators/UVEditorUVLayoutOp.cpp b/Engine/Plugins/Experimental/UVEditor/Source/UVEditorTools/Private/Operators/UVEditorUVLayoutOp.cpp
deleted file mode 100644
diff --git a/Engine/Plugins/Experimental/UVEditor/Source/UVEditorTools/Private/Operators/UVEditorUVTransformOp.cpp b/Engine/Plugins/Experimental/UVEditor/Source/UVEditorTools/Private/Operators/UVEditorUVTransformOp.cpp
--- a/Engine/Plugins/Experimental/UVEditor/Source/UVEditorTools/Private/Operators/UVEditorUVTransformOp.cpp
+++ b/Engine/Plugins/Experimental/UVEditor/Source/UVEditorTools/Private/Operators/UVEditorUVTransformOp.cpp
@@ -6,7 +6,7 @@
 #include "DynamicMesh/DynamicMeshAttributeSet.h"
 #include "Parameterization/MeshUVPacking.h"
 #include "Selections/MeshConnectedComponents.h"
-#include "Utilities/MeshUDIMClassifier.h"
+#include "Parameterization/MeshUDIMClassifier.h"
 #include "Async/ParallelFor.h"
 #include "UVEditorUXSettings.h"
 
diff --git a/Engine/Plugins/Experimental/UVEditor/Source/UVEditorTools/Private/UVEditorLayoutTool.cpp b/Engine/Plugins/Experimental/UVEditor/Source/UVEditorTools/Private/UVEditorLayoutTool.cpp
--- a/Engine/Plugins/Experimental/UVEditor/Source/UVEditorTools/Private/UVEditorLayoutTool.cpp
+++ b/Engine/Plugins/Experimental/UVEditor/Source/UVEditorTools/Private/UVEditorLayoutTool.cpp
@@ -7,7 +7,8 @@
 #include "DynamicMesh/DynamicMeshChangeTracker.h"
 #include "InteractiveToolManager.h"
 #include "MeshOpPreviewHelpers.h" // UMeshOpPreviewWithBackgroundCompute
-#include "Operators/UVEditorUVLayoutOp.h"
+#include "ParameterizationOps/UVLayoutOp.h"
+#include "Properties/UVLayoutProperties.h"
 #include "ToolTargets/UVEditorToolMeshInput.h"
 #include "ContextObjects/UVToolContextObjects.h"
 #include "EngineAnalytics.h"
@@ -45,7 +46,7 @@ void UUVEditorLayoutTool::Setup()
 
 	UInteractiveTool::Setup();
 
-	Settings = NewObject<UUVEditorUVLayoutProperties>(this);
+	Settings = NewObject<UUVLayoutProperties>(this);
 	Settings->RestoreProperties(this);
 	Settings->bUDIMCVAREnabled = (FUVEditorUXSettings::CVarEnablePrototypeUDIMSupport.GetValueOnGameThread() > 0);
 	AddToolPropertySource(Settings);
@@ -61,7 +62,7 @@ void UUVEditorLayoutTool::Setup()
 
 	auto SetupOpFactory = [this](UUVEditorToolMeshInput& Target, const FUVToolSelection* Selection)
 	{
-		TObjectPtr<UUVEditorUVLayoutOperatorFactory> Factory = NewObject<UUVEditorUVLayoutOperatorFactory>();
+		TObjectPtr<UUVLayoutOperatorFactory> Factory = NewObject<UUVLayoutOperatorFactory>();
 		Factory->TargetTransform = Target.AppliedPreview->PreviewMesh->GetTransform();
 		Factory->Settings = Settings;
 		Factory->OriginalMesh = Target.AppliedCanonical;
diff --git a/Engine/Plugins/Experimental/UVEditor/Source/UVEditorTools/Private/Utilities/MeshUDIMClassifier.cpp b/Engine/Plugins/Experimental/UVEditor/Source/UVEditorTools/Private/Utilities/MeshUDIMClassifier.cpp
deleted file mode 100644
diff --git a/Engine/Plugins/Experimental/UVEditor/Source/UVEditorTools/Public/Operators/UVEditorUVLayoutOp.h b/Engine/Plugins/Experimental/UVEditor/Source/UVEditorTools/Public/Operators/UVEditorUVLayoutOp.h
deleted file mode 100644
diff --git a/Engine/Plugins/Experimental/UVEditor/Source/UVEditorTools/Public/UVEditorLayoutTool.h b/Engine/Plugins/Experimental/UVEditor/Source/UVEditorTools/Public/UVEditorLayoutTool.h
--- a/Engine/Plugins/Experimental/UVEditor/Source/UVEditorTools/Public/UVEditorLayoutTool.h
+++ b/Engine/Plugins/Experimental/UVEditor/Source/UVEditorTools/Public/UVEditorLayoutTool.h
@@ -12,8 +12,8 @@
 #include "UVEditorLayoutTool.generated.h"
 
 class UUVEditorToolMeshInput;
-class UUVEditorUVLayoutProperties;
-class UUVEditorUVLayoutOperatorFactory;
+class UUVLayoutProperties;
+class UUVLayoutOperatorFactory;
 
 UCLASS()
 class UVEDITORTOOLS_API UUVEditorLayoutToolBuilder : public UInteractiveToolBuilder
@@ -63,10 +63,10 @@ protected:
 	TArray<TObjectPtr<UUVEditorToolMeshInput>> Targets;
 
 	UPROPERTY()
-	TObjectPtr<UUVEditorUVLayoutProperties> Settings = nullptr;
+	TObjectPtr<UUVLayoutProperties> Settings = nullptr;
 
 	UPROPERTY()
-	TArray<TObjectPtr<UUVEditorUVLayoutOperatorFactory>> Factories;
+	TArray<TObjectPtr<UUVLayoutOperatorFactory>> Factories;
 
 	UPROPERTY()
 	TObjectPtr<UUVToolSelectionAPI> UVToolSelectionAPI = nullptr;
diff --git a/Engine/Plugins/Experimental/UVEditor/Source/UVEditorTools/Public/Utilities/MeshUDIMClassifier.h b/Engine/Plugins/Experimental/UVEditor/Source/UVEditorTools/Public/Utilities/MeshUDIMClassifier.h
deleted file mode 100644
diff --git a/Engine/Plugins/Runtime/GeometryProcessing/Source/DynamicMesh/Private/Parameterization/MeshUDIMClassifier.cpp b/Engine/Plugins/Runtime/GeometryProcessing/Source/DynamicMesh/Private/Parameterization/MeshUDIMClassifier.cpp
new file mode 100644
--- /dev/null
+++ b/Engine/Plugins/Runtime/GeometryProcessing/Source/DynamicMesh/Private/Parameterization/MeshUDIMClassifier.cpp
@@ -0,0 +1,138 @@
+// Copyright Epic Games, Inc. All Rights Reserved.
+
+#include "Parameterization/MeshUDIMClassifier.h"
+#include "Spatial/SparseDynamicOctree3.h"
+#include "Selections/MeshConnectedComponents.h"
+
+using namespace UE::Geometry;
+
+namespace UVLayoutOpLocals
+{
+	FVector2f InternalUVToExternalUV(const FVector2f& UV)
+	{
+		return FVector2f(UV.X, 1 - UV.Y);
+	}
+}
+
+FDynamicMeshUDIMClassifier::FDynamicMeshUDIMClassifier(const FDynamicMeshUVOverlay* UVOverlayIn, TOptional<TArray<int32>> SelectionIn)
+{
+	UVOverlay = UVOverlayIn;
+	Selection = SelectionIn;
+	ClassifyUDIMs();
+}
+
+TArray<FVector2i> FDynamicMeshUDIMClassifier::ActiveTiles() const
+{
+	TArray<FVector2i> Keys;
+	UDIMs.GetKeys(Keys);
+	return Keys;
+}
+
+TArray<int32> FDynamicMeshUDIMClassifier::TidsForTile(FVector2i TileIndexIn) const
+{	
+	if (ensure(UDIMs.Contains(TileIndexIn)))
+	{
+		return *(UDIMs.Find(TileIndexIn));
+	}
+	else
+	{
+		return TArray<int32>();
+	}
+}
+
+FVector2i FDynamicMeshUDIMClassifier::ClassifyTrianglesToUDIM(const FDynamicMeshUVOverlay* UVOverlay, TArray<int32> Tids)
+{
+	auto UVTriangleToUDIM = [UVOverlay](int32 Tid)
+	{
+		FVector2i UDIM;
+		FVector2f Vertex0, Vertex1, Vertex2;
+		UVOverlay->GetTriElements(Tid, Vertex0, Vertex1, Vertex2);
+		FVector2f BaryCenter = (Vertex0 + Vertex1 + Vertex2) / 3.0f;
+		BaryCenter = UVLayoutOpLocals::InternalUVToExternalUV(BaryCenter);
+		UDIM = ClassifyPointToUDIM(BaryCenter);		
+		return UDIM;
+	};
+
+	TMap<FVector2i, int32> UDIMCount;
+	for (int32 Tid : Tids)
+	{
+		FVector2i UDIM = UVTriangleToUDIM(Tid);
+		UDIMCount.FindOrAdd(UDIM, 0)++;
+	}
+	FVector2i MaximumUDIM;
+	int32 MaximumUDIMCount = 0;
+	for (auto UDIMEntry : UDIMCount)
+	{
+		if (UDIMEntry.Value > MaximumUDIMCount)
+		{
+			MaximumUDIMCount = UDIMEntry.Value;
+			MaximumUDIM = UDIMEntry.Key;
+		}
+	}
+	return MaximumUDIM;
+}
+
+FVector2i FDynamicMeshUDIMClassifier::ClassifyBoundingBoxToUDIM(const FDynamicMeshUVOverlay* UVOverlay, const FAxisAlignedBox2d& BoundingBox)
+{
+	FVector2i MinUDIM = ClassifyPointToUDIM((FVector2f)BoundingBox.Min);
+	double MaxOverlapArea = 0.0;
+	FVector2i SelectedUDIM = MinUDIM;
+
+	// We only really need to test the minimum UDIM tile and one greater in each dimension.
+	// Either the BoundingBox will be fully overlapping in any or none of these, but
+	// we don't need to find the largest indexed UDIM with full overlap, just the first.
+	for (int32 UDIM_X = MinUDIM.X; UDIM_X <= MinUDIM.X+1; ++UDIM_X)
+	{
+		for (int32 UDIM_Y = MinUDIM.Y; UDIM_Y <= MinUDIM.Y+1; ++UDIM_Y)
+		{
+			FAxisAlignedBox2d TestUDIMBoundingBox(FVector2d(UDIM_X, UDIM_Y), FVector2d(UDIM_X + 1, UDIM_Y + 1));
+			double OverlapArea = BoundingBox.Intersect(TestUDIMBoundingBox).Area();
+			if (OverlapArea > MaxOverlapArea)
+			{
+				SelectedUDIM = FVector2i(UDIM_X, UDIM_Y);
+				MaxOverlapArea = OverlapArea;
+			}
+		}
+	}
+
+	return SelectedUDIM;
+}
+
+
+FVector2i FDynamicMeshUDIMClassifier::ClassifyPointToUDIM(const FVector2f& UVPoint)
+{
+	FVector2i UDIM;
+	UDIM.X = FMath::FloorToInt32(UVPoint.X);
+	UDIM.Y = FMath::FloorToInt32(UVPoint.Y);
+	return UDIM;
+}
+
+void FDynamicMeshUDIMClassifier::ClassifyUDIMs()
+{
+	int32 CurrentUDIMIndex = 0;
+	const FDynamicMesh3* Mesh = UVOverlay->GetParentMesh();
+
+	auto UVIslandPredicate = [this](int32 Triangle0, int32 Triangle1)
+	{
+		return UVOverlay->AreTrianglesConnected(Triangle0, Triangle1);
+	};
+
+	FMeshConnectedComponents UVComponents(Mesh);
+	if (Selection.IsSet())
+	{
+		UVComponents.FindConnectedTriangles(Selection.GetValue(), UVIslandPredicate);
+	}
+	else
+	{
+		UVComponents.FindConnectedTriangles(UVIslandPredicate);
+	}
+
+	for (int32 Cid = 0; Cid < UVComponents.Num(); ++Cid)
+	{
+		FVector2i UDIM = ClassifyTrianglesToUDIM(UVOverlay, UVComponents.GetComponent(Cid).Indices);
+		UDIMs.FindOrAdd(UDIM).Append(UVComponents.GetComponent(Cid).Indices);
+	}
+
+	return;
+
+}
diff --git a/Engine/Plugins/Runtime/GeometryProcessing/Source/DynamicMesh/Public/Parameterization/MeshUDIMClassifier.h b/Engine/Plugins/Runtime/GeometryProcessing/Source/DynamicMesh/Public/Parameterization/MeshUDIMClassifier.h
new file mode 100644
--- /dev/null
+++ b/Engine/Plugins/Runtime/GeometryProcessing/Source/DynamicMesh/Public/Parameterization/MeshUDIMClassifier.h
@@ -0,0 +1,44 @@
+// Copyright Epic Games, Inc. All Rights Reserved.
+
+#pragma once
+
+#include "DynamicMesh/DynamicMesh3.h"
+#include "DynamicMesh/DynamicMeshAttributeSet.h"
+
+
+namespace UE
+{
+	namespace Geometry
+	{
+
+		struct FUVOverlayView;
+		class FMeshConnectedComponents;
+
+		/**
+		* FDynamicMeshUDIMClassifier is a utility class for identifying active UDIMs from a FDynamicMesh's UV overlay
+		 */
+		class DYNAMICMESH_API FDynamicMeshUDIMClassifier
+		{
+		public:			
+			explicit FDynamicMeshUDIMClassifier(const FDynamicMeshUVOverlay* UVOverlay, TOptional<TArray<int32>> Selection = TOptional<TArray<int32>>());
+
+			TArray<FVector2i> ActiveTiles() const;
+			TArray<int32> TidsForTile(FVector2i TileIndexIn) const;
+
+			static FVector2i ClassifyTrianglesToUDIM(const FDynamicMeshUVOverlay* UVOverlay, TArray<int32> Tids);
+			static FVector2i ClassifyBoundingBoxToUDIM(const FDynamicMeshUVOverlay* UVOverlay, const FAxisAlignedBox2d& BoundingBox);
+			static FVector2i ClassifyPointToUDIM(const FVector2f& UVPoint);
+
+		protected:
+
+			void ClassifyUDIMs();
+
+			/** The UV Overlay to analyze for UDIMs */
+			const FDynamicMeshUVOverlay* UVOverlay = nullptr;
+			
+			TOptional<TArray<int32>> Selection;
+			
+			TMap<FVector2i, TArray<int32> > UDIMs;
+		};
+	}
+}
\ No newline at end of file
diff --git a/Engine/Plugins/Runtime/MeshModelingToolset/Source/MeshModelingTools/Private/UVLayoutTool.cpp b/Engine/Plugins/Runtime/MeshModelingToolset/Source/MeshModelingTools/Private/UVLayoutTool.cpp
--- a/Engine/Plugins/Runtime/MeshModelingToolset/Source/MeshModelingTools/Private/UVLayoutTool.cpp
+++ b/Engine/Plugins/Runtime/MeshModelingToolset/Source/MeshModelingTools/Private/UVLayoutTool.cpp
@@ -72,6 +72,8 @@ void UUVLayoutTool::Setup()
 
 	BasicProperties = NewObject<UUVLayoutProperties>(this);
 	BasicProperties->RestoreProperties(this);
+	BasicProperties->bUDIMCVAREnabled = false;
+	BasicProperties->bEnableUDIMLayout = false;
 	AddToolPropertySource(BasicProperties);
 
 	MaterialSettings = NewObject<UExistingMeshMaterialProperties>(this);
diff --git a/Engine/Plugins/Runtime/MeshModelingToolset/Source/ModelingOperators/Private/ParameterizationOps/UVLayoutOp.cpp b/Engine/Plugins/Runtime/MeshModelingToolset/Source/ModelingOperators/Private/ParameterizationOps/UVLayoutOp.cpp
--- a/Engine/Plugins/Runtime/MeshModelingToolset/Source/ModelingOperators/Private/ParameterizationOps/UVLayoutOp.cpp
+++ b/Engine/Plugins/Runtime/MeshModelingToolset/Source/ModelingOperators/Private/ParameterizationOps/UVLayoutOp.cpp
@@ -7,11 +7,26 @@
 #include "Parameterization/MeshUVPacking.h"
 #include "Properties/UVLayoutProperties.h"
 #include "Selections/MeshConnectedComponents.h"
+#include "Parameterization/MeshUDIMClassifier.h"
 
 #include UE_INLINE_GENERATED_CPP_BY_NAME(UVLayoutOp)
 
 using namespace UE::Geometry;
 
+namespace UVLayoutOpLocals
+{
+	FVector2f ExternalUVToInternalUV(const FVector2f& UV)
+	{
+		return FVector2f(UV.X, 1 - UV.Y);
+	}
+
+	FVector2f InternalUVToExternalUV(const FVector2f& UV)
+	{
+		return FVector2f(UV.X, 1 - UV.Y);
+	}
+}
+
+
 void FUVLayoutOp::SetTransform(const FTransformSRT3d& Transform) 
 {
 	ResultTransform = Transform;
@@ -56,7 +71,88 @@ void FUVLayoutOp::CalculateResult(FProgressCancel* Progress)
 		return;
 	}
 
-	FDynamicMeshUVPacker Packer(UseUVLayer);
+	if (bMaintainOriginatingUDIM)
+	{		
+		TOptional<TArray<int32>> SelectionArray;
+		if (Selection.IsSet())
+		{
+			SelectionArray = Selection.GetValue().Array();
+		}
+		FDynamicMeshUDIMClassifier TileClassifier(UseUVLayer, SelectionArray);
+
+		TArray<FVector2i> Tiles = TileClassifier.ActiveTiles();
+
+		for (FVector2i TileIndex : Tiles)
+		{
+			TUniquePtr<TArray<int32>> TileTids;
+			TileTids = MakeUnique<TArray<int32>>(TileClassifier.TidsForTile(TileIndex));
+
+			// Do this first, so we don't need to keep the TileTids around after moving it into the packer.
+			TSet<int32> ElementsToMove;
+			ElementsToMove.Reserve(TileTids->Num() * 3);
+			for (int Tid : *TileTids)
+			{
+				FIndex3i Elements = UseUVLayer->GetTriangle(Tid);
+				ElementsToMove.Add(Elements[0]);
+				ElementsToMove.Add(Elements[1]);
+				ElementsToMove.Add(Elements[2]);
+			}
+
+			// TODO: There is a second connected components call inside the packer that might be unnessessary. Could be a future optimization.
+			FDynamicMeshUVPacker Packer(UseUVLayer, MoveTemp(TileTids) );
+			ExecutePacker(Packer);
+			if (Progress && Progress->Cancelled())
+			{
+				return;
+			}
+
+			for (int32 Element : ElementsToMove)
+			{
+				FVector2f UV = UVLayoutOpLocals::InternalUVToExternalUV(UseUVLayer->GetElement(Element));
+				UV = UV + FVector2f(TileIndex);
+				UseUVLayer->SetElement(Element, UVLayoutOpLocals::ExternalUVToInternalUV(UV));
+			}			
+		}
+	}
+	else
+	{
+		TUniquePtr<TArray<int32>> TidsToLayout = nullptr;
+		if (Selection.IsSet())
+		{
+			TidsToLayout = MakeUnique<TArray<int32>>();
+			TSet<int32>& SelectionSet = Selection.GetValue();
+			TidsToLayout->Reserve(SelectionSet.Num());
+			for(int Tid : SelectionSet)
+			{
+				if (OriginalMesh->IsTriangle(Tid))
+				{
+					TidsToLayout->Add(Tid);
+				}
+			}
+		}
+		
+		FDynamicMeshUVPacker Packer(UseUVLayer, MoveTemp(TidsToLayout) );
+		ExecutePacker(Packer);
+		if (Progress && Progress->Cancelled())
+		{
+			return;
+		}
+	}
+
+	if (UVScaleFactor != 1.0 || UVTranslation != FVector2f::Zero() )
+	{
+		for (int ElementID : UseUVLayer->ElementIndicesItr())
+		{
+			FVector2f UV = UVLayoutOpLocals::InternalUVToExternalUV(UseUVLayer->GetElement(ElementID));
+			UV = (UV * UVScaleFactor) + UVTranslation;
+			UseUVLayer->SetElement(ElementID, UVLayoutOpLocals::ExternalUVToInternalUV(UV));
+		}
+	}
+
+}
+
+void FUVLayoutOp::ExecutePacker(FDynamicMeshUVPacker& Packer)
+{
 	Packer.TextureResolution = this->TextureResolution;
 	Packer.GutterSize = this->GutterSize;
 	Packer.bAllowFlips = this->bAllowFlips;
@@ -77,22 +173,6 @@ void FUVLayoutOp::CalculateResult(FProgressCancel* Progress)
 			return;
 		}
 	}
-
-	if (Progress && Progress->Cancelled())
-	{
-		return;
-	}
-
-	if (UVScaleFactor != 1.0 || UVTranslation != FVector2f::Zero() )
-	{
-		for (int ElementID : UseUVLayer->ElementIndicesItr())
-		{
-			FVector2f UV = UseUVLayer->GetElement(ElementID);
-			UV = (UV * UVScaleFactor) + UVTranslation;
-			UseUVLayer->SetElement(ElementID, UV);
-		}
-	}
-
 }
 
 TUniquePtr<FDynamicMeshOperator> UUVLayoutOperatorFactory::MakeNewOperator()
@@ -120,6 +200,7 @@ TUniquePtr<FDynamicMeshOperator> UUVLayoutOperatorFactory::MakeNewOperator()
 	Op->UVScaleFactor = Settings->Scale;
 	Op->UVTranslation = FVector2f(Settings->Translation);
 	Op->SetTransform(TargetTransform);
+	Op->bMaintainOriginatingUDIM = Settings->bEnableUDIMLayout;
 
 	return Op;
 }
diff --git a/Engine/Plugins/Runtime/MeshModelingToolset/Source/ModelingOperators/Public/ParameterizationOps/UVLayoutOp.h b/Engine/Plugins/Runtime/MeshModelingToolset/Source/ModelingOperators/Public/ParameterizationOps/UVLayoutOp.h
--- a/Engine/Plugins/Runtime/MeshModelingToolset/Source/ModelingOperators/Public/ParameterizationOps/UVLayoutOp.h
+++ b/Engine/Plugins/Runtime/MeshModelingToolset/Source/ModelingOperators/Public/ParameterizationOps/UVLayoutOp.h
@@ -15,6 +15,7 @@ namespace UE
 namespace Geometry
 {
 	class FDynamicMesh3;
+	class FDynamicMeshUVPacker;
 
 enum class EUVLayoutOpLayoutModes
 {
@@ -39,6 +40,8 @@ public:
 	bool bAlwaysSplitBowties = true;
 	float UVScaleFactor = 1.0;
 	float GutterSize = 1.0;
+	bool bMaintainOriginatingUDIM = false;
+	TOptional<TSet<int32>> Selection;
 	FVector2f UVTranslation = FVector2f::Zero();
 
 	void SetTransform(const FTransformSRT3d& Transform);
@@ -48,6 +51,9 @@ public:
 	// 
 
 	virtual void CalculateResult(FProgressCancel* Progress) override;
+
+protected:
+	void ExecutePacker(FDynamicMeshUVPacker& Packer);
 };
 
 } // end namespace UE::Geometry
@@ -70,7 +76,7 @@ public:
 
 	UPROPERTY()
 	TObjectPtr<UUVLayoutProperties> Settings;
-
+	TOptional<TSet<int32>> Selection;
 	TSharedPtr<UE::Geometry::FDynamicMesh3, ESPMode::ThreadSafe> OriginalMesh;
 	TUniqueFunction<int32()> GetSelectedUVChannel = []() { return 0; };
 	FTransform TargetTransform;
diff --git a/Engine/Plugins/Runtime/MeshModelingToolset/Source/ModelingOperators/Public/Properties/UVLayoutProperties.h b/Engine/Plugins/Runtime/MeshModelingToolset/Source/ModelingOperators/Public/Properties/UVLayoutProperties.h
--- a/Engine/Plugins/Runtime/MeshModelingToolset/Source/ModelingOperators/Public/Properties/UVLayoutProperties.h
+++ b/Engine/Plugins/Runtime/MeshModelingToolset/Source/ModelingOperators/Public/Properties/UVLayoutProperties.h
@@ -51,4 +51,10 @@ public:
 	UPROPERTY(EditAnywhere, Category = "UV Layout", meta = (EditCondition = "LayoutType == EUVLayoutType::Repack"))
 	bool bAllowFlips = false;
 
+	/** Enable UDIM aware layout and keep islands within their originating UDIM tiles when laying out.*/
+	UPROPERTY(EditAnywhere, Category = "UV Layout", meta = (DisplayName = "Preserve UDIMs", EditCondition="bUDIMCVAREnabled", EditConditionHides, HideEditConditionToggle = true))
+	bool bEnableUDIMLayout = false;
+
+	UPROPERTY(Transient)
+	bool bUDIMCVAREnabled = false;
 };
