ModelingComponents: fix UE::AssetUtils::GetStaticMeshLODMaterialListBySection, it was returning incorrect results for static meshes with re-ordered sections/slots.

diff --git a/Engine/Plugins/Runtime/MeshModelingToolset/Source/ModelingComponents/Private/AssetUtils/StaticMeshMaterialUtil.cpp b/Engine/Plugins/Runtime/MeshModelingToolset/Source/ModelingComponents/Private/AssetUtils/StaticMeshMaterialUtil.cpp
--- a/Engine/Plugins/Runtime/MeshModelingToolset/Source/ModelingComponents/Private/AssetUtils/StaticMeshMaterialUtil.cpp
+++ b/Engine/Plugins/Runtime/MeshModelingToolset/Source/ModelingComponents/Private/AssetUtils/StaticMeshMaterialUtil.cpp
@@ -56,7 +56,7 @@ bool UE::AssetUtils::GetStaticMeshLODAssetMaterials(
 		MaterialInfoOut.SectionSlotIndexes[SectionIndex] = -1;
 		MaterialInfoOut.SectionMaterials[SectionIndex] = nullptr;
 
-		if (SectionInfoMap.IsValidSection(LODIndex, SectionIndex))
+		if (SectionInfoMap.IsValidSection(LODIndex, SectionIndex) == false)
 		{
 			// did not find this section 
 			if ( StaticMaterials.IsValidIndex(SectionIndex) )
