UVEditor: Merge and unify the RecomputeUVs codebases with the UVEditor and Modeling Mode.

diff --git a/Engine/Plugins/Experimental/UVEditor/Source/UVEditorTools/Private/Operators/UVEditorRecomputeUVsOp.cpp b/Engine/Plugins/Experimental/UVEditor/Source/UVEditorTools/Private/Operators/UVEditorRecomputeUVsOp.cpp
deleted file mode 100644
diff --git a/Engine/Plugins/Experimental/UVEditor/Source/UVEditorTools/Private/UVEditorRecomputeUVsTool.cpp b/Engine/Plugins/Experimental/UVEditor/Source/UVEditorTools/Private/UVEditorRecomputeUVsTool.cpp
--- a/Engine/Plugins/Experimental/UVEditor/Source/UVEditorTools/Private/UVEditorRecomputeUVsTool.cpp
+++ b/Engine/Plugins/Experimental/UVEditor/Source/UVEditorTools/Private/UVEditorRecomputeUVsTool.cpp
@@ -11,7 +11,8 @@
 #include "MeshOpPreviewHelpers.h"
 #include "UVEditorToolAnalyticsUtils.h"
 #include "EngineAnalytics.h"
-#include "Operators/UVEditorRecomputeUVsOp.h"
+#include "ParameterizationOps/RecomputeUVsOp.h"
+#include "Properties/RecomputeUVsProperties.h"
 #include "UVEditorUXSettings.h"
 
 #include UE_INLINE_GENERATED_CPP_BY_NAME(UVEditorRecomputeUVsTool)
@@ -52,7 +53,7 @@ void UUVEditorRecomputeUVsTool::Setup()
 
 	// initialize our properties
 
-	Settings = NewObject<UUVEditorRecomputeUVsToolProperties>(this);
+	Settings = NewObject<URecomputeUVsToolProperties>(this);
 	Settings->RestoreProperties(this);
 	Settings->bUDIMCVAREnabled = (FUVEditorUXSettings::CVarEnablePrototypeUDIMSupport.GetValueOnGameThread() > 0);
 	AddToolPropertySource(Settings);
@@ -79,13 +80,13 @@ void UUVEditorRecomputeUVsTool::Setup()
 	else
 	{
 		Settings->bEnablePolygroupSupport = false;
-		Settings->IslandGeneration = EUVEditorRecomputeUVsPropertiesIslandMode::ExistingUVs;
+		Settings->IslandGeneration = ERecomputeUVsPropertiesIslandMode::ExistingUVs;
 	}
 	UpdateActiveGroupLayer(false);  /* Don't update factories that don't exist yet. */
 
 	auto SetupOpFactory = [this](UUVEditorToolMeshInput& Target, const FUVToolSelection* Selection)
 	{
-		TObjectPtr<UUVEditorRecomputeUVsOpFactory> Factory = NewObject<UUVEditorRecomputeUVsOpFactory>();
+		TObjectPtr<URecomputeUVsOpFactory> Factory = NewObject<URecomputeUVsOpFactory>();
 		Factory->TargetTransform = (FTransform3d)Target.AppliedPreview->PreviewMesh->GetTransform();
 		Factory->Settings = Settings;
 		Factory->OriginalMesh = Target.AppliedCanonical;
@@ -311,23 +312,23 @@ void UUVEditorRecomputeUVsTool::RecordAnalytics()
 	Attributes.Add(AnalyticsEventAttributeEnum(TEXT("Settings.AutoRotation"), Settings->AutoRotation));
 	
 	Attributes.Add(AnalyticsEventAttributeEnum(TEXT("Settings.UnwrapType"), Settings->UnwrapType));
-	if (Settings->UnwrapType == EUVEditorRecomputeUVsPropertiesUnwrapType::IslandMerging || Settings->UnwrapType == EUVEditorRecomputeUVsPropertiesUnwrapType::ExpMap)
+	if (Settings->UnwrapType == ERecomputeUVsPropertiesUnwrapType::IslandMerging || Settings->UnwrapType == ERecomputeUVsPropertiesUnwrapType::ExpMap)
 	{
 		Attributes.Add(FAnalyticsEventAttribute(TEXT("Settings.SmoothingSteps"), Settings->SmoothingSteps));
 		Attributes.Add(FAnalyticsEventAttribute(TEXT("Settings.SmoothingAlpha"), Settings->SmoothingAlpha));
 	}
-	if (Settings->UnwrapType == EUVEditorRecomputeUVsPropertiesUnwrapType::IslandMerging)
+	if (Settings->UnwrapType == ERecomputeUVsPropertiesUnwrapType::IslandMerging)
 	{
 		Attributes.Add(FAnalyticsEventAttribute(TEXT("Settings.MergingDistortionThreshold"), Settings->MergingDistortionThreshold));
 		Attributes.Add(FAnalyticsEventAttribute(TEXT("Settings.MergingAngleThreshold"), Settings->MergingAngleThreshold));
 	}
 	
 	Attributes.Add(AnalyticsEventAttributeEnum(TEXT("Settings.LayoutType"), Settings->LayoutType));
-	if (Settings->LayoutType == EUVEditorRecomputeUVsPropertiesLayoutType::Repack)
+	if (Settings->LayoutType == ERecomputeUVsPropertiesLayoutType::Repack)
 	{
 		Attributes.Add(FAnalyticsEventAttribute(TEXT("Settings.TextureResolution"), Settings->TextureResolution));
 	}
-	if (Settings->LayoutType == EUVEditorRecomputeUVsPropertiesLayoutType::NormalizeToBounds || Settings->LayoutType == EUVEditorRecomputeUVsPropertiesLayoutType::NormalizeToWorld)
+	if (Settings->LayoutType == ERecomputeUVsPropertiesLayoutType::NormalizeToBounds || Settings->LayoutType == ERecomputeUVsPropertiesLayoutType::NormalizeToWorld)
 	{
 		Attributes.Add(FAnalyticsEventAttribute(TEXT("Settings.NormalizeScale"), Settings->NormalizeScale));
 	}
diff --git a/Engine/Plugins/Experimental/UVEditor/Source/UVEditorTools/Public/Operators/UVEditorRecomputeUVsOp.h b/Engine/Plugins/Experimental/UVEditor/Source/UVEditorTools/Public/Operators/UVEditorRecomputeUVsOp.h
deleted file mode 100644
diff --git a/Engine/Plugins/Experimental/UVEditor/Source/UVEditorTools/Public/UVEditorRecomputeUVsTool.h b/Engine/Plugins/Experimental/UVEditor/Source/UVEditorTools/Public/UVEditorRecomputeUVsTool.h
--- a/Engine/Plugins/Experimental/UVEditor/Source/UVEditorTools/Public/UVEditorRecomputeUVsTool.h
+++ b/Engine/Plugins/Experimental/UVEditor/Source/UVEditorTools/Public/UVEditorRecomputeUVsTool.h
@@ -19,8 +19,8 @@ class UDynamicMeshComponent;
 class UMaterialInterface;
 class UMaterialInstanceDynamic;
 class UUVEditorToolMeshInput;
-class UUVEditorRecomputeUVsToolProperties;
-class UUVEditorRecomputeUVsOpFactory;
+class URecomputeUVsToolProperties;
+class URecomputeUVsOpFactory;
 
 /**
  *
@@ -76,13 +76,13 @@ protected:
 	TArray<TObjectPtr<UUVEditorToolMeshInput>> Targets;
 
 	UPROPERTY()
-	TObjectPtr<UUVEditorRecomputeUVsToolProperties> Settings = nullptr;
+	TObjectPtr<URecomputeUVsToolProperties> Settings = nullptr;
 
 	UPROPERTY()
 	TObjectPtr<UPolygroupLayersProperties> PolygroupLayerProperties = nullptr;
 
 	UPROPERTY()
-	TArray<TObjectPtr<UUVEditorRecomputeUVsOpFactory>> Factories;
+	TArray<TObjectPtr<URecomputeUVsOpFactory>> Factories;
 
 	UPROPERTY()
 	TObjectPtr<UUVToolSelectionAPI> UVToolSelectionAPI = nullptr;
diff --git a/Engine/Plugins/Experimental/UVEditor/Source/UVEditorToolsEditorOnly/Private/DetailsCustomizations/UVUnwrapToolCustomizations.cpp b/Engine/Plugins/Experimental/UVEditor/Source/UVEditorToolsEditorOnly/Private/DetailsCustomizations/UVUnwrapToolCustomizations.cpp
deleted file mode 100644
diff --git a/Engine/Plugins/Experimental/UVEditor/Source/UVEditorToolsEditorOnly/Private/UVEditorToolsEditorOnlyModule.cpp b/Engine/Plugins/Experimental/UVEditor/Source/UVEditorToolsEditorOnly/Private/UVEditorToolsEditorOnlyModule.cpp
--- a/Engine/Plugins/Experimental/UVEditor/Source/UVEditorToolsEditorOnly/Private/UVEditorToolsEditorOnlyModule.cpp
+++ b/Engine/Plugins/Experimental/UVEditor/Source/UVEditorToolsEditorOnly/Private/UVEditorToolsEditorOnlyModule.cpp
@@ -5,10 +5,8 @@
 #include "Modules/ModuleManager.h"
 #include "PropertyEditorModule.h"
 
-#include "DetailsCustomizations/UVUnwrapToolCustomizations.h"
 #include "DetailsCustomizations/UVTransformToolCustomizations.h"
 
-#include "Operators/UVEditorRecomputeUVsOp.h"
 #include "UVEditorTransformTool.h"
 #include "Operators/UVEditorUVTransformOp.h"
 
@@ -51,10 +49,6 @@ void FUVEditorToolsEditorOnlyModule::OnPostEngineInit()
 	// Register details view customizations
 	FPropertyEditorModule& PropertyModule = FModuleManager::LoadModuleChecked<FPropertyEditorModule>("PropertyEditor");
 
-	/// Unwrap
-	PropertyModule.RegisterCustomClassLayout("UVEditorRecomputeUVsToolProperties", FOnGetDetailCustomizationInstance::CreateStatic(&FUVEditorRecomputeUVsToolDetails::MakeInstance));
-	ClassesToUnregisterOnShutdown.Add(UUVEditorRecomputeUVsToolProperties::StaticClass()->GetFName());
-
 	// Transform
 	PropertyModule.RegisterCustomClassLayout(UUVEditorUVTransformProperties::StaticClass()->GetFName(),
 		                                     FOnGetDetailCustomizationInstance::CreateStatic(&FUVEditorUVTransformToolDetails::MakeInstance));
diff --git a/Engine/Plugins/Experimental/UVEditor/Source/UVEditorToolsEditorOnly/Public/DetailsCustomizations/UVUnwrapToolCustomizations.h b/Engine/Plugins/Experimental/UVEditor/Source/UVEditorToolsEditorOnly/Public/DetailsCustomizations/UVUnwrapToolCustomizations.h
deleted file mode 100644
diff --git a/Engine/Plugins/Experimental/UVEditor/Source/UVEditorToolsEditorOnly/UVEditorToolsEditorOnly.Build.cs b/Engine/Plugins/Experimental/UVEditor/Source/UVEditorToolsEditorOnly/UVEditorToolsEditorOnly.Build.cs
--- a/Engine/Plugins/Experimental/UVEditor/Source/UVEditorToolsEditorOnly/UVEditorToolsEditorOnly.Build.cs
+++ b/Engine/Plugins/Experimental/UVEditor/Source/UVEditorToolsEditorOnly/UVEditorToolsEditorOnly.Build.cs
@@ -52,6 +52,7 @@ public class UVEditorToolsEditorOnly : ModuleRules
 				"ModelingComponents",
 				"ModelingOperators",
 				"ModelingOperatorsEditorOnly",
+				"MeshModelingToolsEditorOnly",
 				"UVEditorTools",
 
 				"ToolWidgets",
diff --git a/Engine/Plugins/Runtime/MeshModelingToolset/Source/MeshModelingTools/Private/RecomputeUVsTool.cpp b/Engine/Plugins/Runtime/MeshModelingToolset/Source/MeshModelingTools/Private/RecomputeUVsTool.cpp
--- a/Engine/Plugins/Runtime/MeshModelingToolset/Source/MeshModelingTools/Private/RecomputeUVsTool.cpp
+++ b/Engine/Plugins/Runtime/MeshModelingToolset/Source/MeshModelingTools/Private/RecomputeUVsTool.cpp
@@ -45,6 +45,8 @@ void URecomputeUVsTool::Setup()
 
 	Settings = NewObject<URecomputeUVsToolProperties>(this);
 	Settings->RestoreProperties(this);
+	Settings->bUDIMCVAREnabled = false;
+	Settings->bEnableUDIMLayout = false;
 	AddToolPropertySource(Settings);
 
 	UVChannelProperties = NewObject<UMeshUVChannelProperties>(this);
diff --git a/Engine/Plugins/Runtime/MeshModelingToolset/Source/MeshModelingToolsEditorOnly/MeshModelingToolsEditorOnly.Build.cs b/Engine/Plugins/Runtime/MeshModelingToolset/Source/MeshModelingToolsEditorOnly/MeshModelingToolsEditorOnly.Build.cs
--- a/Engine/Plugins/Runtime/MeshModelingToolset/Source/MeshModelingToolsEditorOnly/MeshModelingToolsEditorOnly.Build.cs
+++ b/Engine/Plugins/Runtime/MeshModelingToolset/Source/MeshModelingToolsEditorOnly/MeshModelingToolsEditorOnly.Build.cs
@@ -17,6 +17,8 @@ public class MeshModelingToolsEditorOnly : ModuleRules
 				"MeshModelingTools",
 				"ModelingComponents",
 				"ModelingOperatorsEditorOnly",
+				"ModelingOperators",
+				"PropertyEditor",
 				"DynamicMesh",
 				// ... add other public dependencies that you statically link with here ...
 			}
diff --git a/Engine/Plugins/Runtime/MeshModelingToolset/Source/MeshModelingToolsEditorOnly/Private/DetailsCustomizations/UVUnwrapToolCustomizations.cpp b/Engine/Plugins/Runtime/MeshModelingToolset/Source/MeshModelingToolsEditorOnly/Private/DetailsCustomizations/UVUnwrapToolCustomizations.cpp
new file mode 100644
--- /dev/null
+++ b/Engine/Plugins/Runtime/MeshModelingToolset/Source/MeshModelingToolsEditorOnly/Private/DetailsCustomizations/UVUnwrapToolCustomizations.cpp
@@ -0,0 +1,116 @@
+// Copyright Epic Games, Inc. All Rights Reserved.
+
+#include "DetailsCustomizations/UVUnwrapToolCustomizations.h"
+#include "DetailCategoryBuilder.h"
+#include "DetailLayoutBuilder.h"
+#include "DetailWidgetRow.h"
+#include "PropertyCustomizationHelpers.h"
+#include "Widgets/Input/SComboButton.h"
+#include "IDetailChildrenBuilder.h"
+#include "Internationalization/BreakIterator.h"
+#include "Framework/MultiBox/MultiBoxBuilder.h"
+#include "IContentBrowserSingleton.h"
+#include "SAssetView.h"
+
+#include "PropertyRestriction.h"
+
+#include "Properties/RecomputeUVsProperties.h"
+
+
+#define LOCTEXT_NAMESPACE "UVUnwrapDetailsCustomization"
+
+namespace UVUnwrapDetailsCustomizationLocal
+{
+
+	template<class UENUM_TYPE>
+	FPropertyAccess::Result GetPropertyValueAsEnum(const TSharedPtr<IPropertyHandle> Property, UENUM_TYPE& Value)
+	{
+		if(Property.IsValid())
+		{
+			uint8 ValueAsByte;
+			FPropertyAccess::Result Result = Property->GetValue(/*out*/ ValueAsByte);
+
+			if (Result == FPropertyAccess::Success)
+			{
+				Value = (UENUM_TYPE)ValueAsByte;
+				return FPropertyAccess::Success;
+			}
+			else
+			{
+				return FPropertyAccess::Fail;
+			}
+		}
+		return FPropertyAccess::Fail;
+	}
+
+}
+
+
+//
+// UVEditorRecomputeUVsTool
+//
+
+
+TSharedRef<IDetailCustomization> FRecomputeUVsToolDetails::MakeInstance()
+{
+	return MakeShareable(new FRecomputeUVsToolDetails);
+}
+
+
+void FRecomputeUVsToolDetails::CustomizeDetails(IDetailLayoutBuilder& DetailBuilder)
+{
+	TSharedRef<IPropertyHandle> IslandGenerationModeHandle = DetailBuilder.GetProperty(GET_MEMBER_NAME_CHECKED(URecomputeUVsToolProperties, IslandGeneration), URecomputeUVsToolProperties::StaticClass());
+	TSharedRef<IPropertyHandle> UnwrapTypeHandle = DetailBuilder.GetProperty(GET_MEMBER_NAME_CHECKED(URecomputeUVsToolProperties, UnwrapType), URecomputeUVsToolProperties::StaticClass());
+	TSharedRef<IPropertyHandle> LayoutTypeHandle = DetailBuilder.GetProperty(GET_MEMBER_NAME_CHECKED(URecomputeUVsToolProperties, LayoutType), URecomputeUVsToolProperties::StaticClass());
+
+	IslandGenerationModeHandle->SetOnPropertyValueChanged(FSimpleDelegate::CreateRaw(this, &FRecomputeUVsToolDetails::EvaluateLayoutTypeRestrictions, IslandGenerationModeHandle, UnwrapTypeHandle, LayoutTypeHandle));
+	UnwrapTypeHandle->SetOnPropertyValueChanged(FSimpleDelegate::CreateRaw(this, &FRecomputeUVsToolDetails::EvaluateLayoutTypeRestrictions, IslandGenerationModeHandle, UnwrapTypeHandle, LayoutTypeHandle));
+
+	static FText RestrictReason = LOCTEXT("LayoutTypeRestrictionText", "Requires existing UVs be used for IslandMode and Island Merging not selected as the unwrap strategy.");
+	LayoutTypeEnumRestriction = MakeShareable(new FPropertyRestriction(RestrictReason));
+	LayoutTypeHandle->AddRestriction(LayoutTypeEnumRestriction.ToSharedRef());
+
+	EvaluateLayoutTypeRestrictions(IslandGenerationModeHandle, UnwrapTypeHandle, LayoutTypeHandle);	
+}
+
+void FRecomputeUVsToolDetails::EvaluateLayoutTypeRestrictions(TSharedRef<IPropertyHandle> IslandGenerationModeHandle, TSharedRef<IPropertyHandle> UnwrapTypeHandle, TSharedRef<IPropertyHandle> LayoutTypeHandle)
+{
+	ERecomputeUVsPropertiesIslandMode IslandMode = ERecomputeUVsPropertiesIslandMode::PolyGroups;
+	ERecomputeUVsPropertiesUnwrapType UnwrapType = ERecomputeUVsPropertiesUnwrapType::Conformal;
+	FPropertyAccess::Result IslandModeResult = UVUnwrapDetailsCustomizationLocal::GetPropertyValueAsEnum(IslandGenerationModeHandle, IslandMode);
+	FPropertyAccess::Result UnwrapTypeResult = UVUnwrapDetailsCustomizationLocal::GetPropertyValueAsEnum(UnwrapTypeHandle, UnwrapType);
+
+	UEnum* ImportTypeEnum = StaticEnum<ERecomputeUVsPropertiesLayoutType>();
+
+	LayoutTypeEnumRestriction->RemoveAll();
+	if (IslandModeResult != FPropertyAccess::Fail &&
+		UnwrapTypeResult != FPropertyAccess::Fail &&
+		(IslandMode == ERecomputeUVsPropertiesIslandMode::PolyGroups ||
+		 UnwrapType == ERecomputeUVsPropertiesUnwrapType::IslandMerging))
+	{
+		uint8 EnumValue = static_cast<uint8>(ImportTypeEnum->GetValueByName("EUVEditorRecomputeUVsPropertiesLayoutType::NormalizeToExistingBounds"));
+		LayoutTypeEnumRestriction->AddDisabledValue("NormalizeToExistingBounds");
+
+		uint8 Value = 0;
+		LayoutTypeHandle->GetValue(Value);
+		// The goal here is to shift our selected layout type away from the offending option, either up or down one.
+		// This code assumes though that there are options other than our offending enum value to move off to. If it's the only one,
+		// this code will end up doing the wrong thing.
+		if (Value == EnumValue)
+		{
+			if (Value > 0)
+			{
+				Value = EnumValue - 1;
+			}
+			else
+			{
+				Value = EnumValue + 1;
+			}
+		}
+		LayoutTypeHandle->SetValue(Value, EPropertyValueSetFlags::NotTransactable | EPropertyValueSetFlags::InteractiveChange);
+	}
+
+
+}
+
+#undef LOCTEXT_NAMESPACE
\ No newline at end of file
diff --git a/Engine/Plugins/Runtime/MeshModelingToolset/Source/MeshModelingToolsEditorOnly/Private/MeshModelingToolsEditorOnly.cpp b/Engine/Plugins/Runtime/MeshModelingToolset/Source/MeshModelingToolsEditorOnly/Private/MeshModelingToolsEditorOnly.cpp
--- a/Engine/Plugins/Runtime/MeshModelingToolset/Source/MeshModelingToolsEditorOnly/Private/MeshModelingToolsEditorOnly.cpp
+++ b/Engine/Plugins/Runtime/MeshModelingToolset/Source/MeshModelingToolsEditorOnly/Private/MeshModelingToolsEditorOnly.cpp
@@ -2,18 +2,57 @@
 
 #include "MeshModelingToolsEditorOnly.h"
 
+#include "Modules/ModuleManager.h"
+#include "PropertyEditorModule.h"
+
+#include "DetailsCustomizations/UVUnwrapToolCustomizations.h"
+#include "Properties/RecomputeUVsProperties.h"
+
 #define LOCTEXT_NAMESPACE "FMeshModelingToolsEditorOnlyModule"
 
 void FMeshModelingToolsEditorOnlyModule::StartupModule()
 {
+	// This code will execute after your module is loaded into memory; the exact timing is specified in the .uplugin file per-module
+
+	FCoreDelegates::OnPostEngineInit.AddRaw(this, &FMeshModelingToolsEditorOnlyModule::OnPostEngineInit);
+
 }
 
 void FMeshModelingToolsEditorOnlyModule::ShutdownModule()
 {
 	// This function may be called during shutdown to clean up your module.  For modules that support dynamic reloading,
 	// we call this function before unloading the module.
+
+	FCoreDelegates::OnPostEngineInit.RemoveAll(this);
+
+	// Unregister customizations
+	FPropertyEditorModule* PropertyEditorModule = FModuleManager::GetModulePtr<FPropertyEditorModule>("PropertyEditor");
+	if (PropertyEditorModule)
+	{
+		for (const FName& ClassName : ClassesToUnregisterOnShutdown)
+		{
+			PropertyEditorModule->UnregisterCustomClassLayout(ClassName);
+		}
+		for (const FName& PropertyName : PropertiesToUnregisterOnShutdown)
+		{
+			PropertyEditorModule->UnregisterCustomPropertyTypeLayout(PropertyName);
+		}
+	}
+}
+
+void FMeshModelingToolsEditorOnlyModule::OnPostEngineInit()
+{
+
+	// Register details view customizations
+	FPropertyEditorModule& PropertyModule = FModuleManager::LoadModuleChecked<FPropertyEditorModule>("PropertyEditor");
+
+	/// Unwrap
+	PropertyModule.RegisterCustomClassLayout("RecomputeUVsToolProperties", FOnGetDetailCustomizationInstance::CreateStatic(&FRecomputeUVsToolDetails::MakeInstance));
+	ClassesToUnregisterOnShutdown.Add(URecomputeUVsToolProperties::StaticClass()->GetFName());
+
+
 }
 
 #undef LOCTEXT_NAMESPACE
 	
-IMPLEMENT_MODULE(FMeshModelingToolsEditorOnlyModule, MeshModelingToolsEditorOnly)
+IMPLEMENT_MODULE(FMeshModelingToolsEditorOnlyModule, MeshModelingToolsEditorOnly)
\ No newline at end of file
diff --git a/Engine/Plugins/Runtime/MeshModelingToolset/Source/MeshModelingToolsEditorOnly/Public/DetailsCustomizations/UVUnwrapToolCustomizations.h b/Engine/Plugins/Runtime/MeshModelingToolset/Source/MeshModelingToolsEditorOnly/Public/DetailsCustomizations/UVUnwrapToolCustomizations.h
new file mode 100644
--- /dev/null
+++ b/Engine/Plugins/Runtime/MeshModelingToolset/Source/MeshModelingToolsEditorOnly/Public/DetailsCustomizations/UVUnwrapToolCustomizations.h
@@ -0,0 +1,25 @@
+// Copyright Epic Games, Inc. All Rights Reserved.
+
+#pragma once
+
+#include "IDetailCustomization.h"
+#include "Templates/SharedPointer.h"
+
+class IPropertyHandle;
+class FPropertyRestriction;
+
+// Customization for UUVEditorRecomputeUVsToolProperties
+class FRecomputeUVsToolDetails : public IDetailCustomization
+{
+public:
+	static TSharedRef<IDetailCustomization> MakeInstance();
+	virtual void CustomizeDetails(IDetailLayoutBuilder& DetailBuilder) override;
+
+protected:
+	void EvaluateLayoutTypeRestrictions(TSharedRef<IPropertyHandle> IslandGenerationModeHandle, TSharedRef<IPropertyHandle> UnwrapTypeHandle, TSharedRef<IPropertyHandle> LayoutTypeHandle);
+
+	/** Property restriction applied to blend paint enum dropdown box */
+	TSharedPtr<FPropertyRestriction> LayoutTypeEnumRestriction;
+};
+
+
diff --git a/Engine/Plugins/Runtime/MeshModelingToolset/Source/MeshModelingToolsEditorOnly/Public/MeshModelingToolsEditorOnly.h b/Engine/Plugins/Runtime/MeshModelingToolset/Source/MeshModelingToolsEditorOnly/Public/MeshModelingToolsEditorOnly.h
--- a/Engine/Plugins/Runtime/MeshModelingToolset/Source/MeshModelingToolsEditorOnly/Public/MeshModelingToolsEditorOnly.h
+++ b/Engine/Plugins/Runtime/MeshModelingToolset/Source/MeshModelingToolsEditorOnly/Public/MeshModelingToolsEditorOnly.h
@@ -12,4 +12,10 @@ public:
 	/** IModuleInterface implementation */
 	virtual void StartupModule() override;
 	virtual void ShutdownModule() override;
+	void OnPostEngineInit();
+
+protected:
+
+	TArray<FName> ClassesToUnregisterOnShutdown;
+	TArray<FName> PropertiesToUnregisterOnShutdown;
 };
diff --git a/Engine/Plugins/Runtime/MeshModelingToolset/Source/ModelingOperators/Private/ParameterizationOps/RecomputeUVsOp.cpp b/Engine/Plugins/Runtime/MeshModelingToolset/Source/ModelingOperators/Private/ParameterizationOps/RecomputeUVsOp.cpp
--- a/Engine/Plugins/Runtime/MeshModelingToolset/Source/ModelingOperators/Private/ParameterizationOps/RecomputeUVsOp.cpp
+++ b/Engine/Plugins/Runtime/MeshModelingToolset/Source/ModelingOperators/Private/ParameterizationOps/RecomputeUVsOp.cpp
@@ -9,6 +9,7 @@
 #include "Parameterization/DynamicMeshUVEditor.h"
 #include "Parameterization/PatchBasedMeshUVGenerator.h"
 #include "Properties/RecomputeUVsProperties.h"
+#include "Parameterization/MeshUDIMClassifier.h"
 
 #include "DynamicSubmesh3.h"
 
@@ -22,68 +23,90 @@ using namespace UE::Geometry;
 #define LOCTEXT_NAMESPACE "RecomputeUVsOp"
 
 
-void FRecomputeUVsOp::NormalizeUVAreas(const FDynamicMesh3& Mesh, FDynamicMeshUVOverlay* Overlay, float GlobalScale)
+namespace UE
 {
-	FMeshConnectedComponents UVComponents(&Mesh);
-	UVComponents.FindConnectedTriangles([&](int32 Triangle0, int32 Triangle1) {
-		return Overlay->AreTrianglesConnected(Triangle0, Triangle1);
-	});
-
-	// TODO ParallelFor
-	for (FMeshConnectedComponents::FComponent& Component : UVComponents)
+	namespace Geometry
 	{
-		const TArray<int>& Triangles = Component.Indices;
-		TSet<int> UVElements;
-		UVElements.Reserve(Triangles.Num() * 3);
-		double AreaUV = 0;
-		double Area3D = 0;
-		FVector3d Triangle3D[3];
-		FVector2d TriangleUV[3];
-		FAxisAlignedBox2d BoundsUV = FAxisAlignedBox2d::Empty();
-
-		for (int tid : Triangles)
+		struct TileConnectedComponents
 		{
-			FIndex3i TriElements = Overlay->GetTriangle(tid);
-			if (!TriElements.Contains(FDynamicMesh3::InvalidID))
+			TileConnectedComponents(const FMeshConnectedComponents& ConnectedComponentsIn, const FVector2i& TileIn)
+				: ConnectedComponents(ConnectedComponentsIn),
+				Tile(TileIn)
 			{
-				for (int j = 0; j < 3; ++j)
+				for (int32 ComponentIndex = 0; ComponentIndex < ConnectedComponents.Num(); ++ComponentIndex)
 				{
-					TriangleUV[j] = FVector2d(Overlay->GetElement(TriElements[j]));
-					BoundsUV.Contain(TriangleUV[j]);
-					Triangle3D[j] = Mesh.GetVertex(Overlay->GetParentVertex(TriElements[j]));
-					UVElements.Add(TriElements[j]);
+					TileTids.Append(ConnectedComponents[ComponentIndex].Indices);
 				}
-				AreaUV += VectorUtil::Area(TriangleUV[0], TriangleUV[1], TriangleUV[2]);
-				Area3D += VectorUtil::Area(Triangle3D[0], Triangle3D[1], Triangle3D[2]);
 			}
-		}
-		
-		double LinearScale = (AreaUV > 0.00001) ? ( FMathd::Sqrt(Area3D) / FMathd::Sqrt(AreaUV)) : 1.0;
-		LinearScale = LinearScale * GlobalScale;
-		FVector2d ComponentOrigin = BoundsUV.Center();
 
-		for (int elemid : UVElements)
-		{
-			FVector2d UV = FVector2d(Overlay->GetElement(elemid));
-			UV = (UV - ComponentOrigin) * LinearScale;
-			Overlay->SetElement(elemid, FVector2f(UV));
-		}
+			TArray<int32> TileTids;
+			FMeshConnectedComponents ConnectedComponents;
+			FVector2i Tile;
+		};
 	}
 }
 
+void FRecomputeUVsOp::NormalizeUVAreas(const FDynamicMesh3& Mesh, FDynamicMeshUVOverlay* Overlay, const TileConnectedComponents& TileComponents, float GlobalScale)
+{
+	int32 NumComponents = TileComponents.ConnectedComponents.Num();
+
+	ParallelFor(NumComponents, [&](int32 k)
+		{
+			const TArray<int>& Triangles = TileComponents.ConnectedComponents[k].Indices;
+			TSet<int> UVElements;
+			UVElements.Reserve(Triangles.Num() * 3);
+			double AreaUV = 0;
+			double Area3D = 0;
+			FVector3d Triangle3D[3];
+			FVector2d TriangleUV[3];
+			FAxisAlignedBox2d BoundsUV = FAxisAlignedBox2d::Empty();
+
+			for (int tid : Triangles)
+			{
+				FIndex3i TriElements = Overlay->GetTriangle(tid);
+				if (!TriElements.Contains(FDynamicMesh3::InvalidID))
+				{
+					for (int j = 0; j < 3; ++j)
+					{
+						TriangleUV[j] = FVector2d(Overlay->GetElement(TriElements[j]));
+						BoundsUV.Contain(TriangleUV[j]);
+						Triangle3D[j] = Mesh.GetVertex(Overlay->GetParentVertex(TriElements[j]));
+						UVElements.Add(TriElements[j]);
+					}
+					AreaUV += VectorUtil::Area(TriangleUV[0], TriangleUV[1], TriangleUV[2]);
+					Area3D += VectorUtil::Area(Triangle3D[0], Triangle3D[1], Triangle3D[2]);
+				}
+			}
+
+			double LinearScale = (AreaUV > 0.00001) ? (FMathd::Sqrt(Area3D) / FMathd::Sqrt(AreaUV)) : 1.0;
+			LinearScale = LinearScale * GlobalScale;
+			FVector2d ComponentOrigin = BoundsUV.Center();
+
+			for (int elemid : UVElements)
+			{
+				FVector2d UV = FVector2d(Overlay->GetElement(elemid));
+				UV = (UV - ComponentOrigin) * LinearScale;
+				Overlay->SetElement(elemid, FVector2f(UV));
+			}
+		});
+}
+
 
 void FRecomputeUVsOp::CalculateResult(FProgressCancel* Progress)
 {
 	NewResultInfo = FGeometryResult(EGeometryResultType::InProgress);
+	bool bOK = false;
 
-	bool bOK;
-	if (bMergingOptimization)
-	{
-		bOK = CalculateResult_RegionOptimization(Progress);
-	}
-	else
+	if (IsValid())
 	{
-		bOK = CalculateResult_Basic(Progress);
+		if (bMergingOptimization)
+		{
+			bOK = CalculateResult_RegionOptimization(Progress);
+		}
+		else
+		{
+			bOK = CalculateResult_Basic(Progress);
+		}
 	}
 
 	NewResultInfo.SetSuccess(bOK, Progress);
@@ -91,8 +114,99 @@ void FRecomputeUVsOp::CalculateResult(FProgressCancel* Progress)
 
 }
 
+void FRecomputeUVsOp::CollectIslandComponentsPerTile(const FDynamicMeshUVOverlay& UVOverlay, TArray< TileConnectedComponents >& ComponentsPerTile, bool& bUseExistingUVTopology)
+{
+	bUseExistingUVTopology = false;
+	ComponentsPerTile.Empty();
 
+	TArray<FVector2i> Tiles;
+	TArray<TUniquePtr<TArray<int32>>> TileTids;
+
+	if (bUDIMsEnabled)
+	{
+		TOptional<TArray<int32>> SelectionArray;
+		if (Selection.IsSet())
+		{
+			SelectionArray = Selection.GetValue().Array();
+		}
+		FDynamicMeshUDIMClassifier TileClassifier(&UVOverlay, SelectionArray);
+
+		Tiles = TileClassifier.ActiveTiles();
+		for (const FVector2i& Tile : Tiles)
+		{
+			TileTids.Emplace(MakeUnique<TArray<int32>>(TileClassifier.TidsForTile(Tile)));
+		}
+	}
+	else
+	{
+		if (Selection.IsSet())
+		{
+			Tiles.Add({ 0,0 });
+			TileTids.Emplace(MakeUnique<TArray<int32>>(Selection.GetValue().Array()));
+		}
+		else
+		{
+			Tiles.Add({ 0,0 });
+			TileTids.Add(nullptr);
+		}
+	}
+
+	for (int32 TileIndex = 0; TileIndex < Tiles.Num(); ++TileIndex)
+	{
+		FMeshConnectedComponents ConnectedComponents(ResultMesh.Get());
+		if (IslandMode == ERecomputeUVsIslandMode::PolyGroups)
+		{
+			if (InputGroups != nullptr)
+			{
+				if (TileTids[TileIndex])
+				{
+					ConnectedComponents.FindConnectedTriangles(*TileTids[TileIndex], [this](int32 CurTri, int32 NbrTri) {
+						return InputGroups->GetTriangleGroup(CurTri) == InputGroups->GetTriangleGroup(NbrTri);
+						});
+				}
+				else
+				{
+					ConnectedComponents.FindConnectedTriangles([this](int32 CurTri, int32 NbrTri) {
+						return InputGroups->GetTriangleGroup(CurTri) == InputGroups->GetTriangleGroup(NbrTri);
+						});
+				}
+			}
+			else
+			{
+				if (TileTids[TileIndex])
+				{
+					ConnectedComponents.FindConnectedTriangles(*TileTids[TileIndex], [this](int32 CurTri, int32 NbrTri) {
+						return ResultMesh->GetTriangleGroup(CurTri) == ResultMesh->GetTriangleGroup(NbrTri);
+						});
+				}
+				else
+				{
+					ConnectedComponents.FindConnectedTriangles([this](int32 CurTri, int32 NbrTri) {
+						return ResultMesh->GetTriangleGroup(CurTri) == ResultMesh->GetTriangleGroup(NbrTri);
+						});
+				}
+			}
+		}
+		else
+		{
+			if (TileTids[TileIndex])
+			{
+				ConnectedComponents.FindConnectedTriangles(*TileTids[TileIndex], [&](int32 Triangle0, int32 Triangle1) {
+					return UVOverlay.AreTrianglesConnected(Triangle0, Triangle1);
+					});
+			}
+			else
+			{
+				ConnectedComponents.FindConnectedTriangles([&](int32 Triangle0, int32 Triangle1) {
+					return UVOverlay.AreTrianglesConnected(Triangle0, Triangle1);
+					});
+			}
 
+			bUseExistingUVTopology = true;
+		}
+		ComponentsPerTile.Emplace(ConnectedComponents, Tiles[TileIndex]);
+	}
+}
 
 bool FRecomputeUVsOp::CalculateResult_Basic(FProgressCancel* Progress)
 {
@@ -117,136 +231,155 @@ bool FRecomputeUVsOp::CalculateResult_Basic(FProgressCancel* Progress)
 	}
 
 	// find group-connected-components
-	bool bUseExisingUVTopology = false;
-	FMeshConnectedComponents ConnectedComponents(ResultMesh.Get());
-	if (IslandMode == ERecomputeUVsIslandMode::PolyGroups)
-	{
-		if (InputGroups != nullptr)
-		{
-			ConnectedComponents.FindConnectedTriangles([this](int32 CurTri, int32 NbrTri) {
-				return InputGroups->GetTriangleGroup(CurTri) == InputGroups->GetTriangleGroup(NbrTri);
-			});
-		}
-		else
-		{
-			ConnectedComponents.FindConnectedTriangles([this](int32 CurTri, int32 NbrTri) {
-				return ResultMesh->GetTriangleGroup(CurTri) == ResultMesh->GetTriangleGroup(NbrTri);
-			});
-		}
-	}
-	else
-	{
-		ConnectedComponents.FindConnectedTriangles([&](int32 Triangle0, int32 Triangle1) {
-			return UseOverlay->AreTrianglesConnected(Triangle0, Triangle1);
-		});
-		bUseExisingUVTopology = true;
-	}
+	bool bUseExistingUVTopology = false;
+	TArray< TileConnectedComponents > ComponentsPerTile;
+	CollectIslandComponentsPerTile(*UseOverlay, ComponentsPerTile, bUseExistingUVTopology);
 
 	if (Progress && Progress->Cancelled())
 	{
 		return false;
 	}
 
-	// TODO: the solves here could be done in parallel if we pre-allocated the island element IDs
 
-	int32 NumComponents = ConnectedComponents.Num();
-	TArray<bool> bComponentSolved;
-	bComponentSolved.Init(false, NumComponents);
-	int32 SuccessCount = 0;
-	for (int32 k = 0; k < NumComponents; ++k)
+	for (TileConnectedComponents& TileComponents : ComponentsPerTile)
 	{
-		const TArray<int32>& ComponentTris = ConnectedComponents[k].Indices;
+		int32 NumComponents = TileComponents.ConnectedComponents.Num();
+		TArray<bool> bComponentSolved;
+		bComponentSolved.Init(false, NumComponents);
+		int32 SuccessCount = 0;
+		TArray<FAxisAlignedBox2f> PerComponentBoundingBoxes;
+		PerComponentBoundingBoxes.SetNum(NumComponents);
 
-		bComponentSolved[k] = false;
-		switch (UnwrapType)
+		// TODO: the solves here could be done in parallel if we pre-allocated the island element IDs
+		for (int32 k = 0; k < NumComponents; ++k)
 		{
+			const TArray<int32>& ComponentTris = TileComponents.ConnectedComponents[k].Indices;
+
+			FAxisAlignedBox2f& UVBounds = PerComponentBoundingBoxes[k];
+			if (bPackToOriginalBounds)
+			{
+				for (int32 tid : ComponentTris)
+				{
+					if (UseOverlay->IsSetTriangle(tid))
+					{
+						FIndex3i UVTri = UseOverlay->GetTriangle(tid);
+						FVector2f U = UseOverlay->GetElement(UVTri.A);
+						FVector2f V = UseOverlay->GetElement(UVTri.B);
+						FVector2f W = UseOverlay->GetElement(UVTri.C);
+						UVBounds.Contain(U); UVBounds.Contain(V); UVBounds.Contain(W);
+					}
+				}
+			}
+
+			bComponentSolved[k] = false;
+			switch (UnwrapType)
+			{
 			case ERecomputeUVsUnwrapType::ExpMap:
 			{
 				FDynamicMeshUVEditor::FExpMapOptions Options;
 				Options.NormalSmoothingRounds = this->NormalSmoothingRounds;
 				Options.NormalSmoothingAlpha = this->NormalSmoothingAlpha;
 				bComponentSolved[k] = UVEditor.SetTriangleUVsFromExpMap(ComponentTris, Options);
-				break;
 			}
+			break;
 
 			case ERecomputeUVsUnwrapType::ConformalFreeBoundary:
-			{
-				bComponentSolved[k] = UVEditor.SetTriangleUVsFromFreeBoundaryConformal(ComponentTris, bUseExisingUVTopology);
+				bComponentSolved[k] = UVEditor.SetTriangleUVsFromFreeBoundaryConformal(ComponentTris, bUseExistingUVTopology);
 				if (bComponentSolved[k])
 				{
 					UVEditor.ScaleUVAreaTo3DArea(ComponentTris, true);
 				}
 				break;
-			}
-			
 			case ERecomputeUVsUnwrapType::SpectralConformal:
 			{
-				bComponentSolved[k] = UVEditor.SetTriangleUVsFromFreeBoundarySpectralConformal(ComponentTris, bUseExisingUVTopology, bPreserveIrregularity);
+				bComponentSolved[k] = UVEditor.SetTriangleUVsFromFreeBoundarySpectralConformal(ComponentTris, bUseExistingUVTopology, bPreserveIrregularity);
 				if (bComponentSolved[k])
 				{
 					UVEditor.ScaleUVAreaTo3DArea(ComponentTris, true);
 				}
 				break;
 			}
+			}
+
+			if (bComponentSolved[k])
+			{
+				SuccessCount++;
+			}
+
+			if (Progress && Progress->Cancelled())
+			{
+				return false;
+			}
 		}
-		
 
-		if (bComponentSolved[k])
+
+		if (bAutoRotate)
 		{
-			SuccessCount++;
+			ParallelFor(NumComponents, [&](int32 k)
+				{
+					if (Progress && Progress->Cancelled())
+					{
+						return;
+					}
+					if (bComponentSolved[k])
+					{
+						const TArray<int32>& ComponentTris = TileComponents.ConnectedComponents[k].Indices;
+						UVEditor.AutoOrientUVArea(ComponentTris);
+					};
+				});
 		}
-
 		if (Progress && Progress->Cancelled())
 		{
 			return false;
 		}
-	}
 
-	if (bAutoRotate)
-	{
-		ParallelFor(NumComponents, [&](int32 k)
+		if (bPackToOriginalBounds)
 		{
-			if (Progress && Progress->Cancelled())
-			{
-				return;
-			}
-			if (bComponentSolved[k])
-			{
-				const TArray<int32>& ComponentTris = ConnectedComponents[k].Indices;
-				UVEditor.AutoOrientUVArea(ComponentTris);
-			};
-		});
-	}
-	if (Progress && Progress->Cancelled())
-	{
-		return false;
-	}
+			ParallelFor(NumComponents, [&](int32 k)
+				{
+					if (Progress && Progress->Cancelled())
+					{
+						return;
+					}
+					if (bComponentSolved[k])
+					{
+						const TArray<int32>& ComponentTris = TileComponents.ConnectedComponents[k].Indices;
+						const FAxisAlignedBox2f& UVBounds = PerComponentBoundingBoxes[k];
+						UVEditor.ScaleUVAreaToBoundingBox(ComponentTris, UVBounds, true, true);
+					}
+				});
+		}
+		if (Progress && Progress->Cancelled())
+		{
+			return false;
+		}
 
+		if (bNormalizeAreas)
+		{
+			// todo should be a DynamicUVEditor function?
+			NormalizeUVAreas(*ResultMesh, UseOverlay, TileComponents, AreaScaling);
+		}
 
-	if (bNormalizeAreas)
-	{
-		// todo should be a DynamicUVEditor function?
-		NormalizeUVAreas(*ResultMesh, UseOverlay, AreaScaling);
-	}
+		if (Progress && Progress->Cancelled())
+		{
+			return false;
+		}
 
-	if (Progress && Progress->Cancelled())
-	{
-		return false;
-	}
+		if (bPackUVs && ensure(TileComponents.TileTids.Num() > 0))
+		{
+			bool bPackingSuccess = UVEditor.UDIMPack(PackingTextureResolution, PackingGutterWidth, TileComponents.Tile, &TileComponents.TileTids);
+			if (!bPackingSuccess)
+			{
+				NewResultInfo.AddError(FGeometryError(0, LOCTEXT("IslandPackingFailed", "Failed to pack UV islands")));
+				return false;
+			}
+		}
 
-	if (bPackUVs)
-	{
-		bool bPackingSuccess = UVEditor.QuickPack(PackingTextureResolution, PackingGutterWidth);
-		if (!bPackingSuccess)
+		if (Progress && Progress->Cancelled())
 		{
-			NewResultInfo.AddError(FGeometryError(0, LOCTEXT("IslandPackingFailed", "Failed to pack UV islands")));
 			return false;
 		}
-	}
 
-	if (Progress && Progress->Cancelled())
-	{
-		return false;
 	}
 
 	return true;
@@ -279,111 +412,106 @@ bool FRecomputeUVsOp::CalculateResult_RegionOptimization(FProgressCancel* Progre
 		return false;
 	}
 
-	// initialize UV generation with connected components
-	FMeshConnectedComponents ConnectedComponents(ResultMesh.Get());
-	if (IslandMode == ERecomputeUVsIslandMode::PolyGroups)
-	{
-		if (InputGroups != nullptr)
-		{
-			ConnectedComponents.FindConnectedTriangles([this](int32 CurTri, int32 NbrTri) {
-				return InputGroups->GetTriangleGroup(CurTri) == InputGroups->GetTriangleGroup(NbrTri);
-			});
-		}
-		else
-		{
-			ConnectedComponents.FindConnectedTriangles([this](int32 CurTri, int32 NbrTri) {
-				return ResultMesh->GetTriangleGroup(CurTri) == ResultMesh->GetTriangleGroup(NbrTri);
-			});
-		}
-	}
-	else
-	{
-		ConnectedComponents.FindConnectedTriangles([&](int32 Triangle0, int32 Triangle1) {
-			return UseOverlay->AreTrianglesConnected(Triangle0, Triangle1);
-		});
-	}
+	// find group-connected-components
+	bool bUseExistingUVTopology = false;
+	TArray< TileConnectedComponents > ComponentsPerTile;
+	CollectIslandComponentsPerTile(*UseOverlay, ComponentsPerTile, bUseExistingUVTopology);
 
 	if (Progress && Progress->Cancelled())
 	{
 		return false;
 	}
 
+	for (TileConnectedComponents& TileComponents : ComponentsPerTile)
+	{
+		FPatchBasedMeshUVGenerator UVGenerator;
+		UVGenerator.MergingThreshold = this->MergingThreshold;
+		UVGenerator.CompactnessThreshold = this->CompactnessThreshold;
+		UVGenerator.MaxNormalDeviationDeg = this->MaxNormalDeviationDeg;
+		UVGenerator.NormalSmoothingRounds = this->NormalSmoothingRounds;
+		UVGenerator.NormalSmoothingAlpha = this->NormalSmoothingAlpha;
 
-	FPatchBasedMeshUVGenerator UVGenerator;
-	UVGenerator.MergingThreshold = this->MergingThreshold;
-	UVGenerator.CompactnessThreshold = this->CompactnessThreshold;
-	UVGenerator.MaxNormalDeviationDeg = this->MaxNormalDeviationDeg;
-	UVGenerator.NormalSmoothingRounds = this->NormalSmoothingRounds;
-	UVGenerator.NormalSmoothingAlpha = this->NormalSmoothingAlpha;
+		TArray<TArray<int32>> UVIslandTriangleSets;
+		bool bIslandsOK = UVGenerator.ComputeIslandsByRegionMerging(*BaseMesh, *UseOverlay, TileComponents.ConnectedComponents, UVIslandTriangleSets, Progress);
 
-	TArray<TArray<int32>> UVIslandTriangleSets;
-	bool bIslandsOK = UVGenerator.ComputeIslandsByRegionMerging(*BaseMesh, *UseOverlay, ConnectedComponents, UVIslandTriangleSets, Progress);
+		if (bIslandsOK == false)
+		{
+			NewResultInfo.AddError(FGeometryError(0, LOCTEXT("IslandMergingFailed", "Failed to merge UV islands")));
+			return false;
+		}
 
-	if (bIslandsOK == false)
-	{
-		NewResultInfo.AddError(FGeometryError(0, LOCTEXT("IslandMergingFailed", "Failed to merge UV islands")));
-		return false;
-	}
+		if (Progress && Progress->Cancelled())
+		{
+			return false;
+		}
 
-	if (Progress && Progress->Cancelled())
-	{
-		return false;
-	}
+		TArray<bool> bIslandUVsValid;
+		int32 NumSolvesFailed = UVGenerator.ComputeUVsFromTriangleSets(*BaseMesh, *UseOverlay, UVIslandTriangleSets, bIslandUVsValid, Progress);
+		if (NumSolvesFailed > 0)
+		{
+			NewResultInfo.AddWarning(FGeometryWarning(0, LOCTEXT("PartialSolvesFailed", "Failed to compute UVs for some UV islands")));
+		}
 
-	TArray<bool> bIslandUVsValid;
-	int32 NumSolvesFailed = UVGenerator.ComputeUVsFromTriangleSets(*BaseMesh, *UseOverlay, UVIslandTriangleSets, bIslandUVsValid, Progress);
-	if (NumSolvesFailed > 0)
-	{
-		NewResultInfo.AddWarning(FGeometryWarning(0, LOCTEXT("PartialSolvesFailed", "Failed to compute UVs for some UV islands")));
-	}
+		if (Progress && Progress->Cancelled())
+		{
+			return false;
+		}
 
-	if (Progress && Progress->Cancelled())
-	{
-		return false;
-	}
+		if (bAutoRotate)
+		{
+			ParallelFor(UVIslandTriangleSets.Num(), [&](int32 k)
+				{
+					if (Progress && Progress->Cancelled())
+					{
+						return;
+					}
+					if (bIslandUVsValid[k])
+					{
+						UVEditor.AutoOrientUVArea(UVIslandTriangleSets[k]);
+					}
+				});
+		}
+		if (Progress && Progress->Cancelled())
+		{
+			return false;
+		}
 
-	if (bAutoRotate)
-	{
-		ParallelFor(UVIslandTriangleSets.Num(), [&](int32 k)
+		if (bNormalizeAreas)
 		{
-			if (Progress && Progress->Cancelled())
-			{
-				return;
-			}
-			if (bIslandUVsValid[k])
-			{
-				UVEditor.AutoOrientUVArea(UVIslandTriangleSets[k]);
-			}
-		});
-	}
-	if (Progress && Progress->Cancelled())
-	{
-		return false;
-	}
+			// todo should be a DynamicUVEditor function?
+			NormalizeUVAreas(*ResultMesh, UseOverlay, TileComponents, AreaScaling);
+		}
 
-	if (bNormalizeAreas)
-	{
-		// todo should be a DynamicUVEditor function?
-		NormalizeUVAreas(*ResultMesh, UseOverlay, AreaScaling);
-	}
 
-	if (Progress && Progress->Cancelled())
-	{
-		return false;
-	}
+		if (Progress && Progress->Cancelled())
+		{
+			return false;
+		}
 
-	if (bPackUVs)
-	{
-		bool bPackingSuccess = UVEditor.QuickPack(PackingTextureResolution, PackingGutterWidth);
-		if (!bPackingSuccess)
+		if (bPackUVs)
+		{
+			bool bPackingSuccess = UVEditor.UDIMPack(PackingTextureResolution, PackingGutterWidth, TileComponents.Tile, &TileComponents.TileTids);
+			if (!bPackingSuccess)
+			{
+				NewResultInfo.AddError(FGeometryError(0, LOCTEXT("IslandPackingFailed", "Failed to pack UV islands")));
+				return false;
+			}
+		}
+
+		if (Progress && Progress->Cancelled())
 		{
-			NewResultInfo.AddError(FGeometryError(0, LOCTEXT("IslandPackingFailed", "Failed to pack UV islands")));
 			return false;
 		}
 	}
 
-	if (Progress && Progress->Cancelled())
+	return true;
+}
+
+bool FRecomputeUVsOp::IsValid() const
+{
+	if (bPackToOriginalBounds && (bMergingOptimization || IslandMode == ERecomputeUVsIslandMode::PolyGroups))
 	{
+		// It is not a valid state to be packing into original UV island bounds if islands are being merged or re-generated from polygroups
 		return false;
 	}
 
@@ -416,27 +544,29 @@ TUniquePtr<FDynamicMeshOperator> URecomputeUVsOpFactory::MakeNewOperator()
 		RecomputeUVsOp->UnwrapType = ERecomputeUVsUnwrapType::ConformalFreeBoundary;
 		RecomputeUVsOp->bMergingOptimization = false;
 		break;
+	case ERecomputeUVsPropertiesUnwrapType::IslandMerging:
+		RecomputeUVsOp->UnwrapType = ERecomputeUVsUnwrapType::ExpMap;
+		RecomputeUVsOp->bMergingOptimization = true;
+		break;
 	case ERecomputeUVsPropertiesUnwrapType::SpectralConformal:
 		RecomputeUVsOp->UnwrapType = ERecomputeUVsUnwrapType::SpectralConformal;
 		RecomputeUVsOp->bMergingOptimization = false;
 		break;
-	case ERecomputeUVsPropertiesUnwrapType::IslandMerging:
-		RecomputeUVsOp->UnwrapType = ERecomputeUVsUnwrapType::ExpMap;
-		RecomputeUVsOp->bMergingOptimization = true;
 	}
 
 	RecomputeUVsOp->bAutoRotate = (Settings->AutoRotation == ERecomputeUVsToolOrientationMode::MinBounds);
 
-	RecomputeUVsOp->bPreserveIrregularity = Settings->bPreserveIrregularity;
-
 	RecomputeUVsOp->NormalSmoothingRounds = Settings->SmoothingSteps;
 	RecomputeUVsOp->NormalSmoothingAlpha = Settings->SmoothingAlpha;
 
+	RecomputeUVsOp->bPreserveIrregularity = Settings->bPreserveIrregularity;
+
 	RecomputeUVsOp->MergingThreshold = Settings->MergingDistortionThreshold;
 	RecomputeUVsOp->MaxNormalDeviationDeg = Settings->MergingAngleThreshold;
 
 	RecomputeUVsOp->bPackUVs = false;
 	RecomputeUVsOp->bNormalizeAreas = false;
+	RecomputeUVsOp->bUDIMsEnabled = Settings->bEnableUDIMLayout;
 
 	switch (Settings->LayoutType)
 	{
@@ -446,9 +576,14 @@ TUniquePtr<FDynamicMeshOperator> URecomputeUVsOpFactory::MakeNewOperator()
 		RecomputeUVsOp->bPackUVs = true;
 		RecomputeUVsOp->PackingTextureResolution = Settings->TextureResolution;
 		break;
+	case ERecomputeUVsPropertiesLayoutType::NormalizeToExistingBounds:
+		RecomputeUVsOp->bPackToOriginalBounds = true;
+		RecomputeUVsOp->bPackUVs = false;
+		RecomputeUVsOp->bNormalizeAreas = false;
+		break;
 	case ERecomputeUVsPropertiesLayoutType::NormalizeToBounds:
 		RecomputeUVsOp->bNormalizeAreas = true;
-		RecomputeUVsOp->AreaScaling = Settings->NormalizeScale / OriginalMesh->GetBounds().MaxDim();
+		RecomputeUVsOp->AreaScaling = Settings->NormalizeScale / static_cast<float>(OriginalMesh->GetBounds().MaxDim());
 		break;
 	case ERecomputeUVsPropertiesLayoutType::NormalizeToWorld:
 		RecomputeUVsOp->bNormalizeAreas = true;
@@ -456,6 +591,7 @@ TUniquePtr<FDynamicMeshOperator> URecomputeUVsOpFactory::MakeNewOperator()
 	}
 
 	RecomputeUVsOp->SetTransform(TargetTransform);
+	RecomputeUVsOp->Selection = Selection;
 
 	return RecomputeUVsOp;
 }
diff --git a/Engine/Plugins/Runtime/MeshModelingToolset/Source/ModelingOperators/Public/ParameterizationOps/RecomputeUVsOp.h b/Engine/Plugins/Runtime/MeshModelingToolset/Source/ModelingOperators/Public/ParameterizationOps/RecomputeUVsOp.h
--- a/Engine/Plugins/Runtime/MeshModelingToolset/Source/ModelingOperators/Public/ParameterizationOps/RecomputeUVsOp.h
+++ b/Engine/Plugins/Runtime/MeshModelingToolset/Source/ModelingOperators/Public/ParameterizationOps/RecomputeUVsOp.h
@@ -17,6 +17,7 @@ namespace UE
 namespace Geometry
 {
 
+	struct TileConnectedComponents;
 
 enum class ERecomputeUVsUnwrapType
 {
@@ -45,13 +46,16 @@ public:
 
 	// source mesh
 	TSharedPtr<FDynamicMesh3, ESPMode::ThreadSafe> InputMesh;
-	
+
 	// source groups (optional)
 	TSharedPtr<UE::Geometry::FPolygroupSet, ESPMode::ThreadSafe> InputGroups;
 
 	// orientation control
 	bool bAutoRotate = true;
 
+	// Individual packing
+	bool bPackToOriginalBounds = false;
+
 	// area scaling
 	bool bNormalizeAreas = true;
 	float AreaScaling = 1.0;
@@ -68,10 +72,10 @@ public:
 	ERecomputeUVsIslandMode IslandMode = ERecomputeUVsIslandMode::PolyGroups;
 
 	//
-	// Conformal Map Options 
+	// Spectral Conformal Map Options 
 	//
 	bool bPreserveIrregularity = false;
-	
+
 	// 
 	// ExpMap Options
 	//
@@ -86,6 +90,12 @@ public:
 	double CompactnessThreshold = 9999999.0;		// effectively disabled as it usually is not a good idea
 	double MaxNormalDeviationDeg = 45.0;
 
+	// UDIM options	
+	bool bUDIMsEnabled = false;
+
+	// Selection
+	TOptional<TSet<int32>> Selection;
+
 	// set ability on protected transform.
 	void SetTransform(const FTransformSRT3d& XForm)
 	{
@@ -101,9 +111,16 @@ public:
 
 protected:
 
+	/*
+	* Check if operation is valid to do work. It may not be valid if there are conflicting or inconsistent options passed to it from the user.
+	*/
+	bool IsValid() const;
+
 	FGeometryResult NewResultInfo;
 
-	void NormalizeUVAreas(const FDynamicMesh3& Mesh, FDynamicMeshUVOverlay* Overlay, float GlobalScale = 1.0f);
+	void NormalizeUVAreas(const FDynamicMesh3& Mesh, FDynamicMeshUVOverlay* Overlay, const TileConnectedComponents& TileComponents, float GlobalScale = 1.0f);
+
+	void CollectIslandComponentsPerTile(const FDynamicMeshUVOverlay& UVOverlay, TArray< TileConnectedComponents >& ComponentsPerTile, bool& UseExistingUVTopology);
 
 	virtual bool CalculateResult_Basic(FProgressCancel* Progress);
 	virtual bool CalculateResult_RegionOptimization(FProgressCancel* Progress);
@@ -132,6 +149,8 @@ public:
 	UPROPERTY()
 	TObjectPtr<URecomputeUVsToolProperties> Settings = nullptr;
 
+	TOptional<TSet<int32>> Selection;
+
 	TSharedPtr<UE::Geometry::FPolygroupSet, ESPMode::ThreadSafe> InputGroups;
 	TSharedPtr<UE::Geometry::FDynamicMesh3, ESPMode::ThreadSafe> OriginalMesh;
 	TUniqueFunction<int32()> GetSelectedUVChannel = []() { return 0; };
diff --git a/Engine/Plugins/Runtime/MeshModelingToolset/Source/ModelingOperators/Public/Properties/RecomputeUVsProperties.h b/Engine/Plugins/Runtime/MeshModelingToolset/Source/ModelingOperators/Public/Properties/RecomputeUVsProperties.h
--- a/Engine/Plugins/Runtime/MeshModelingToolset/Source/ModelingOperators/Public/Properties/RecomputeUVsProperties.h
+++ b/Engine/Plugins/Runtime/MeshModelingToolset/Source/ModelingOperators/Public/Properties/RecomputeUVsProperties.h
@@ -54,6 +54,8 @@ enum class ERecomputeUVsPropertiesLayoutType
 	None,
 	/** Uniformly scale and translate UV islands collectively to pack them into the unit square, i.e. fit between 0 and 1 with no overlap */
 	Repack,
+	/** Scale and center all islands to fit within their original bounding boxes. Only applicable if using existing UVs. */
+	NormalizeToExistingBounds,
 	/** Uniformly scale UV islands such that they have constant relative area, relative to object bounds */
 	NormalizeToBounds,
 	/** Uniformly scale UV islands such that they have constant relative area, relative to world space */
@@ -83,27 +85,27 @@ public:
 	ERecomputeUVsToolOrientationMode AutoRotation = ERecomputeUVsToolOrientationMode::MinBounds;
 
 	/**  If enabled, reduces distortion for meshes with triangles of vastly different sizes, This is only enabled if the Unwrap Type is set to Spectral Conformal. */
-	UPROPERTY(EditAnywhere, Category = "UV Unwrap", meta = (EditCondition = "UnwrapType == ERecomputeUVsPropertiesUnwrapType::SpectralConformal"))
+	UPROPERTY(EditAnywhere, Category = "UV Unwrap", meta = (EditCondition = "UnwrapType == ERecomputeUVsPropertiesUnwrapType::SpectralConformal", EditConditionHides, HideEditConditionToggle))
 	bool bPreserveIrregularity = true;
 
 	/** Number of smoothing steps to apply; this slightly increases distortion but produces more stable results. This is only enabled if the Unwrap Type is set to ExpMap or Island Merging. */
 	UPROPERTY(EditAnywhere, Category = "UV Unwrap", meta = (UIMin = "0", UIMax = "25", ClampMin = "0", ClampMax = "1000",
-		EditCondition = "UnwrapType == ERecomputeUVsPropertiesUnwrapType::ExpMap || UnwrapType == ERecomputeUVsPropertiesUnwrapType::IslandMerging"))
+		EditCondition = "UnwrapType == ERecomputeUVsPropertiesUnwrapType::ExpMap || UnwrapType == ERecomputeUVsPropertiesUnwrapType::IslandMerging", EditConditionHides, HideEditConditionToggle))
 	int SmoothingSteps = 5;
 
 	/** Smoothing parameter; larger values result in faster smoothing in each step. This is only enabled if the Unwrap Type is set to ExpMap or Island Merging. */
 	UPROPERTY(EditAnywhere, Category = "UV Unwrap", meta = (UIMin = "0", UIMax = "1.0", ClampMin = "0", ClampMax = "1.0",
-		EditCondition = "UnwrapType == ERecomputeUVsPropertiesUnwrapType::ExpMap || UnwrapType == ERecomputeUVsPropertiesUnwrapType::IslandMerging"))
+		EditCondition = "UnwrapType == ERecomputeUVsPropertiesUnwrapType::ExpMap || UnwrapType == ERecomputeUVsPropertiesUnwrapType::IslandMerging", EditConditionHides, HideEditConditionToggle))
 	float SmoothingAlpha = 0.25f;
 
 	/** Threshold for stretching and distortion below which island merging is allowed; larger values increase the allowable UV distortion. This is only enabled if the Unwrap Type is set to Island Merging. */
 	UPROPERTY(EditAnywhere, Category = "UV Unwrap", meta = (DisplayName = "Distortion Threshold", UIMin = "1.0", UIMax = "5.0", ClampMin = "1.0",
-		EditCondition = "UnwrapType == ERecomputeUVsPropertiesUnwrapType::IslandMerging"))
+		EditCondition = "UnwrapType == ERecomputeUVsPropertiesUnwrapType::IslandMerging", EditConditionHides, HideEditConditionToggle))
 	float MergingDistortionThreshold = 1.5f;
 
 	/** Threshold for the average face normal deviation below  which island merging is allowed. This is only enabled if the Unwrap Type is set to Island Merging. */
 	UPROPERTY(EditAnywhere, Category = "UV Unwrap", meta = (DisplayName = "Angle Threshold", UIMin = "0.0", UIMax = "90.0", ClampMin = "0.0", ClampMax = "180.0",
-		EditCondition = "UnwrapType == ERecomputeUVsPropertiesUnwrapType::IslandMerging"))
+		EditCondition = "UnwrapType == ERecomputeUVsPropertiesUnwrapType::IslandMerging", EditConditionHides, HideEditConditionToggle))
 	float MergingAngleThreshold = 45.0f;
 
 	/** Uniformly scale and translate UV islands collectively to pack them into the unit square, i.e. fit between 0 and 1 with no overlap. */
@@ -112,11 +114,18 @@ public:
 
 	/** Expected resolution of the output textures; this controls spacing left between UV islands to avoid interpolation artifacts. This is only enabled when the Layout Type is set to Repack. */
 	UPROPERTY(EditAnywhere, Category = "UV Layout", meta = (UIMin = "64", UIMax = "2048", ClampMin = "2", ClampMax = "4096",
-		EditCondition = "LayoutType == ERecomputeUVsPropertiesLayoutType::Repack"))
+		EditCondition = "LayoutType == ERecomputeUVsPropertiesLayoutType::Repack", EditConditionHides, HideEditConditionToggle))
 	int TextureResolution = 1024;
 
 	/** Scaling factor used for UV island normalization/scaling. This is only enabled when the Layout Type is set to Normalize to Bounds or Normalize to World. */
-	UPROPERTY(EditAnywhere, Category = "UV Layout",	meta = (UIMin = "0.001", UIMax = "10", ClampMin = "0.00001", ClampMax = "1000000.0",
-		EditCondition = "LayoutType == ERecomputeUVsPropertiesLayoutType::NormalizeToBounds || LayoutType == ERecomputeUVsPropertiesLayoutType::NormalizeToWorld"))
+	UPROPERTY(EditAnywhere, Category = "UV Layout", meta = (UIMin = "0.001", UIMax = "10", ClampMin = "0.00001", ClampMax = "1000000.0",
+		EditCondition = "LayoutType == ERecomputeUVsPropertiesLayoutType::NormalizeToBounds || LayoutType == ERecomputeUVsPropertiesLayoutType::NormalizeToWorld", EditConditionHides, HideEditConditionToggle))
 	float NormalizeScale = 1.0f;
+
+	/** Enable UDIM aware layout and keep islands within their originating UDIM tiles when laying out.*/
+	UPROPERTY(EditAnywhere, Category = "UV Layout", meta = (DisplayName = "Preserve UDIMs", EditCondition = "bUDIMCVAREnabled", EditConditionHides, HideEditConditionToggle = true))
+	bool bEnableUDIMLayout = false;
+
+	UPROPERTY(Transient)
+	bool bUDIMCVAREnabled = false;
 };
