MeshModelingTools: Fix TriEd flip edge operation selecting wrong edges afterward.

diff --git a/Engine/Plugins/Runtime/MeshModelingToolset/Source/MeshModelingTools/Private/EditMeshPolygonsTool.cpp b/Engine/Plugins/Runtime/MeshModelingToolset/Source/MeshModelingTools/Private/EditMeshPolygonsTool.cpp
--- a/Engine/Plugins/Runtime/MeshModelingToolset/Source/MeshModelingTools/Private/EditMeshPolygonsTool.cpp
+++ b/Engine/Plugins/Runtime/MeshModelingToolset/Source/MeshModelingTools/Private/EditMeshPolygonsTool.cpp
@@ -2163,14 +2163,14 @@ void UEditMeshPolygonsTool::ApplyPokeSingleFace()
 
 void UEditMeshPolygonsTool::ApplyFlipSingleEdge()
 {
-	if (BeginMeshEdgeEditChange() == false)
+	FDynamicMesh3* Mesh = CurrentMesh.Get();
+
+	if (!BeginMeshEdgeEditChange([Mesh](int32 Eid) {return !Mesh->IsBoundaryEdge(Eid) && !Mesh->Attributes()->IsSeamEdge(Eid); }))
 	{
-		GetToolManager()->DisplayMessage(LOCTEXT("OnFlipFailedMessage", "Cannot Flip Current Selection"), EToolMessageLevel::UserWarning);
+		GetToolManager()->DisplayMessage(LOCTEXT("OnFlipFailedMessage", "Cannot Flip Current Selection (no non-seam edges selected)"), EToolMessageLevel::UserWarning);
 		return;
 	}
 
-	FDynamicMesh3* Mesh = CurrentMesh.Get();
-	FGroupTopologySelection ActiveSelection = SelectionMechanic->GetActiveSelection();
 	FDynamicMeshChangeTracker ChangeTracker(Mesh);
 	ChangeTracker.BeginChange();
 	for (FSelectedEdge& Edge : ActiveEdgeSelection)
@@ -2186,9 +2186,24 @@ void UEditMeshPolygonsTool::ApplyFlipSingleEdge()
 		}
 	}
 
-	// Group topology may or may not change, but just assume that it does.
+	// After flipping edges, edge ID's stay the same but group edge id's in our topology end up swapping around after
+	// the topology rebuild (because they are assigned in order of iteration through triangles). In order to update them,
+	// we need to go ahead and do the topology rebuild now. This means that we don't actually need to do the rebuild
+	// that happens inside EmitCurrentMeshChangeAndUpdate, but it's not worth trying to refactor to avoid it, so we
+	// just end up doing an extraneous update.
+	FGroupTopologySelection NewSelection;
+	Topology->RebuildTopology();
+	for (FSelectedEdge& Edge : ActiveEdgeSelection)
+	{
+		int32 Eid = Edge.EdgeIDs[0];
+		if (Mesh->IsEdge(Eid))
+		{
+			NewSelection.SelectedEdgeIDs.Add(Topology->FindGroupEdgeID(Eid));
+		}
+	}
+
 	EmitCurrentMeshChangeAndUpdate(LOCTEXT("PolyMeshFlipChange", "Flip Edges"),
-		ChangeTracker.EndChange(), ActiveSelection);
+		ChangeTracker.EndChange(), NewSelection);
 }
 
 void UEditMeshPolygonsTool::ApplyCollapseSingleEdge()
