MeshModelingTools: Fix crash when trying to read from texture that does not have source data in editor.

diff --git a/Engine/Plugins/Runtime/MeshModelingToolset/Source/ModelingComponents/Private/AssetUtils/Texture2DUtil.cpp b/Engine/Plugins/Runtime/MeshModelingToolset/Source/ModelingComponents/Private/AssetUtils/Texture2DUtil.cpp
--- a/Engine/Plugins/Runtime/MeshModelingToolset/Source/ModelingComponents/Private/AssetUtils/Texture2DUtil.cpp
+++ b/Engine/Plugins/Runtime/MeshModelingToolset/Source/ModelingComponents/Private/AssetUtils/Texture2DUtil.cpp
@@ -39,13 +39,18 @@ static bool ReadTexture_PlatformData(
 
 	// the Platform BulkData for most texture formats is compressed and we cannot process it on
 	// the CPU. However TC_VectorDisplacementmap is uncompressed RGBA8.
-	// in the Editor, we can temporarily change the texture type and rebuild the PlatformData.
+	// in the Editor, if we have source data, we can temporarily change the texture type and rebuild the PlatformData.
 	// However at Runtime we cannot, and so we can only error-out if we do not have a readable compression type.
 #if WITH_EDITOR
 	const TextureMipGenSettings InitialMipGenSettings = TextureMap->MipGenSettings;
 	bool bNeedToRevertTextureChanges = false;
 	if (InitialMipGenSettings != TextureMipGenSettings::TMGS_NoMipmaps || InitialCompressionSettings != TextureCompressionSettings::TC_VectorDisplacementmap)
 	{
+		if (!TextureMap->Source.IsValid())
+		{
+			return false;
+		}
+
 		TextureMap->CompressionSettings = TextureCompressionSettings::TC_VectorDisplacementmap;
 		TextureMap->MipGenSettings = TextureMipGenSettings::TMGS_NoMipmaps;
 		TextureMap->UpdateResource();
