Fix missing header guard in RenderCaptureFunctions.h (Hoard Issue 231728)

#rnx

diff --git a/Engine/Plugins/Runtime/MeshModelingToolset/Source/ModelingComponents/Public/Baking/RenderCaptureFunctions.h b/Engine/Plugins/Runtime/MeshModelingToolset/Source/ModelingComponents/Public/Baking/RenderCaptureFunctions.h
--- a/Engine/Plugins/Runtime/MeshModelingToolset/Source/ModelingComponents/Public/Baking/RenderCaptureFunctions.h
+++ b/Engine/Plugins/Runtime/MeshModelingToolset/Source/ModelingComponents/Public/Baking/RenderCaptureFunctions.h
@@ -1,5 +1,7 @@
 ﻿// Copyright Epic Games, Inc. All Rights Reserved.
 
+#pragma once
+
 #include "Scene/SceneCapturePhotoSet.h"
 #include "Sampling/MeshBakerCommon.h"
 #include "Image/ImageInfilling.h"
