Fill out the missing tooltips in PolyEd and TriEd tools, and tweak some of the tooltip wording.

diff --git a/Engine/Plugins/Runtime/MeshModelingToolset/Source/MeshModelingTools/Public/EditMeshPolygonsTool.h b/Engine/Plugins/Runtime/MeshModelingToolset/Source/MeshModelingTools/Public/EditMeshPolygonsTool.h
--- a/Engine/Plugins/Runtime/MeshModelingToolset/Source/MeshModelingTools/Public/EditMeshPolygonsTool.h
+++ b/Engine/Plugins/Runtime/MeshModelingToolset/Source/MeshModelingTools/Public/EditMeshPolygonsTool.h
@@ -235,30 +235,35 @@ public:
 	UFUNCTION(CallInEditor, Category = FaceEdits, meta = (DisplayName = "Push/Pull", DisplayPriority = 1))
 	void PushPull() { PostAction(EEditMeshPolygonsToolActions::PushPull); }
 
-	/** Like Extrude, but defaults to moving verts along vertex normals instead of a single direction. */
+	/** Like Extrude, but defaults to moving verts along vertex normals instead of a single direction.*/
 	UFUNCTION(CallInEditor, Category = FaceEdits, meta = (DisplayName = "Offset", DisplayPriority = 1))
 	void Offset() { PostAction(EEditMeshPolygonsToolActions::Offset); }
 
-	/** Inset the current set of selected faces. Click in viewport to confirm inset distance. */
+	/**
+	 * Inset the current set of selected faces. Click in viewport to confirm inset distance.
+	 * 
+	 * (An Inset operation stitches in a smaller version of selected faces inside the existing ones)
+	 */
 	UFUNCTION(CallInEditor, Category = FaceEdits, meta = (DisplayName = "Inset", DisplayPriority = 2))
 	void Inset() { PostAction(EEditMeshPolygonsToolActions::Inset);	}
 
-	/** Outset the current set of selected faces. Click in viewport to confirm outset distance. */
+	/**
+	 * Outset the current set of selected faces. Click in viewport to confirm outset distance.
+	 * 
+	 * (An Outset operation stitches in a larger version of selected faces inside the existing ones)
+	 */
 	UFUNCTION(CallInEditor, Category = FaceEdits, meta = (DisplayName = "Outset", DisplayPriority = 3))
 	void Outset() { PostAction(EEditMeshPolygonsToolActions::Outset);	}
 
-	//~ TODO: Make the Merge, Delete, and Flip comments visible as tooltips. Currently we can't due to a bug that
-	//~ limits our total tooltip text allotment: UE-124608
-
-	//~ Bevel the edge loops around the selected faces 
+	/** Bevel the edge loops around the selected faces, inserting edge-aligned faces that interpolate the normals of the selected faces */
 	UFUNCTION(CallInEditor, Category = FaceEdits, meta = (DisplayName = "Bevel", DisplayPriority = 4))
 	void Bevel() { PostAction(EEditMeshPolygonsToolActions::BevelFaces); }
 
-	//~ Merge the current set of selected faces into a single face.
+	/** Merge the current set of selected faces into a single face */
 	UFUNCTION(CallInEditor, Category = FaceEdits, meta = (DisplayName = "Merge", DisplayPriority = 4))
 	void Merge() { PostAction(EEditMeshPolygonsToolActions::Merge);	}
 
-	//~ Delete the current set of selected faces
+	/** Delete the current set of selected faces */
 	UFUNCTION(CallInEditor, Category = FaceEdits, meta = (DisplayName = "Delete", DisplayPriority = 4))
 	void Delete() { PostAction(EEditMeshPolygonsToolActions::Delete); }
 
@@ -270,7 +275,7 @@ public:
 	UFUNCTION(CallInEditor, Category = FaceEdits, meta = (DisplayName = "RecalcNormals", DisplayPriority = 6))
 	void RecalcNormals() { PostAction(EEditMeshPolygonsToolActions::RecalculateNormals); }
 
-	//~ Flip normals and face orientation for the current set of selected faces
+	/** Flip normalsand face orientation for the current set of selected faces */
 	UFUNCTION(CallInEditor, Category = FaceEdits, meta = (DisplayName = "Flip", DisplayPriority = 7))
 	void Flip() { PostAction(EEditMeshPolygonsToolActions::FlipNormals); }
 
@@ -290,11 +295,11 @@ public:
 	UFUNCTION(CallInEditor, Category = FaceEdits, meta = (DisplayName = "Duplicate", DisplayPriority = 12))
 	void Duplicate() { PostAction(EEditMeshPolygonsToolActions::Duplicate); }
 
-	//~ TODO: add tooltip, especially explaining limitations, for this and InsertEdge. Can't currently do that
-	//~ without cutting down another tooltip until UE-124608 is fixed...
+	/** Insert a chain of edges across quads (faces with four edges) in the mesh. Due to ambiguity, edges will not be inserted on non-quad faces. */
 	UFUNCTION(CallInEditor, Category = ShapeEdits, meta = (DisplayName = "InsertEdgeLoop", DisplayPriority = 13))
 	void InsertEdgeLoop() { PostAction(EEditMeshPolygonsToolActions::InsertEdgeLoop); }
 
+	/** Insert a new edge connecting existing edges or vertices on a single face */
 	UFUNCTION(CallInEditor, Category = ShapeEdits, meta = (DisplayName = "Insert Edge", DisplayPriority = 14))
 	void InsertEdge() { PostAction(EEditMeshPolygonsToolActions::InsertEdge); }
 
@@ -355,7 +360,7 @@ public:
 	UFUNCTION(CallInEditor, Category = TriangleEdits, meta = (DisplayName = "Duplicate", DisplayPriority = 12))
 	void Duplicate() { PostAction(EEditMeshPolygonsToolActions::Duplicate); }
 
-	/** Poke each face at its center point */
+	/** Insert a new vertex at the center of each selected face */
 	UFUNCTION(CallInEditor, Category = TriangleEdits, meta = (DisplayName = "Poke", DisplayPriority = 13))
 	void Poke() { PostAction(EEditMeshPolygonsToolActions::PokeSingleFace); }
 };
@@ -388,9 +393,11 @@ class MESHMODELINGTOOLS_API UEditMeshPolygonsToolEdgeActions : public UEditMeshP
 {
 	GENERATED_BODY()
 public:
+	/** Merge selected edges, moving the first edge to the second */
 	UFUNCTION(CallInEditor, Category = EdgeEdits, meta = (DisplayName = "Weld", DisplayPriority = 1))
 	void Weld() { PostAction(EEditMeshPolygonsToolActions::WeldEdges); }
 
+	/** Make each selected polygroup edge follow a straight path between its endpoints */
 	UFUNCTION(CallInEditor, Category = EdgeEdits, meta = (DisplayName = "Straighten", DisplayPriority = 2))
 	void Straighten() { PostAction(EEditMeshPolygonsToolActions::StraightenEdge); }
 
@@ -398,9 +405,11 @@ public:
 	UFUNCTION(CallInEditor, Category = EdgeEdits, meta = (DisplayName = "Fill Hole", DisplayPriority = 3))
 	void FillHole()	{ PostAction(EEditMeshPolygonsToolActions::FillHole); }
 
+	/** Bevel the selected edges, replacing them with angled faces */
 	UFUNCTION(CallInEditor, Category = EdgeEdits, meta = (DisplayName = "Bevel", DisplayPriority = 4))
 	void Bevel() { PostAction(EEditMeshPolygonsToolActions::BevelEdges); }
 	
+	/** Create a new face that connects the selected edges */
 	UFUNCTION(CallInEditor, Category = EdgeEdits, meta = (DisplayName = "Bridge", DisplayPriority = 5))
 	void Bridge() { PostAction(EEditMeshPolygonsToolActions::BridgeEdges); }
 };
@@ -411,6 +420,7 @@ class MESHMODELINGTOOLS_API UEditMeshPolygonsToolEdgeActions_Triangles : public
 {
 	GENERATED_BODY()
 public:
+	/** Merge selected edges, moving the first edge to the second */
 	UFUNCTION(CallInEditor, Category = EdgeEdits, meta = (DisplayName = "Weld", DisplayPriority = 1))
 	void Weld() { PostAction(EEditMeshPolygonsToolActions::WeldEdges); }
 
@@ -418,12 +428,15 @@ public:
 	UFUNCTION(CallInEditor, Category = EdgeEdits, meta = (DisplayName = "Fill Hole", DisplayPriority = 1))
 	void FillHole() { PostAction(EEditMeshPolygonsToolActions::FillHole); }
 
+	/** Collapse the selected edges, deleting the attached triangles and merging its two vertices into one */
 	UFUNCTION(CallInEditor, Category = EdgeEdits, meta = (DisplayName = "Collapse", DisplayPriority = 1))
 	void Collapse() { PostAction(EEditMeshPolygonsToolActions::CollapseSingleEdge); }
 
+	/** Flip the selected (non-border, non-seam) edges, replacing them with new edges in the crossing direction */
 	UFUNCTION(CallInEditor, Category = EdgeEdits, meta = (DisplayName = "Flip", DisplayPriority = 1))
 	void Flip() { PostAction(EEditMeshPolygonsToolActions::FlipSingleEdge); }
 
+	/** Split the selected edges, inserting a new vertex at each edge midpoint */
 	UFUNCTION(CallInEditor, Category = EdgeEdits, meta = (DisplayName = "Split", DisplayPriority = 1))
 	void Split() { PostAction(EEditMeshPolygonsToolActions::SplitSingleEdge); }
 
