PolyEdit Insert Edge Activity: Optionally allow users to chain edge insertions by automatically selecting the end point of the just inserted edge as the new start point.

diff --git a/Engine/Plugins/Runtime/MeshModelingToolset/Source/MeshModelingTools/Private/ToolActivities/PolyEditInsertEdgeActivity.cpp b/Engine/Plugins/Runtime/MeshModelingToolset/Source/MeshModelingTools/Private/ToolActivities/PolyEditInsertEdgeActivity.cpp
--- a/Engine/Plugins/Runtime/MeshModelingToolset/Source/MeshModelingTools/Private/ToolActivities/PolyEditInsertEdgeActivity.cpp
+++ b/Engine/Plugins/Runtime/MeshModelingToolset/Source/MeshModelingTools/Private/ToolActivities/PolyEditInsertEdgeActivity.cpp
@@ -259,6 +259,19 @@ void UPolyEditInsertEdgeActivity::Tick(float DeltaTime)
 				ChangeTracker.EndChange(), EmptySelection);
 
 			ToolState = EState::GettingStart;
+			
+			if (Settings->bContinuousInsertion)
+			{
+				// If continuous insertion is enabled, try to find information associated with new start point.
+				// If found, remain in GettingEnd mode.
+				FVector3d PreviewPoint;
+				if (GetHoveredItem(LastEndPointWorldRay, StartPoint, StartTopologyID, bStartIsCorner, PreviewPoint))
+				{
+					PreviewPoints.Reset();
+					PreviewPoints.Add(PreviewPoint);
+					ToolState = EState::GettingEnd;
+				}
+			}
 		}
 		else
 		{
@@ -599,6 +612,9 @@ void UPolyEditInsertEdgeActivity::OnClicked(const FInputDeviceRay& ClickPos)
 				ClearPreview(false);
 			}
 		}
+
+		LastEndPointWorldRay = ClickPos.WorldRay;
+			
 		break;
 	}
 	}
diff --git a/Engine/Plugins/Runtime/MeshModelingToolset/Source/MeshModelingTools/Public/ToolActivities/PolyEditInsertEdgeActivity.h b/Engine/Plugins/Runtime/MeshModelingToolset/Source/MeshModelingTools/Public/ToolActivities/PolyEditInsertEdgeActivity.h
--- a/Engine/Plugins/Runtime/MeshModelingToolset/Source/MeshModelingTools/Public/ToolActivities/PolyEditInsertEdgeActivity.h
+++ b/Engine/Plugins/Runtime/MeshModelingToolset/Source/MeshModelingTools/Public/ToolActivities/PolyEditInsertEdgeActivity.h
@@ -42,6 +42,10 @@ public:
 	UPROPERTY(EditAnywhere, Category = InsertEdge)
 	EGroupEdgeInsertionMode InsertionMode = EGroupEdgeInsertionMode::PlaneCut;
 
+	/** If true, edge insertions are chained together so that each end point becomes the new start point */
+	UPROPERTY(EditAnywhere, Category = InsertEdge)
+	bool bContinuousInsertion = true;
+	
 	/** How close a new loop edge needs to pass next to an existing vertex to use that vertex rather than creating a new one (used for plane cut). */
 	UPROPERTY(EditAnywhere, Category = InsertEdge, AdvancedDisplay, meta = (UIMin = "0", UIMax = "0.01", ClampMin = "0", ClampMax = "10"))
 	double VertexTolerance = 0.001;
@@ -131,6 +135,7 @@ protected:
 	int32 CommonGroupID = FDynamicMesh3::InvalidID;
 	int32 CommonBoundaryIndex = FDynamicMesh3::InvalidID;
 
+	FRay LastEndPointWorldRay;
 
 	// State control:
 	EState ToolState = EState::GettingStart;
