MeshModelingToolset: Fix SaveGeneratedTexture2DAsset not preserving sRGB state when overwriting an exisiting texture asset

#rnx

diff --git a/Engine/Plugins/Runtime/MeshModelingToolset/Source/ModelingComponents/Public/AssetUtils/Texture2DBuilder.h b/Engine/Plugins/Runtime/MeshModelingToolset/Source/ModelingComponents/Public/AssetUtils/Texture2DBuilder.h
--- a/Engine/Plugins/Runtime/MeshModelingToolset/Source/ModelingComponents/Public/AssetUtils/Texture2DBuilder.h
+++ b/Engine/Plugins/Runtime/MeshModelingToolset/Source/ModelingComponents/Public/AssetUtils/Texture2DBuilder.h
@@ -17,7 +17,7 @@ namespace Geometry
  * Use Initialize() functions to configure, can either generate a new UTexture2D (in the Transient package) or modify an existing UTexture2D.
  *
  * Currently the generated UTexture2D will only have Mip 0, and only Mip 0 can be edited.
- * The generated UTexture2D has format PF_B8G8R8A8.
+ * The generated UTexture2D has format PF_B8G8R8A8 or PF_FloatRGBA.
  *
  * Use Commit() to lock and update the texture after editing is complete. LockForEditing() can be used to re-open.
  * By default textures are locked for editing on Initialize()
diff --git a/Engine/Plugins/Runtime/MeshModelingToolset/Source/ModelingComponentsEditorOnly/Private/AssetUtils/CreateTexture2DUtil.cpp b/Engine/Plugins/Runtime/MeshModelingToolset/Source/ModelingComponentsEditorOnly/Private/AssetUtils/CreateTexture2DUtil.cpp
--- a/Engine/Plugins/Runtime/MeshModelingToolset/Source/ModelingComponentsEditorOnly/Private/AssetUtils/CreateTexture2DUtil.cpp
+++ b/Engine/Plugins/Runtime/MeshModelingToolset/Source/ModelingComponentsEditorOnly/Private/AssetUtils/CreateTexture2DUtil.cpp
@@ -48,7 +48,7 @@ UE::AssetUtils::ECreateTexture2DResult UE::AssetUtils::SaveGeneratedTexture2DAss
 		if (ExistingTexture2D && Options.bOverwriteIfExists)
 		{
 			// TODO: Direct copy between UTexture2D
-			UE::Geometry::TImageBuilder<FVector4f> TexImage;
+			UE::Geometry::TImageBuilder<FVector4f> TexImage; // This will contain FLinearColor data
 			if (!UE::AssetUtils::ReadTexture(GeneratedTexture, TexImage))
 			{
 				return ECreateTexture2DResult::InvalidInputTexture;
@@ -56,7 +56,7 @@ UE::AssetUtils::ECreateTexture2DResult UE::AssetUtils::SaveGeneratedTexture2DAss
 			
 			UE::Geometry::FTexture2DBuilder Builder;
 			Builder.InitializeAndReplaceExistingTexture(ExistingTexture2D, GeneratedTexture);
-			Builder.Copy(TexImage);
+			Builder.Copy(TexImage, GeneratedTexture->SRGB);
 			Builder.Commit(true);
 			GeneratedTexture = Builder.GetTexture2D();
 		}
