RenderCaptureBaking: Improve performance when adding/removing capture channels

- When adding a new channel: Only the newly requested channel is captured, previously all enabled channels were (re)captured
- When removing an already captured channel: We no longer recompute all remaining channels, we just wipe the one that was removed

#rnx

diff --git a/Engine/Plugins/Experimental/MeshModelingToolsetExp/Source/MeshModelingToolsExp/Private/BakeRenderCaptureTool.cpp b/Engine/Plugins/Experimental/MeshModelingToolsetExp/Source/MeshModelingToolsExp/Private/BakeRenderCaptureTool.cpp
--- a/Engine/Plugins/Experimental/MeshModelingToolsetExp/Source/MeshModelingToolsExp/Private/BakeRenderCaptureTool.cpp
+++ b/Engine/Plugins/Experimental/MeshModelingToolsetExp/Source/MeshModelingToolsExp/Private/BakeRenderCaptureTool.cpp
@@ -602,7 +602,7 @@ void UBakeRenderCaptureTool::UpdateResult()
 	}
 
 	//
-	// create a set of spatially located render captures of the scene ("photo set"). We need to recompute this if the
+	// Create a set of spatially located render captures of the scene ("photo set"). We need to recompute this if the
 	// render capture properties changed. Note we only compare the URenderCaptureProperties, and not the
 	// ValidSampleDepthThreshold, this is intentional so that we only trigger a scene capture recompute when we go from
 	// a zero to a positive threshold (we need to compute the depth capture), or a positive to a zero threshold (we can
@@ -619,13 +619,26 @@ void UBakeRenderCaptureTool::UpdateResult()
 		// Do not allow user-cancellation on the call that occurs when the Render Capture Tool starts up
 		const bool bAllowCancel = (bFirstEverSceneCapture == false);
 
-		SceneCapture.Reset();
 		FRenderCaptureOptions Options = MakeRenderCaptureOptions(*RenderCaptureProperties, *Settings, *InputMeshSettings);
 		{
 			FScopedSlowTask Progress(1.f, LOCTEXT("CapturingScene", "Capturing Scene..."));
 			Progress.EnterProgressFrame(1.f);
 			Progress.MakeDialog(bAllowCancel);
-			SceneCapture = CapturePhotoSet(Actors, Options, bAllowCancel);
+
+			const bool bRequireFullRecompute = (
+				RenderCaptureProperties->Resolution         != ComputedRenderCaptureProperties->Resolution ||
+				RenderCaptureProperties->bAntiAliasing      != ComputedRenderCaptureProperties->bAntiAliasing ||
+				RenderCaptureProperties->CaptureFieldOfView != ComputedRenderCaptureProperties->CaptureFieldOfView ||
+				RenderCaptureProperties->NearPlaneDist      != ComputedRenderCaptureProperties->NearPlaneDist);
+
+			if (!SceneCapture || (SceneCapture && bRequireFullRecompute))
+			{
+				SceneCapture = CapturePhotoSet(Actors, Options, bAllowCancel);
+			}
+			else
+			{
+				UpdatePhotoSet(SceneCapture, Actors, Options, bAllowCancel);
+			}
 		}
 
 		for (int Idx = 1; Idx < Targets.Num(); ++Idx)
@@ -736,12 +749,12 @@ void UBakeRenderCaptureTool::UpdateResult()
 			UpdateVisualization(); // Clear the preview mesh material inputs
 		}
 		OpState = EBakeOpState::Invalid;
-		return;
 	}
-
-	InvalidateComputeRC();
-
-	OpState = EBakeOpState::Clean;
+	else
+	{
+		InvalidateComputeRC();
+		OpState = EBakeOpState::Clean;
+	}
 }
 
 
diff --git a/Engine/Plugins/Runtime/MeshModelingToolset/Source/ModelingComponents/Private/Baking/RenderCaptureFunctions.cpp b/Engine/Plugins/Runtime/MeshModelingToolset/Source/ModelingComponents/Private/Baking/RenderCaptureFunctions.cpp
--- a/Engine/Plugins/Runtime/MeshModelingToolset/Source/ModelingComponents/Private/Baking/RenderCaptureFunctions.cpp
+++ b/Engine/Plugins/Runtime/MeshModelingToolset/Source/ModelingComponents/Private/Baking/RenderCaptureFunctions.cpp
@@ -214,12 +214,25 @@ TUniquePtr<FSceneCapturePhotoSet> UE::Geometry::CapturePhotoSet(
 {
 	TRACE_CPUPROFILER_EVENT_SCOPE(CapturePhotoSet);
 
+	TUniquePtr<FSceneCapturePhotoSet> SceneCapture = MakeUnique<FSceneCapturePhotoSet>();
+
+	UpdatePhotoSet(SceneCapture,Actors, Options, bAllowCancel);
+
+	return SceneCapture;
+}
+
+void UE::Geometry::UpdatePhotoSet(
+	TUniquePtr<FSceneCapturePhotoSet>& SceneCapture,
+	const TArray<TObjectPtr<AActor>>& Actors,
+	const FRenderCaptureOptions& Options,
+	bool bAllowCancel
+)
+{
 	double FieldOfView = Options.FieldOfViewDegrees;
 	double NearPlaneDist = Options.NearPlaneDist;
 
 	FImageDimensions CaptureDimensions(Options.RenderCaptureImageSize, Options.RenderCaptureImageSize);
 
-	TUniquePtr<FSceneCapturePhotoSet> SceneCapture = MakeUnique<FSceneCapturePhotoSet>();
 	SceneCapture->SetAllowCancel(bAllowCancel);
 
 	SceneCapture->SetCaptureTypeEnabled(ERenderCaptureType::DeviceDepth, Options.bBakeDeviceDepth);
@@ -246,8 +259,6 @@ TUniquePtr<FSceneCapturePhotoSet> UE::Geometry::CapturePhotoSet(
 	SceneCapture->AddStandardExteriorCapturesFromBoundingBox(
 		CaptureDimensions, FieldOfView, NearPlaneDist,
 		true, true, true, true, true);
-
-	return SceneCapture;
 }
 
 
diff --git a/Engine/Plugins/Runtime/MeshModelingToolset/Source/ModelingComponents/Private/Scene/SceneCapturePhotoSet.cpp b/Engine/Plugins/Runtime/MeshModelingToolset/Source/ModelingComponents/Private/Scene/SceneCapturePhotoSet.cpp
--- a/Engine/Plugins/Runtime/MeshModelingToolset/Source/ModelingComponents/Private/Scene/SceneCapturePhotoSet.cpp
+++ b/Engine/Plugins/Runtime/MeshModelingToolset/Source/ModelingComponents/Private/Scene/SceneCapturePhotoSet.cpp
@@ -56,7 +56,7 @@ void FSceneCapturePhotoSet::SetCaptureConfig(ERenderCaptureType CaptureType, con
 			DeviceDepthConfig = Config;
 			break;
 		default:
-			check(false);
+			ensure(false);
 	}
 }
 
@@ -81,41 +81,101 @@ FRenderCaptureConfig FSceneCapturePhotoSet::GetCaptureConfig(ERenderCaptureType
 		case ERenderCaptureType::DeviceDepth:
 			return DeviceDepthConfig;
 		default:
-			check(false);
+			ensure(false);
 	}
 	return BaseColorConfig;
 }
 
+void FSceneCapturePhotoSet::DisableAllCaptureTypes()
+{
+	SetCaptureTypeEnabled(ERenderCaptureType::BaseColor, false);
+	SetCaptureTypeEnabled(ERenderCaptureType::WorldNormal, false);
+	SetCaptureTypeEnabled(ERenderCaptureType::Roughness, false);
+	SetCaptureTypeEnabled(ERenderCaptureType::Metallic, false);
+	SetCaptureTypeEnabled(ERenderCaptureType::Specular, false);
+	SetCaptureTypeEnabled(ERenderCaptureType::Emissive, false);
+	SetCaptureTypeEnabled(ERenderCaptureType::CombinedMRS, false);
+	SetCaptureTypeEnabled(ERenderCaptureType::DeviceDepth, false);
+}
+
+
 void FSceneCapturePhotoSet::SetCaptureTypeEnabled(ERenderCaptureType CaptureType, bool bEnabled)
 {
 	switch (CaptureType)
 	{
 	case ERenderCaptureType::BaseColor:
-			bEnableBaseColor = bEnabled;
+		bEnableBaseColor = bEnabled;
+		break;
+	case ERenderCaptureType::WorldNormal:
+		bEnableWorldNormal = bEnabled;
+		break;
+	case ERenderCaptureType::Roughness:
+		bEnableRoughness = bEnabled;
+		break;
+	case ERenderCaptureType::Metallic:
+		bEnableMetallic = bEnabled;
+		break;
+	case ERenderCaptureType::Specular:
+		bEnableSpecular = bEnabled;
+		break;
+	case ERenderCaptureType::Emissive:
+		bEnableEmissive = bEnabled;
+		break;
+	case ERenderCaptureType::CombinedMRS:
+		bEnablePackedMRS = bEnabled;
+		break;
+	case ERenderCaptureType::DeviceDepth:
+		bEnableDeviceDepth = bEnabled;
+		break;
+	default:
+		ensure(false);
+	}
+
+	// If we're disabling a CaptureType we need to remove any existing photos
+	if (bEnabled == false)
+	{
+		switch (CaptureType)
+		{
+		case ERenderCaptureType::BaseColor:
+			BaseColorPhotoSet.Empty();
 			break;
 		case ERenderCaptureType::WorldNormal:
-			bEnableWorldNormal = bEnabled;
+			WorldNormalPhotoSet.Empty();
 			break;
 		case ERenderCaptureType::Roughness:
-			bEnableRoughness = bEnabled;
+			RoughnessPhotoSet.Empty();
 			break;
 		case ERenderCaptureType::Metallic:
-			bEnableMetallic = bEnabled;
+			MetallicPhotoSet.Empty();
 			break;
 		case ERenderCaptureType::Specular:
-			bEnableSpecular = bEnabled;
+			SpecularPhotoSet.Empty();
 			break;
 		case ERenderCaptureType::Emissive:
-			bEnableEmissive = bEnabled;
+			EmissivePhotoSet.Empty();
 			break;
 		case ERenderCaptureType::CombinedMRS:
-			bEnablePackedMRS = bEnabled;
+			PackedMRSPhotoSet.Empty();
 			break;
 		case ERenderCaptureType::DeviceDepth:
-			bEnableDeviceDepth = bEnabled;
+			DeviceDepthPhotoSet.Empty();
 			break;
 		default:
-			check(false);
+			ensure(false);
+		}
+	}
+
+	// If all photo sets are disabled we should also clear the photo set params
+	if (!bEnableBaseColor &&
+		!bEnableRoughness &&
+		!bEnableSpecular &&
+		!bEnableMetallic &&
+		!bEnablePackedMRS &&
+		!bEnableWorldNormal &&
+		!bEnableEmissive &&
+		!bEnableDeviceDepth)
+	{
+		PhotoSetParams.Empty();
 	}
 }
 
@@ -262,42 +322,57 @@ void FSceneCapturePhotoSet::AddExteriorCaptures(
 		FVector3d ViewDirection = Directions[di];
 		ViewDirection.Normalize();
 
+		FSpatialPhotoParams Params;
+		Params.NearPlaneDist = NearPlaneDist;
+		Params.HorzFOVDegrees = HorizontalFOVDegrees;
+		Params.Dimensions = PhotoDimensions;
 		// TODO Align the frame with the renderer coordinate system then remove the axis swapping in WorldRenderCapture.cpp
-		FFrame3d ViewFrame;
-		ViewFrame.AlignAxis(0, ViewDirection);
-		ViewFrame.ConstrainedAlignAxis(2, FVector3d::UnitZ(), ViewFrame.X());
-		ViewFrame.Origin = (FVector3d)RenderSphere.Center;
-		ViewFrame.Origin -= (double)RenderSphere.W * ViewFrame.X();
-
-		FSpatialPhoto3f BasePhoto3f;
-		BasePhoto3f.Frame = ViewFrame;
-		BasePhoto3f.NearPlaneDist = NearPlaneDist;
-		BasePhoto3f.HorzFOVDegrees = HorizontalFOVDegrees;
-		BasePhoto3f.Dimensions = PhotoDimensions;
-
-		auto CaptureImageTypeFunc_3f = [this, &Progress, &BasePhoto3f, &RenderCapture](ERenderCaptureType CaptureType, FSpatialPhotoSet3f& PhotoSet)
+		Params.Frame.AlignAxis(0, ViewDirection);
+		Params.Frame.ConstrainedAlignAxis(2, FVector3d::UnitZ(), Params.Frame.X());
+		Params.Frame.Origin = RenderSphere.Center;
+		Params.Frame.Origin -= RenderSphere.W * Params.Frame.X();
+
+		auto CaptureImageTypeFunc_3f = [this, &Progress, &RenderCapture, &Params, &NumDirections]
+			(ERenderCaptureType CaptureType, FSpatialPhotoSet3f& PhotoSet)
 		{
-			FSpatialPhoto3f NewPhoto = BasePhoto3f;
-			FImageAdapter Image(&NewPhoto.Image);
-			FRenderCaptureConfig Config = GetCaptureConfig(CaptureType);
-			RenderCapture.CaptureFromPosition(CaptureType, NewPhoto.Frame, NewPhoto.HorzFOVDegrees, NewPhoto.NearPlaneDist, Image, Config);
-			PhotoSet.Add(MoveTemp(NewPhoto));
+			// Test NumDirections to avoid recomputing photo sets. Search :SceneCaptureWithExistingCaptures
+			if (PhotoSet.Num() < NumDirections)
+			{
+				FSpatialPhoto3f NewPhoto;
+				NewPhoto.Frame = Params.Frame;
+				NewPhoto.NearPlaneDist = Params.NearPlaneDist;
+				NewPhoto.HorzFOVDegrees = Params.HorzFOVDegrees;
+				NewPhoto.Dimensions = Params.Dimensions;
+
+				// TODO Do something with the success boolean returned by RenderCapture.CaptureFromPosition
+				FImageAdapter Image(&NewPhoto.Image);
+				FRenderCaptureConfig Config = GetCaptureConfig(CaptureType);
+				RenderCapture.CaptureFromPosition(CaptureType, NewPhoto.Frame, NewPhoto.HorzFOVDegrees, NewPhoto.NearPlaneDist, Image, Config);
+				PhotoSet.Add(MoveTemp(NewPhoto));
+			}
+
 			Progress.TickProgress();
 		};
 
-		FSpatialPhoto1f BasePhoto1f;
-		BasePhoto1f.Frame = ViewFrame;
-		BasePhoto1f.NearPlaneDist = NearPlaneDist;
-		BasePhoto1f.HorzFOVDegrees = HorizontalFOVDegrees;
-		BasePhoto1f.Dimensions = PhotoDimensions;
-
-		auto CaptureImageTypeFunc_1f = [this, &Progress, &BasePhoto1f, &RenderCapture](ERenderCaptureType CaptureType, FSpatialPhotoSet1f& PhotoSet)
+		auto CaptureImageTypeFunc_1f = [this, &Progress, &RenderCapture, &Params, &NumDirections]
+			(ERenderCaptureType CaptureType, FSpatialPhotoSet1f& PhotoSet)
 		{
-			FSpatialPhoto1f NewPhoto = BasePhoto1f;
-			FImageAdapter Image(&NewPhoto.Image);
-			FRenderCaptureConfig Config = GetCaptureConfig(CaptureType);
-			RenderCapture.CaptureFromPosition(CaptureType, NewPhoto.Frame, NewPhoto.HorzFOVDegrees, NewPhoto.NearPlaneDist, Image, Config);
-			PhotoSet.Add(MoveTemp(NewPhoto));
+			// Test NumDirections to avoid recomputing photo sets. Search :SceneCaptureWithExistingCaptures
+			if (PhotoSet.Num() < NumDirections)
+			{
+				FSpatialPhoto1f NewPhoto;
+				NewPhoto.Frame = Params.Frame;
+				NewPhoto.NearPlaneDist = Params.NearPlaneDist;
+				NewPhoto.HorzFOVDegrees = Params.HorzFOVDegrees;
+				NewPhoto.Dimensions = Params.Dimensions;
+
+				// TODO Do something with the success boolean returned by RenderCapture.CaptureFromPosition
+				FImageAdapter Image(&NewPhoto.Image);
+				FRenderCaptureConfig Config = GetCaptureConfig(CaptureType);
+				RenderCapture.CaptureFromPosition(CaptureType, NewPhoto.Frame, NewPhoto.HorzFOVDegrees, NewPhoto.NearPlaneDist, Image, Config);
+				PhotoSet.Add(MoveTemp(NewPhoto));
+			}
+
 			Progress.TickProgress();
 		};
 
@@ -334,13 +409,14 @@ void FSceneCapturePhotoSet::AddExteriorCaptures(
 			CaptureImageTypeFunc_3f(ERenderCaptureType::Emissive, EmissivePhotoSet);
 		}
 
-		FSpatialPhotoParams Params;
-		Params.Frame = ViewFrame;
-		Params.NearPlaneDist = NearPlaneDist;
-		Params.HorzFOVDegrees = HorizontalFOVDegrees;
-		Params.Dimensions = PhotoDimensions;
-		Params.ViewMatrices = RenderCapture.GetLastCaptureViewMatrices();
-		PhotoSetParams.Add(Params);
+		// AddExteriorCaptures can be called on an FSceneCapturePhotoSet which already has some capture types computed
+		// and in this case we should not modify the existing photo sets/cached photo set parameters.
+		// See :SceneCaptureWithExistingCaptures 
+		if (PhotoSetParams.Num() < NumDirections)
+		{
+			Params.ViewMatrices = RenderCapture.GetLastCaptureViewMatrices();
+			PhotoSetParams.Add(Params);
+		}
 	}
 }
 
@@ -387,7 +463,7 @@ FVector3f FSceneCapturePhotoSet::FSceneSample::GetValue3f(ERenderCaptureType Cap
 	case ERenderCaptureType::DeviceDepth:
 		return DeviceDepth * FVector3f::One();
 	default:
-		check(false);
+		ensure(false);
 	}
 	return FVector3f::Zero();
 }
@@ -413,7 +489,7 @@ FVector4f FSceneCapturePhotoSet::FSceneSample::GetValue4f(ERenderCaptureType Cap
 	case ERenderCaptureType::DeviceDepth:
 		return FVector4f(DeviceDepth, DeviceDepth, DeviceDepth, 1.0f);
 	default:
-		check(false);
+		ensure(false);
 	}
 	return FVector4f::Zero();
 }
diff --git a/Engine/Plugins/Runtime/MeshModelingToolset/Source/ModelingComponents/Public/Baking/RenderCaptureFunctions.h b/Engine/Plugins/Runtime/MeshModelingToolset/Source/ModelingComponents/Public/Baking/RenderCaptureFunctions.h
--- a/Engine/Plugins/Runtime/MeshModelingToolset/Source/ModelingComponents/Public/Baking/RenderCaptureFunctions.h
+++ b/Engine/Plugins/Runtime/MeshModelingToolset/Source/ModelingComponents/Public/Baking/RenderCaptureFunctions.h
@@ -145,6 +145,13 @@ TUniquePtr<FSceneCapturePhotoSet> CapturePhotoSet(
 	const FRenderCaptureOptions& Options,
 	bool bAllowCancel);
 
+MODELINGCOMPONENTS_API
+void UpdatePhotoSet(
+	TUniquePtr<FSceneCapturePhotoSet>& SceneCapture,
+	const TArray<TObjectPtr<AActor>>& Actors,
+	const FRenderCaptureOptions& Options,
+	bool bAllowCancel);
+
 
 
 // Return a render capture baker, note the lifetime of all arguments such match the lifetime of the returned baker
diff --git a/Engine/Plugins/Runtime/MeshModelingToolset/Source/ModelingComponents/Public/Scene/SceneCapturePhotoSet.h b/Engine/Plugins/Runtime/MeshModelingToolset/Source/ModelingComponents/Public/Scene/SceneCapturePhotoSet.h
--- a/Engine/Plugins/Runtime/MeshModelingToolset/Source/ModelingComponents/Public/Scene/SceneCapturePhotoSet.h
+++ b/Engine/Plugins/Runtime/MeshModelingToolset/Source/ModelingComponents/Public/Scene/SceneCapturePhotoSet.h
@@ -50,6 +50,12 @@ public:
 	 */
 	void SetCaptureSceneActors(UWorld* World, const TArray<AActor*>& Actors);
 
+	/**
+	 * Disable all capture types. By default a standard set of capture types is enabled but, for performance reasons,
+	 * you can use this function to disable them all and then only enable the ones you need.
+	 */
+	void DisableAllCaptureTypes();
+
 	/**
 	 * Enable/Disable a particular capture type
 	 */
diff --git a/Engine/Source/Runtime/GeometryCore/Public/Image/SpatialPhotoSet.h b/Engine/Source/Runtime/GeometryCore/Public/Image/SpatialPhotoSet.h
--- a/Engine/Source/Runtime/GeometryCore/Public/Image/SpatialPhotoSet.h
+++ b/Engine/Source/Runtime/GeometryCore/Public/Image/SpatialPhotoSet.h
@@ -77,6 +77,9 @@ public:
 	/** @return the number of photos in the photo set */
 	int32 Num() const { return Photos.Num(); }
 
+	/** @return removes all photos from the photo set */
+	void Empty() { return Photos.Empty(); }
+
 	/** @return the photo at the given index */
 	const TSpatialPhoto<PixelType>& Get(int32 Index) const { return *Photos[Index]; }
 
