UV Project Tool: Make uniform checkbox behave like pattern tool's Proportional checkbox in scale settings.

Instead of just using the x-component when bUniformDimensions is true, the actual dimension vector is changed uniformly and no check is needed when getting the projection box.

diff --git a/Engine/Plugins/Runtime/MeshModelingToolset/Source/MeshModelingTools/Private/UVProjectionTool.cpp b/Engine/Plugins/Runtime/MeshModelingToolset/Source/MeshModelingTools/Private/UVProjectionTool.cpp
--- a/Engine/Plugins/Runtime/MeshModelingToolset/Source/MeshModelingTools/Private/UVProjectionTool.cpp
+++ b/Engine/Plugins/Runtime/MeshModelingToolset/Source/MeshModelingTools/Private/UVProjectionTool.cpp
@@ -114,18 +114,54 @@ void UUVProjectionTool::Setup()
 	TransformGizmo->SetActiveTarget(TransformProxy, GetToolManager());
 
 	InitialDimensions = BasicProperties->Dimensions;
-	bInitialUniformDimensions = BasicProperties->bUniformDimensions;
+	bInitialProportionalDimensions = BasicProperties->bProportionalDimensions;
 	InitialTransform = TransformGizmo->GetGizmoTransform();
 	
 	ApplyInitializationMode();
 
+	// Some lambdas for proportional dimension changes
+	auto OnProportionalDimensionsChanged = [this](bool bNewValue)
+	{
+		if (bNewValue) { CachedDimensions = BasicProperties->Dimensions; }
+		bTransformModified = true;
+	};
+
+	auto ApplyProportionalDimensions = [this](FVector& NewVector, FVector& CachedVector)
+	{
+		// Determines which component of the vector is being changed by looking at which component is
+		// most different from the previous values
+		const FVector Difference = NewVector - CachedVector;
+		const int32 DifferenceMaxElementIndex = MaxAbsElementIndex(Difference);
+
+		// Readability aliases
+		const double& NewComponent = NewVector[DifferenceMaxElementIndex];
+		const double& ChangedComponent = CachedVector[DifferenceMaxElementIndex];
+
+		// If the component we are changing is 0 (or very close to it), then just set all components to new component,
+		// effectively setting the vector's direction to (1, 1, 1) and scaling according to the new value
+		const double ScaleFactor = NewComponent / ChangedComponent;
+		NewVector = FMath::Abs(ChangedComponent) < UE_KINDA_SMALL_NUMBER ? FVector(NewComponent) : CachedVector * ScaleFactor;
+		CachedVector = NewVector;
+	};
+
 	// start watching for dimensions changes
-	DimensionsWatcher = BasicProperties->WatchProperty(BasicProperties->Dimensions, [this](FVector) { bTransformModified = true; });
-	DimensionsModeWatcher = BasicProperties->WatchProperty(BasicProperties->bUniformDimensions, [this](bool) { bTransformModified = true; });
+	DimensionsWatcher = BasicProperties->WatchProperty(BasicProperties->Dimensions, [this, &ApplyProportionalDimensions](FVector)
+	{
+		if (BasicProperties->bProportionalDimensions)
+		{
+			ApplyProportionalDimensions(BasicProperties->Dimensions, CachedDimensions);
+			BasicProperties->SilentUpdateWatcherAtIndex(DimensionsWatcher);
+		}
+
+		bTransformModified = true;
+	});
+	DimensionsModeWatcher = BasicProperties->WatchProperty(BasicProperties->bProportionalDimensions, OnProportionalDimensionsChanged);
 	BasicProperties->WatchProperty(BasicProperties->Initialization, [this](EUVProjectionToolInitializationMode NewMode) { OnInitializationModeChanged(); });
 	bTransformModified = false;
 	BasicProperties->SilentUpdateWatched();
 
+	OnProportionalDimensionsChanged(true);	// Initialize CachedDimensions
+	
 	// click to set plane behavior
 	SetPlaneCtrlClickBehaviorTarget = MakeUnique<FSelectClickedAction>();
 	SetPlaneCtrlClickBehaviorTarget->SnapManager = USceneSnappingManager::Find(GetToolManager());
@@ -203,7 +239,7 @@ void UUVProjectionTool::OnShutdown(EToolShutdownType ShutdownType)
 {
 	UVChannelProperties->SaveProperties(this);
 	BasicProperties->SavedDimensions = BasicProperties->Dimensions;
-	BasicProperties->bSavedUniformDimensions = BasicProperties->bUniformDimensions;
+	BasicProperties->bSavedProportionalDimensions = BasicProperties->bProportionalDimensions;
 	BasicProperties->SavedTransform = TransformGizmo->GetGizmoTransform();
 	BasicProperties->SaveProperties(this);
 	MaterialSettings->SaveProperties(this);
@@ -309,9 +345,7 @@ void UUVProjectionTool::InitializeMesh()
 FOrientedBox3d UUVProjectionTool::GetProjectionBox() const
 {
 	const FFrame3d BoxFrame(TransformProxy->GetTransform());
-	const FVector3d BoxDimensions = 0.5 * (BasicProperties->bUniformDimensions
-		                                       ? FVector3d(BasicProperties->Dimensions.X, BasicProperties->Dimensions.X, BasicProperties->Dimensions.X)
-		                                       : static_cast<FVector3d>(BasicProperties->Dimensions));
+	const FVector3d BoxDimensions = 0.5 * BasicProperties->Dimensions;
 	return FOrientedBox3d(BoxFrame, BoxDimensions);
 }
 
@@ -355,6 +389,23 @@ void UUVProjectionTool::Render(IToolsContextRenderAPI* RenderAPI)
 	ProjectionShapeVisualizer.EndFrame();
 }
 
+void UUVProjectionTool::OnPropertyModified(UObject* PropertySet, FProperty* Property)
+{
+	if (PropertySet == BasicProperties && BasicProperties->bProportionalDimensions)
+	{
+		// We are silencing all watchers if Property corresponds to an FVector UProperty because this indicates that
+		// the "Reset to Default" button was pressed. If individual components are modified instead, then Property will
+		// not be castable to an FStructProperty
+		
+		const FStructProperty* StructProperty = CastField<FStructProperty>(Property);
+		if (StructProperty != nullptr && StructProperty->Struct->GetName() == FString("Vector"))
+		{
+			CachedDimensions = BasicProperties->Dimensions;
+			BasicProperties->SilentUpdateWatcherAtIndex(DimensionsWatcher);
+		}
+	}
+}
+
 void UUVProjectionTool::OnTick(float DeltaTime)
 {
 	if (bHavePendingAction)
diff --git a/Engine/Plugins/Runtime/MeshModelingToolset/Source/MeshModelingTools/Public/UVProjectionTool.h b/Engine/Plugins/Runtime/MeshModelingToolset/Source/MeshModelingTools/Public/UVProjectionTool.h
--- a/Engine/Plugins/Runtime/MeshModelingToolset/Source/MeshModelingTools/Public/UVProjectionTool.h
+++ b/Engine/Plugins/Runtime/MeshModelingToolset/Source/MeshModelingTools/Public/UVProjectionTool.h
@@ -113,9 +113,9 @@ public:
 	UPROPERTY(EditAnywhere, Category = "UV Projection")
 	FVector Dimensions = FVector(100.0f, 100.0f, 100.0f);
 
-	/** Only use the Dimensions X value to uniformly define all projection shape dimensions */
+	/** If true, changes to Dimensions result in all components be changed proportionally */
 	UPROPERTY(EditAnywhere, Category = "UV Projection")
-	bool bUniformDimensions = false;
+	bool bProportionalDimensions = false;
 
 	/** Determines how projection settings will be initialized; this only takes effect if the projection shape dimensions or position are unchanged */
 	UPROPERTY(EditAnywhere, Category = "UV Projection")
@@ -173,7 +173,7 @@ public:
 	FVector SavedDimensions = FVector::ZeroVector;
 
 	UPROPERTY()
-	bool bSavedUniformDimensions = false;
+	bool bSavedProportionalDimensions = false;
 
 	UPROPERTY()
 	FTransform SavedTransform;
@@ -215,6 +215,8 @@ public:
 	virtual void OnTick(float DeltaTime) override;
 	virtual void Render(IToolsContextRenderAPI* RenderAPI) override;
 
+	virtual void OnPropertyModified(UObject* PropertySet, FProperty* Property) override;
+	
 	virtual bool HasCancel() const override { return true; }
 	virtual bool HasAccept() const override { return true; }
 	virtual bool CanAccept() const override;
@@ -262,8 +264,9 @@ protected:
 	FTransform3d WorldTransform;
 	UE::Geometry::FAxisAlignedBox3d WorldBounds;
 
+	FVector CachedDimensions;
 	FVector InitialDimensions;
-	bool bInitialUniformDimensions;
+	bool bInitialProportionalDimensions;
 	FTransform InitialTransform;
 	int32 DimensionsWatcher = -1;
 	int32 DimensionsModeWatcher = -1;
