Construction Plane Mechanic: Add Ctrl+Middle Click behavior which moves the gizmo to the clicked position in the current plane.

A raycast is performed and its intersection with the plane determines the new gizmo position.

diff --git a/Engine/Plugins/Runtime/MeshModelingToolset/Source/ModelingComponents/Private/Mechanics/ConstructionPlaneMechanic.cpp b/Engine/Plugins/Runtime/MeshModelingToolset/Source/ModelingComponents/Private/Mechanics/ConstructionPlaneMechanic.cpp
--- a/Engine/Plugins/Runtime/MeshModelingToolset/Source/ModelingComponents/Private/Mechanics/ConstructionPlaneMechanic.cpp
+++ b/Engine/Plugins/Runtime/MeshModelingToolset/Source/ModelingComponents/Private/Mechanics/ConstructionPlaneMechanic.cpp
@@ -28,7 +28,7 @@ void UConstructionPlaneMechanic::Shutdown()
 void UConstructionPlaneMechanic::Initialize(UWorld* TargetWorld, const FFrame3d& InitialPlane)
 {
 	Plane = InitialPlane;
-
+	
 	// create proxy and gizmo
 	UInteractiveGizmoManager* GizmoManager = GetParentTool()->GetToolManager()->GetPairedGizmoManager();
 
@@ -55,6 +55,33 @@ void UConstructionPlaneMechanic::Initialize(UWorld* TargetWorld, const FFrame3d&
 	ClickToSetPlaneBehavior->Initialize(SetPlaneCtrlClickBehaviorTarget.Get());
 
 	GetParentTool()->AddInputBehavior(ClickToSetPlaneBehavior);
+
+	// click to set gizmo position behavior
+	MiddleClickToSetGizmoBehavior = NewObject<ULocalSingleClickInputBehavior>();
+	MiddleClickToSetGizmoBehavior->IsHitByClickFunc = [this](const FInputDeviceRay& InputRay) {
+		FInputRayHit OutResult;
+		FVector3d HitPointOut;
+		OutResult.bHit = Plane.RayPlaneIntersection(InputRay.WorldRay.Origin, InputRay.WorldRay.Direction, 2, HitPointOut);
+		return OutResult;
+	};
+	MiddleClickToSetGizmoBehavior->OnClickedFunc = [this](const FInputDeviceRay& ClickPos)
+	{
+		FVector3d HitPointOut;
+		if (Plane.RayPlaneIntersection(ClickPos.WorldRay.Origin, ClickPos.WorldRay.Direction, 2, HitPointOut))
+		{
+			FTransform NewTransform = PlaneTransformGizmo->GetGizmoTransform();
+			NewTransform.SetLocation(HitPointOut);
+
+			// Note: The plane of this mechanic is not being updated and OnPlaneChanged is not being broadcasted here, but the
+			// gizmo's position is. Anything listening for the gizmo's OnTransformChanged delegate however will be notified.
+			PlaneTransformGizmo->SetNewGizmoTransform(NewTransform);
+		}
+	};
+	MiddleClickToSetGizmoBehavior->ModifierCheckFunc = FInputDeviceState::IsCtrlKeyDown;
+	MiddleClickToSetGizmoBehavior->SetUseMiddleMouseButton();
+	MiddleClickToSetGizmoBehavior->Initialize();	// This takes no target because this behavior is its own target.
+	
+	GetParentTool()->AddInputBehavior(MiddleClickToSetGizmoBehavior);
 }
 
 
diff --git a/Engine/Plugins/Runtime/MeshModelingToolset/Source/ModelingComponents/Public/Mechanics/ConstructionPlaneMechanic.h b/Engine/Plugins/Runtime/MeshModelingToolset/Source/ModelingComponents/Public/Mechanics/ConstructionPlaneMechanic.h
--- a/Engine/Plugins/Runtime/MeshModelingToolset/Source/ModelingComponents/Public/Mechanics/ConstructionPlaneMechanic.h
+++ b/Engine/Plugins/Runtime/MeshModelingToolset/Source/ModelingComponents/Public/Mechanics/ConstructionPlaneMechanic.h
@@ -10,6 +10,7 @@
 
 #include "ConstructionPlaneMechanic.generated.h"
 
+class ULocalSingleClickInputBehavior;
 class UCombinedTransformGizmo;
 class UTransformProxy;
 class IClickBehaviorTarget;
@@ -25,7 +26,7 @@ class MODELINGCOMPONENTS_API UConstructionPlaneMechanic : public UInteractionMec
 {
 	GENERATED_BODY()
 public:
-
+	
 	/** Replace this to externally control if plane can be updated */
 	TUniqueFunction<bool()> CanUpdatePlaneFunc = []() { return true; };
 
@@ -74,6 +75,10 @@ public:
 protected:
 	UPROPERTY()
 	TObjectPtr<USingleClickInputBehavior> ClickToSetPlaneBehavior;
+	
+	/** This is the behavior and behavior target used for the Shift+click behavior that sets the gizmo's position in the world. */
+	UPROPERTY()
+	TObjectPtr<ULocalSingleClickInputBehavior> MiddleClickToSetGizmoBehavior;
 
 	void TransformChanged(UTransformProxy* Proxy, FTransform Transform);
 };
\ No newline at end of file
