Draw Polygon Tool: Add construction plane mechanic to tool.

The tool seems to predate the construction plane mechanic and ought to use it so that changes to the construction plane mechanic are reflected in this tool.

diff --git a/Engine/Plugins/Runtime/MeshModelingToolset/Source/MeshModelingTools/Private/DrawPolygonTool.cpp b/Engine/Plugins/Runtime/MeshModelingToolset/Source/MeshModelingTools/Private/DrawPolygonTool.cpp
--- a/Engine/Plugins/Runtime/MeshModelingToolset/Source/MeshModelingTools/Private/DrawPolygonTool.cpp
+++ b/Engine/Plugins/Runtime/MeshModelingToolset/Source/MeshModelingTools/Private/DrawPolygonTool.cpp
@@ -35,6 +35,7 @@
 #include "Selection/SelectClickedAction.h"
 #include "Selection/ToolSelectionUtil.h"
 #include "ModelingObjectsCreationAPI.h"
+#include "Mechanics/ConstructionPlaneMechanic.h"
 
 #include "Mechanics/DragAlignmentMechanic.h"
 
@@ -75,8 +76,6 @@ UDrawPolygonToolStandardProperties::UDrawPolygonToolStandardProperties()
  */
 UDrawPolygonTool::UDrawPolygonTool()
 {
-	DrawPlaneOrigin = FVector3d::Zero();
-	DrawPlaneOrientation = FQuaterniond::Identity();
 	bInInteractiveExtrude = false;
 	UInteractiveTool::SetToolDisplayName(LOCTEXT("ToolName", "Polygon Extrude"));
 }
@@ -96,25 +95,6 @@ void UDrawPolygonTool::Setup()
 	MouseBehavior->Modifiers.RegisterModifier(IgnoreSnappingModifier, FInputDeviceState::IsShiftKeyDown);
 	AddInputBehavior(MouseBehavior);
 
-	// Register a click behavior/action pair, that sets the draw plane to the clicked world position
-	FSelectClickedAction* SetPlaneAction = new FSelectClickedAction();
-	SetPlaneAction->SnapManager = USceneSnappingManager::Find(GetToolManager());
-	SetPlaneAction->OnClickedPositionFunc = [this](const FHitResult& Hit) {
-		SetDrawPlaneFromWorldPos((FVector3d)Hit.ImpactPoint, (FVector3d)Hit.ImpactNormal);
-	};
-	SetPointInWorldConnector = SetPlaneAction;
-
-	USingleClickInputBehavior* ClickToSetPlaneBehavior = NewObject<USingleClickInputBehavior>();
-	ClickToSetPlaneBehavior->ModifierCheckFunc = FInputDeviceState::IsCtrlKeyDown;
-	ClickToSetPlaneBehavior->Initialize(SetPointInWorldConnector);
-	ClickToSetPlaneBehavior->SetDefaultPriority(MouseBehavior->GetPriority().MakeHigher());
-	AddInputBehavior(ClickToSetPlaneBehavior);
-
-	// register modifier key behaviors   (disabled because it is not implemented yet)
-	//UKeyAsModifierInputBehavior* AKeyBehavior = NewObject<UKeyAsModifierInputBehavior>();
-	//AKeyBehavior->Initialize(this, AngleSnapModifier, EKeys::A);
-	//AddInputBehavior(AKeyBehavior);
-
 	OutputTypeProperties = NewObject<UCreateMeshObjectTypeProperties>(this);
 	OutputTypeProperties->RestoreProperties(this);
 	OutputTypeProperties->InitializeDefault();
@@ -123,22 +103,18 @@ void UDrawPolygonTool::Setup()
 	PolygonProperties = NewObject<UDrawPolygonToolStandardProperties>(this);
 	PolygonProperties->RestoreProperties(this);
 	PolygonProperties->WatchProperty(PolygonProperties->bShowGridGizmo,
-	                                 [this](bool bNewValue) { this->UpdateShowGizmoState(bNewValue); });
-
-	// Create a new TransformGizmo and associated TransformProxy. The TransformProxy will not be the
-	// parent of any Components in this case, we just use it's transform and change delegate.
-	PlaneTransformProxy = NewObject<UTransformProxy>(this);
-	PlaneTransformProxy->SetTransform(FTransform((FQuat)DrawPlaneOrientation, (FVector)DrawPlaneOrigin));
-	PlaneTransformGizmo = UE::TransformGizmoUtil::CreateCustomTransformGizmo(GetToolManager(),
-		ETransformGizmoSubElements::StandardTranslateRotate, this);
-	PlaneTransformGizmo->SetActiveTarget(PlaneTransformProxy, GetToolManager());
-	// listen for changes to the proxy and update the plane when that happens
-	PlaneTransformProxy->OnTransformChanged.AddUObject(this, &UDrawPolygonTool::PlaneTransformChanged);
+	                                 [this](bool bNewValue) { PlaneMechanic->PlaneTransformGizmo->SetVisibility(bNewValue); });
 
+	PlaneMechanic = NewObject<UConstructionPlaneMechanic>(this);
+	PlaneMechanic->Setup(this);
+	PlaneMechanic->CanUpdatePlaneFunc = [this]() { return AllowDrawPlaneUpdates(); };
+	PlaneMechanic->OnPlaneChanged.AddLambda([this]() { SnapEngine.Plane = PlaneMechanic->Plane; });	// Keep SnapEngine plane up to date with PlaneMechanic's plane
+	PlaneMechanic->Initialize(TargetWorld, FFrame3d());
+	
 	DragAlignmentMechanic = NewObject<UDragAlignmentMechanic>(this);
 	DragAlignmentMechanic->Setup(this);
-	DragAlignmentMechanic->AddToGizmo(PlaneTransformGizmo);
-
+	DragAlignmentMechanic->AddToGizmo(PlaneMechanic->PlaneTransformGizmo);
+	
 	// initialize material properties for new objects
 	MaterialProperties = NewObject<UNewMeshMaterialProperties>(this);
 	MaterialProperties->RestoreProperties(this);
@@ -164,7 +140,7 @@ void UDrawPolygonTool::Setup()
 	SnapEngine.SnapMetricFunc = [this](const FVector3d& Position1, const FVector3d& Position2) {
 		return ToolSceneQueriesUtil::CalculateNormalizedViewVisualAngleD(this->CameraState, Position1, Position2);
 	};
-	SnapEngine.Plane = FFrame3d(DrawPlaneOrigin, DrawPlaneOrientation);
+	SnapEngine.Plane = PlaneMechanic->Plane;
 
 	SnapProperties = NewObject<UDrawPolygonToolSnapProperties>(this);
 	SnapProperties->RestoreProperties(this);
@@ -195,13 +171,11 @@ void UDrawPolygonTool::Shutdown(EToolShutdownType ShutdownType)
 	PreviewMesh->Disconnect();
 	PreviewMesh = nullptr;
 
-	if (SetPointInWorldConnector != nullptr)
-	{
-		delete SetPointInWorldConnector;
-	}
-
 	DragAlignmentMechanic->Shutdown();
 
+	PlaneMechanic->Shutdown();
+	PlaneMechanic = nullptr;
+	
 	GetToolManager()->GetPairedGizmoManager()->DestroyAllGizmosByOwner(this);
 
 	OutputTypeProperties->SaveProperties(this);
@@ -256,11 +230,6 @@ void UDrawPolygonTool::ApplyUndoPoints(const TArray<FVector3d>& ClickPointsIn, c
 
 void UDrawPolygonTool::OnTick(float DeltaTime)
 {
-	if (PlaneTransformGizmo)
-	{
-		// faster to do this as an override rather than destroying/recreating the gizmo via UpdateShowGizmoState
-		PlaneTransformGizmo->SetVisibility(AllowDrawPlaneUpdates());
-	}
 	if (SnapProperties)
 	{
 		bool bSnappingEnabledInViewport = GetToolManager()->GetContextQueriesAPI()
@@ -315,9 +284,6 @@ void UDrawPolygonTool::Render(IToolsContextRenderAPI* RenderAPI)
 	FColor ErrorColor = FColor::Magenta;
 	float HiddenLineThickness = 1.0f*PDIScale;
 	float LineThickness = 4.0f*PDIScale;
-	FColor GridColor(128, 128, 128, 32);
-	float GridThickness = 0.5f*PDIScale;
-	int NumGridLines = 21;
 	FColor SnapLineColor = FColor::Yellow;
 	FColor SnapHighlightColor = SnapLineColor;
 	float ElementSize = CurViewSizeFactor;
@@ -326,13 +292,10 @@ void UDrawPolygonTool::Render(IToolsContextRenderAPI* RenderAPI)
 		|| (SnapEngine.HaveActiveSnap() && SnapEngine.GetActiveSnapTargetID() == StartPointSnapID);
 
 	//
-	// Draw the grid
+	// Render the plane mechanic after correctly setting bShowGrid
 	//
-	if (bInInteractiveExtrude == false)
-	{
-		FFrame3d DrawFrame(DrawPlaneOrigin, DrawPlaneOrientation);
-		MeshDebugDraw::DrawSimpleFixedScreenAreaGrid(RenderCameraState, DrawFrame, NumGridLines, 45.0, GridThickness, GridColor, false, PDI, FTransform::Identity);
-	}
+	PlaneMechanic->bShowGrid = !bInInteractiveExtrude;
+	PlaneMechanic->Render(RenderAPI);
 
 	//
 	// Generate the fixed polygon contour
@@ -392,7 +355,7 @@ void UDrawPolygonTool::Render(IToolsContextRenderAPI* RenderAPI)
 			{
 				int iSegment = SnapEngine.GetActiveSnapDistanceID();
 				TArray<FVector3d>& HistoryPoints = (bInFixedPolygonMode) ? FixedPolygonClickPoints : PolygonVertices;
-				FVector3d UseNormal = DrawPlaneOrientation.AxisZ();
+				FVector3d UseNormal = PlaneMechanic->Plane.Rotation.AxisZ();
 				DrawEdgeTicks(PDI, FSegment3d(HistoryPoints[iSegment], HistoryPoints[iSegment+1]),
 					0.75f*ElementSize, UseNormal, SnapHighlightColor, SDPG_Foreground, 1.0f*PDIScale, true);
 				DrawEdgeTicks(PDI, FSegment3d(HistoryPoints[HistoryPoints.Num()-1], PreviewVertex),
@@ -556,7 +519,7 @@ bool UDrawPolygonTool::FindDrawPlaneHitPoint(const FInputDeviceRay& ClickPos, FV
 {
 	bHaveSurfaceHit = false;
 
-	FFrame3d Frame(DrawPlaneOrigin, DrawPlaneOrientation);
+	const FFrame3d& Frame = PlaneMechanic->Plane;
 	FVector3d HitPos;
 	bool bHit = Frame.RayPlaneIntersection((FVector3d)ClickPos.WorldRay.Origin, (FVector3d)ClickPos.WorldRay.Direction, 2, HitPos);
 	if (bHit == false)
@@ -874,7 +837,7 @@ bool UDrawPolygonTool::UpdateSelfIntersection()
 		return false;
 	}
 
-	FFrame3d DrawFrame(DrawPlaneOrigin, DrawPlaneOrientation);
+	const FFrame3d& DrawFrame = PlaneMechanic->Plane;
 	FSegment2d PreviewSegment(DrawFrame.ToPlaneUV(PolygonVertices[NumVertices - 1],2), DrawFrame.ToPlaneUV(PreviewVertex,2));
 
 	double BestIntersectionParameter = FMathd::MaxReal;
@@ -903,7 +866,7 @@ void UDrawPolygonTool::GetPolygonParametersFromFixedPoints(const TArray<FVector3
 		return;
 	}
 
-	FFrame3d DrawFrame(DrawPlaneOrigin, DrawPlaneOrientation);
+	const FFrame3d& DrawFrame = PlaneMechanic->Plane;
 	FirstReferencePt = DrawFrame.ToPlaneUV(FixedPoints[0], 2);
 
 	FVector2d EdgePt = DrawFrame.ToPlaneUV(FixedPoints[1], 2);
@@ -959,7 +922,7 @@ void UDrawPolygonTool::GenerateFixedPolygon(const TArray<FVector3d>& FixedPoints
 		Hole.Transform([RotationMat](const FVector2d& Pt) { return RotationMat * Pt; });
 	}
 
-	FFrame3d DrawFrame(DrawPlaneOrigin, DrawPlaneOrientation);
+	const FFrame3d& DrawFrame = PlaneMechanic->Plane;
 	VerticesOut.SetNum(Polygon.VertexCount());
 	for (int k = 0; k < Polygon.VertexCount(); ++k)
 	{
@@ -1050,59 +1013,6 @@ bool UDrawPolygonTool::AllowDrawPlaneUpdates()
 }
 
 
-
-void UDrawPolygonTool::SetDrawPlaneFromWorldPos(const FVector3d& Position, const FVector3d& Normal)
-{
-	if (!AllowDrawPlaneUpdates())
-	{
-		return;
-	}
-
-	DrawPlaneOrigin = Position;
-
-	FFrame3d DrawPlane(Position, DrawPlaneOrientation);
-	if (bIgnoreSnappingToggle == false)
-	{
-		DrawPlane.AlignAxis(2, Normal);
-		DrawPlane.ConstrainedAlignPerpAxes();
-		DrawPlaneOrientation = DrawPlane.Rotation;
-	}
-
-	SnapEngine.Plane = FFrame3d(DrawPlane.Origin, DrawPlane.Rotation);
-
-	if (PlaneTransformGizmo != nullptr)
-	{
-		PlaneTransformGizmo->SetNewGizmoTransform(FTransform((FQuat)DrawPlaneOrientation, (FVector)DrawPlaneOrigin));
-	}
-}
-
-
-void UDrawPolygonTool::PlaneTransformChanged(UTransformProxy* Proxy, FTransform Transform)
-{
-	DrawPlaneOrientation = (FQuaterniond)Transform.GetRotation();
-	DrawPlaneOrigin = (FVector3d)Transform.GetLocation();
-	SnapEngine.Plane = FFrame3d(DrawPlaneOrigin, DrawPlaneOrientation);
-}
-
-void UDrawPolygonTool::UpdateShowGizmoState(bool bNewVisibility)
-{
-	if (bNewVisibility == false)
-	{
-		GetToolManager()->GetPairedGizmoManager()->DestroyAllGizmosByOwner(this);
-		PlaneTransformGizmo = nullptr;
-	}
-	else
-	{
-		if (!PlaneTransformGizmo) {
-			PlaneTransformGizmo = UE::TransformGizmoUtil::CreateCustomTransformGizmo(GetToolManager(),
-				ETransformGizmoSubElements::StandardTranslateRotate, this);
-		}
-		PlaneTransformGizmo->SetActiveTarget(PlaneTransformProxy, GetToolManager());
-		PlaneTransformGizmo->ReinitializeGizmoTransform(FTransform((FQuat)DrawPlaneOrientation, (FVector)DrawPlaneOrigin));
-	}
-}
-
-
 void UDrawPolygonTool::EmitCurrentPolygon()
 {
 	FString BaseName = (PolygonProperties->ExtrudeMode == EDrawPolygonExtrudeMode::Flat) ?
@@ -1171,7 +1081,7 @@ void UDrawPolygonTool::UpdateLivePreview()
 bool UDrawPolygonTool::GeneratePolygonMesh(const TArray<FVector3d>& Polygon, const TArray<TArray<FVector3d>>& PolygonHoles, FDynamicMesh3* ResultMeshOut, FFrame3d& WorldFrameOut, bool bIncludePreviewVtx, double ExtrudeDistance, bool bExtrudeSymmetric)
 {
 	// construct centered frame for polygon
-	WorldFrameOut = FFrame3d(DrawPlaneOrigin, DrawPlaneOrientation);
+	WorldFrameOut = PlaneMechanic->Plane;
 
 	int NumVerts = Polygon.Num();
 	FVector3d Centroid3d(0, 0, 0);
diff --git a/Engine/Plugins/Runtime/MeshModelingToolset/Source/MeshModelingTools/Public/DrawPolygonTool.h b/Engine/Plugins/Runtime/MeshModelingToolset/Source/MeshModelingTools/Public/DrawPolygonTool.h
--- a/Engine/Plugins/Runtime/MeshModelingToolset/Source/MeshModelingTools/Public/DrawPolygonTool.h
+++ b/Engine/Plugins/Runtime/MeshModelingToolset/Source/MeshModelingTools/Public/DrawPolygonTool.h
@@ -17,6 +17,7 @@
 #include "DrawPolygonTool.generated.h"
 
 
+class UConstructionPlaneMechanic;
 class UCombinedTransformGizmo;
 class UDragAlignmentMechanic;
 class UTransformProxy;
@@ -235,13 +236,6 @@ protected:
 	UPROPERTY()
 	TObjectPtr<UNewMeshMaterialProperties> MaterialProperties;
 	
-
-	/** Origin of plane we will draw polygon on */
-	FVector3d DrawPlaneOrigin;
-
-	/** Orientation of plane we will draw polygon on */
-	UE::Geometry::FQuaterniond DrawPlaneOrientation;
-	
 	/** Vertices of current preview polygon */
 	TArray<FVector3d> PolygonVertices;
 
@@ -257,26 +251,9 @@ protected:
 
 	UPROPERTY()
 	TObjectPtr<UPreviewMesh> PreviewMesh;
-
 	
-	// drawing plane gizmo
-
-	UPROPERTY()
-	TObjectPtr<UCombinedTransformGizmo> PlaneTransformGizmo;
-
-	UPROPERTY()
-	TObjectPtr<UTransformProxy> PlaneTransformProxy;
-
-	// called on PlaneTransformProxy.OnTransformChanged
-	void PlaneTransformChanged(UTransformProxy* Proxy, FTransform Transform);
-
-	// calls SetDrawPlaneFromWorldPos when user ctrl+clicks on scene
-	IClickBehaviorTarget* SetPointInWorldConnector = nullptr;
-
-	// updates plane and gizmo position
-	virtual void SetDrawPlaneFromWorldPos(const FVector3d& Position, const FVector3d& Normal);
-
-	void UpdateShowGizmoState(bool bNewVisibility);
+	// called on PlaneMechanic.OnPlaneChanged
+	void PlaneTransformChanged();
 
 	// whether to allow the draw plane to be updated in the UI -- returns false if there is an in-progress shape relying on the current draw plane
 	bool AllowDrawPlaneUpdates();
@@ -323,6 +300,9 @@ protected:
 	UPROPERTY()
 	TObjectPtr<UDragAlignmentMechanic> DragAlignmentMechanic = nullptr;
 
+	UPROPERTY()
+	TObjectPtr<UConstructionPlaneMechanic> PlaneMechanic = nullptr;
+
 	/** Generate extruded meshes.  Returns true on success. */
 	bool GeneratePolygonMesh(const TArray<FVector3d>& Polygon, const TArray<TArray<FVector3d>>& PolygonHoles, FDynamicMesh3* ResultMeshOut, UE::Geometry::FFrame3d& WorldFrameOut, bool bIncludePreviewVtx, double ExtrudeDistance, bool bExtrudeSymmetric);
 
