MeshModelingTools: Fix mesh boolean marking created group boundary edges as mesh boundary edges. Make EmbedPolygonsOp correctly set embedded edges and edges on failure in bool path.

diff --git a/Engine/Plugins/Runtime/GeometryProcessing/Source/DynamicMesh/Private/Operations/MeshBoolean.cpp b/Engine/Plugins/Runtime/GeometryProcessing/Source/DynamicMesh/Private/Operations/MeshBoolean.cpp
--- a/Engine/Plugins/Runtime/GeometryProcessing/Source/DynamicMesh/Private/Operations/MeshBoolean.cpp
+++ b/Engine/Plugins/Runtime/GeometryProcessing/Source/DynamicMesh/Private/Operations/MeshBoolean.cpp
@@ -737,7 +737,8 @@ bool FMeshBoolean::Compute()
 			}
 		}
 	}
-	else
+	// For NewGroupInside and NewGroupOutside, the cut doesn't create boundary edges.
+	else if (Operation != EBooleanOp::NewGroupInside && Operation != EBooleanOp::NewGroupOutside)
 	{
 		CreatedBoundaryEdges = CutBoundaryEdges[0];
 	}
diff --git a/Engine/Plugins/Runtime/MeshModelingToolset/Source/ModelingOperatorsEditorOnly/Private/CuttingOps/EmbedPolygonsOp.cpp b/Engine/Plugins/Runtime/MeshModelingToolset/Source/ModelingOperatorsEditorOnly/Private/CuttingOps/EmbedPolygonsOp.cpp
--- a/Engine/Plugins/Runtime/MeshModelingToolset/Source/ModelingOperatorsEditorOnly/Private/CuttingOps/EmbedPolygonsOp.cpp
+++ b/Engine/Plugins/Runtime/MeshModelingToolset/Source/ModelingOperatorsEditorOnly/Private/CuttingOps/EmbedPolygonsOp.cpp
@@ -183,7 +183,7 @@ void FEmbedPolygonsOp::BooleanPath(FProgressCancel* Progress)
 	Boolean.Compute();
 
 	// One might think that we would want to set this to the return value of Boolean.Compute() above,
-	// but that call currently returns false whenever any boundary edges are generated. That type of
+	// but that call currently returns false whenever any boundary edges are generated in a non-trim op. That type of
 	// "failure" is frequently ok for the user, for instance when cutting an open mesh, like a simple rectangle. 
 	// We should probably change FMeshBoolean to differentiate between types of errors so that we can still
 	// respond appropriately to other types of failures. For now, this is what we do here (and to some
@@ -195,7 +195,12 @@ void FEmbedPolygonsOp::BooleanPath(FProgressCancel* Progress)
 		return;
 	}
 
-	if (Boolean.CreatedBoundaryEdges.Num() > 0 && bAttemptFixHolesOnBoolean)
+	// Deal with created boundary edges
+	if (Operation == EEmbeddedPolygonOpMethod::TrimInside || Operation == EEmbeddedPolygonOpMethod::TrimOutside)
+	{
+		EmbeddedEdges = Boolean.CreatedBoundaryEdges;
+	}
+	else if (Boolean.CreatedBoundaryEdges.Num() > 0 && bAttemptFixHolesOnBoolean)
 	{
 		FMeshBoundaryLoops OpenBoundary(Boolean.Result, false);
 		TSet<int> ConsiderEdges(Boolean.CreatedBoundaryEdges);
@@ -250,16 +255,20 @@ void FEmbedPolygonsOp::BooleanPath(FProgressCancel* Progress)
 		EdgesOnFailure = Boolean.CreatedBoundaryEdges;
 	}
 
-	for (int EID : ResultMesh->EdgeIndicesItr())
+	// Non-mesh-boundary cut boundary edges can be detected by the group change across the cut. Not relevant to trim
+	// operations, where we get them from the created boundary edges above.
+	if (Operation == EEmbeddedPolygonOpMethod::CutThrough || Operation == EEmbeddedPolygonOpMethod::InsertPolygon)
 	{
-		FDynamicMesh3::FEdge Edge = ResultMesh->GetEdge(EID);
-		if (ResultMesh->GetTriangleGroup(Edge.Tri[0]) < MaxGroupID !=
-			ResultMesh->GetTriangleGroup(Edge.Tri[1]) < MaxGroupID)
+		for (int EID : ResultMesh->EdgeIndicesItr())
 		{
-			EmbeddedEdges.Add(EID);
+			FDynamicMesh3::FEdge Edge = ResultMesh->GetEdge(EID);
+			if (ResultMesh->GetTriangleGroup(Edge.Tri[0]) < MaxGroupID !=
+				ResultMesh->GetTriangleGroup(Edge.Tri[1]) < MaxGroupID)
+			{
+				EmbeddedEdges.Add(EID);
+			}
 		}
 	}
-
 }
 
 void FEmbedPolygonsOp::CalculateResult(FProgressCancel* Progress)
