ModelingMode: new selection system improvements.
Introduced FGeometrySelectionHitQueryConfig (GeometrySelection.h), this provides information about the active selection mode. FGeometrySelectionEditor now tracks active FGeometrySelectionHitQueryConfig.
IGeometrySelector hit-testing/processing functions now take a struct FWorldRayQueryInfo that combines a FRay3d with FViewCameraState, to allow for (eg) snapping wrt the active camera/etc
Added FGeometrySelectionPreview, this is a utility subclass of FGeometrySelectionEditor that is used for selection-preview highlighting, relative to a base FGeometrySelectionEditor
Added IGeometrySelector::GetSelectionPreviewForRaycast() which is used to accumulate a preview of what might be selected in a FGeometrySelectionPreview. AccumulateSelectionElements() can also now accumulate 'preview' geometry, and FDynamicMeshSelector uses this to provide the edge-loop of a facegroup for it's preview (as providing all triangles is too expensive for hover highlighting on hires meshes).
Added support for hover-highlighting mechanic in UGeometrySelectionManager and UModelingSelectionInteraction
FDynamicMeshSelector now creates a GroupEdgeSegmentTree for polygroup edges. This allows for better selection interaction, eg no dependency on the edge-adjacent triangles.
Added AreSelectionsIdentical() function in GeometrySelectionUtil.h

diff --git a/Engine/Plugins/Editor/ModelingToolsEditorMode/Source/ModelingToolsEditorMode/Private/ModelingSelectionInteraction.cpp b/Engine/Plugins/Editor/ModelingToolsEditorMode/Source/ModelingToolsEditorMode/Private/ModelingSelectionInteraction.cpp
--- a/Engine/Plugins/Editor/ModelingToolsEditorMode/Source/ModelingToolsEditorMode/Private/ModelingSelectionInteraction.cpp
+++ b/Engine/Plugins/Editor/ModelingToolsEditorMode/Source/ModelingToolsEditorMode/Private/ModelingSelectionInteraction.cpp
@@ -8,6 +8,10 @@
 #include "ToolSceneQueriesUtil.h"
 #include "SceneQueries/SceneSnappingManager.h"
 
+#include "BaseBehaviors/SingleClickBehavior.h"
+#include "BaseBehaviors/MouseHoverBehavior.h"
+
+
 using namespace UE::Geometry;
 
 void UModelingSelectionInteraction::Initialize(
@@ -26,8 +30,12 @@ void UModelingSelectionInteraction::Initialize(
 	ClickBehavior->Modifiers.RegisterModifier(ToggleSelectionModifier, FInputDeviceState::IsCtrlKeyDown);
 	ClickBehavior->Initialize(this);
 
+	HoverBehavior = NewObject<UMouseHoverBehavior>(this);
+	HoverBehavior->Initialize(this);
+
 	BehaviorSet = NewObject<UInputBehaviorSet>();
 	BehaviorSet->Add(ClickBehavior, this);
+	BehaviorSet->Add(HoverBehavior, this);
 
 	TransformProxy = NewObject<UTransformProxy>(this);
 	// todo: make this repositionable etc. Maybe make this function a delegate? or allow caller to provide the gizmo?
@@ -75,6 +83,7 @@ void UModelingSelectionInteraction::OnUpdateModifierState(int ModifierID, bool b
 	}
 }
 
+
 void UModelingSelectionInteraction::ComputeSceneHits(const FInputDeviceRay& ClickPos,
 	bool& bHitActiveObjects, FInputRayHit& ActiveObjectHit,
 	bool& bHitInactiveObjectFirst, FInputRayHit& InactiveObjectHit)
@@ -112,6 +121,7 @@ void UModelingSelectionInteraction::ComputeSceneHits(const FInputDeviceRay& Clic
 	}
 }
 
+
 FInputRayHit UModelingSelectionInteraction::IsHitByClick(const FInputDeviceRay& ClickPos)
 {
 	// ignore hits in these cases
@@ -174,6 +184,8 @@ FInputRayHit UModelingSelectionInteraction::IsHitByClick(const FInputDeviceRay&
 	return FInputRayHit();
 }
 
+
+
 void UModelingSelectionInteraction::OnClicked(const FInputDeviceRay& ClickPos)
 {
 	// this flag is set in IsHitByClick test
@@ -203,6 +215,47 @@ void UModelingSelectionInteraction::OnClicked(const FInputDeviceRay& ClickPos)
 }
 
 
+
+
+
+FInputRayHit UModelingSelectionInteraction::BeginHoverSequenceHitTest(const FInputDeviceRay& PressPos)
+{
+	if ((SelectionManager->GetMeshTopologyMode() == UGeometrySelectionManager::EMeshTopologyMode::None)
+		|| (CanChangeSelectionCallback() == false)
+		|| ExternalHitCaptureCallback(PressPos))
+	{
+		return FInputRayHit();
+	}
+
+	bool bHitActiveObjects, bHitInactiveObjectFirst;
+	FInputRayHit ActiveObjectHit, InactiveObjectHit;
+	ComputeSceneHits(PressPos, bHitActiveObjects, ActiveObjectHit, bHitInactiveObjectFirst, InactiveObjectHit);
+
+	if (bHitActiveObjects && bHitInactiveObjectFirst == false)
+	{
+		return ActiveObjectHit;
+	}
+	return FInputRayHit();
+}
+
+void UModelingSelectionInteraction::OnBeginHover(const FInputDeviceRay& DevicePos)
+{
+	SelectionManager->UpdateSelectionPreviewViaRaycast(DevicePos.WorldRay);
+}
+
+bool UModelingSelectionInteraction::OnUpdateHover(const FInputDeviceRay& DevicePos)
+{
+	return SelectionManager->UpdateSelectionPreviewViaRaycast(DevicePos.WorldRay);
+}
+
+void UModelingSelectionInteraction::OnEndHover()
+{
+	SelectionManager->ClearSelectionPreview();
+}
+
+
+
+
 void UModelingSelectionInteraction::UpdateGizmoOnSelectionChange()
 {
 	if (SelectionManager->HasSelection() == false)
@@ -216,10 +269,6 @@ void UModelingSelectionInteraction::UpdateGizmoOnSelectionChange()
 		FFrame3d SelectionFrame;
 		SelectionManager->GetSelectionWorldFrame(SelectionFrame);
 		TransformGizmo->ReinitializeGizmoTransform( SelectionFrame.ToFTransform() );
-
-		//FGeometrySelectionBounds Bounds;
-		//SelectionManager->GetSelectionBounds(Bounds);
-		//TransformGizmo->ReinitializeGizmoTransform( FTransform(Bounds.WorldBounds.Center()) );
 	}
 }
 
diff --git a/Engine/Plugins/Editor/ModelingToolsEditorMode/Source/ModelingToolsEditorMode/Public/ModelingSelectionInteraction.h b/Engine/Plugins/Editor/ModelingToolsEditorMode/Source/ModelingToolsEditorMode/Public/ModelingSelectionInteraction.h
--- a/Engine/Plugins/Editor/ModelingToolsEditorMode/Source/ModelingToolsEditorMode/Public/ModelingSelectionInteraction.h
+++ b/Engine/Plugins/Editor/ModelingToolsEditorMode/Source/ModelingToolsEditorMode/Public/ModelingSelectionInteraction.h
@@ -4,13 +4,14 @@
 #pragma once
 
 #include "BaseBehaviors/BehaviorTargetInterfaces.h"
-#include "BaseBehaviors/SingleClickBehavior.h"
 #include "InputBehaviorSet.h"
 #include "BaseGizmos/TransformGizmoUtil.h"
 #include "FrameTypes.h"
 #include "ModelingSelectionInteraction.generated.h"
 
 class UGeometrySelectionManager;
+class USingleClickInputBehavior;
+class UMouseHoverBehavior;
 
 /**
  * UModelingSelectionInteraction provides element-level selection behavior (ie mesh triangles/edges/vertices)
@@ -31,7 +32,7 @@ class UGeometrySelectionManager;
  * 
  */
 UCLASS()
-class UModelingSelectionInteraction : public UObject, public IInputBehaviorSource, public IClickBehaviorTarget
+class UModelingSelectionInteraction : public UObject, public IInputBehaviorSource, public IClickBehaviorTarget, public IHoverBehaviorTarget
 {
 	GENERATED_BODY()
 
@@ -62,6 +63,10 @@ public:
 	UPROPERTY()
 	TObjectPtr<USingleClickInputBehavior> ClickBehavior;
 
+	// mouse hover behavior
+	UPROPERTY()
+	TObjectPtr<UMouseHoverBehavior> HoverBehavior;
+
 	// set of all behaviors, will be passed up to UInputRouter
 	UPROPERTY()
 	TObjectPtr<UInputBehaviorSet> BehaviorSet;
@@ -82,6 +87,14 @@ public:
 	virtual void OnClicked(const FInputDeviceRay& ClickPos) override;
 
 
+	//
+	// IHoverBehaviorTarget implementation
+	//
+	virtual FInputRayHit BeginHoverSequenceHitTest(const FInputDeviceRay& PressPos) override;
+	virtual void OnBeginHover(const FInputDeviceRay& DevicePos) override;
+	virtual bool OnUpdateHover(const FInputDeviceRay& DevicePos) override;
+	virtual void OnEndHover() override;
+
 	//
 	// IModifierToggleBehaviorTarget implementation
 	//
diff --git a/Engine/Plugins/Runtime/GeometryProcessing/Source/DynamicMesh/Private/Selections/GeometrySelection.cpp b/Engine/Plugins/Runtime/GeometryProcessing/Source/DynamicMesh/Private/Selections/GeometrySelection.cpp
--- a/Engine/Plugins/Runtime/GeometryProcessing/Source/DynamicMesh/Private/Selections/GeometrySelection.cpp
+++ b/Engine/Plugins/Runtime/GeometryProcessing/Source/DynamicMesh/Private/Selections/GeometrySelection.cpp
@@ -5,10 +5,21 @@
 
 using namespace UE::Geometry;
 
-void FGeometrySelectionEditor::Initialize(FGeometrySelection* TargetSelectionIn)
+void FGeometrySelectionEditor::Initialize(FGeometrySelection* TargetSelectionIn, const FGeometrySelectionHitQueryConfig& QueryConfigIn)
 {
 	check(TargetSelectionIn != nullptr);
+	check(QueryConfigIn.ElementType == TargetSelectionIn->ElementType);
+	check(QueryConfigIn.TopologyType == TargetSelectionIn->TopologyType);
 	TargetSelection = TargetSelectionIn;
+	QueryConfig = QueryConfigIn;
+}
+
+
+void FGeometrySelectionEditor::UpdateQueryConfig(const FGeometrySelectionHitQueryConfig& NewConfig)
+{
+	check(NewConfig.ElementType == TargetSelection->ElementType);
+	check(NewConfig.TopologyType == TargetSelection->TopologyType);
+	QueryConfig = NewConfig;
 }
 
 void FGeometrySelectionEditor::ClearSelection(FGeometrySelectionDelta& DeltaOut)
diff --git a/Engine/Plugins/Runtime/GeometryProcessing/Source/DynamicMesh/Private/Selections/GeometrySelectionUtil.cpp b/Engine/Plugins/Runtime/GeometryProcessing/Source/DynamicMesh/Private/Selections/GeometrySelectionUtil.cpp
--- a/Engine/Plugins/Runtime/GeometryProcessing/Source/DynamicMesh/Private/Selections/GeometrySelectionUtil.cpp
+++ b/Engine/Plugins/Runtime/GeometryProcessing/Source/DynamicMesh/Private/Selections/GeometrySelectionUtil.cpp
@@ -12,6 +12,52 @@
 #include "Algo/Find.h"
 
 
+
+bool UE::Geometry::AreSelectionsIdentical(
+	const FGeometrySelection& SelectionA, const FGeometrySelection& SelectionB)
+{
+	if ((SelectionA.ElementType != SelectionB.ElementType) || (SelectionA.TopologyType != SelectionB.TopologyType))
+	{
+		return false;
+	}
+	int32 Num = SelectionA.Num();
+	if (Num != SelectionB.Num())
+	{
+		return false;
+	}
+
+	if (SelectionA.TopologyType == EGeometryTopologyType::Polygroup)
+	{
+		// for polygroup topology we may have stored an arbitrary geometry ID and so we cannot rely on TSet Contains
+		for (uint64 ItemA : SelectionA.Selection)
+		{
+			uint32 TopologyID = FGeoSelectionID(ItemA).TopologyID;
+			const uint64* Found = Algo::FindByPredicate(SelectionB.Selection, [&](uint64 Item)
+			{
+				return FGeoSelectionID(Item).TopologyID == TopologyID;
+			});
+			if (Found == nullptr)
+			{
+				return false;
+			}
+		}
+	}
+	else
+	{
+		for (uint64 ItemA : SelectionA.Selection)
+		{
+			if (SelectionB.Selection.Contains(ItemA) == false)
+			{
+				return false;
+			}
+		}
+	}
+
+	return true;
+}
+
+
+
 bool UE::Geometry::FindInSelectionByTopologyID(
 	const FGeometrySelection& GeometrySelection,
 	uint32 TopologyID,
@@ -137,24 +183,28 @@ void UE::Geometry::UpdateGroupSelectionViaRaycast(
 			FVector3d HitPos = LocalRay.PointAt(RayHitT);
 			FIndex3i TriVerts = ColliderMesh->GetTriangle(HitTriangleID);
 			int32 NearestIdx = -1;
-			double NearestDistSqr = DistanceSquared(ColliderMesh->GetVertex(TriVerts[0]), HitPos);
+			int32 NearestCornerID = IndexConstants::InvalidID;
+			double NearestDistSqr = TNumericLimits<double>::Max();
 			for (int32 k = 0; k < 3; ++k)
 			{
-				if (GroupTopology->GetCornerIDFromVertexID(TriVerts[k]) != IndexConstants::InvalidID)
+				int32 FoundCornerID = GroupTopology->GetCornerIDFromVertexID(TriVerts[k]);
+				if (FoundCornerID != IndexConstants::InvalidID)
 				{
 					double DistSqr = DistanceSquared(ColliderMesh->GetVertex(TriVerts[k]), HitPos);
 					if (DistSqr < NearestDistSqr)
 					{
 						NearestDistSqr = DistSqr;
 						NearestIdx = k;
+						NearestCornerID = FoundCornerID;
 					}
 				}
 			}
-			if (NearestIdx >= 0)
+			if (NearestCornerID != IndexConstants::InvalidID)
 			{
 				// do we need a group here?
 				int32 VertexID = TriVerts[NearestIdx];
-				ResultOut.bSelectionModified = UpdateSelectionWithNewElements(Editor, UpdateConfig.ChangeType, TArray<uint64>{(uint64)VertexID}, &ResultOut.SelectionDelta);
+				FGeoSelectionID SelectionID(VertexID, NearestCornerID);
+				ResultOut.bSelectionModified = UpdateSelectionWithNewElements(Editor, UpdateConfig.ChangeType, TArray<uint64>{SelectionID.Encoded()}, &ResultOut.SelectionDelta);
 				ResultOut.bSelectionMissed = false;
 			}
 		}
@@ -182,8 +232,14 @@ void UE::Geometry::UpdateGroupSelectionViaRaycast(
 			{
 				// do we need a group here?
 				FMeshTriEdgeID TriEdgeID(HitTriangleID, NearestIdx);
-				ResultOut.bSelectionModified = UpdateSelectionWithNewElements(Editor, UpdateConfig.ChangeType, TArray<uint64>{(uint64)TriEdgeID.Encoded()}, & ResultOut.SelectionDelta);
-				ResultOut.bSelectionMissed = false;
+				int32 GroupEdgeID = GroupTopology->FindGroupEdgeID(TriEdgeID);
+				checkSlow(GroupEdgeID >= 0);  // should never fail...
+				if (GroupEdgeID >= 0)
+				{
+					FGeoSelectionID SelectionID(TriEdgeID.Encoded(), GroupEdgeID);
+					ResultOut.bSelectionModified = UpdateSelectionWithNewElements(Editor, UpdateConfig.ChangeType, TArray<uint64>{SelectionID.Encoded()}, & ResultOut.SelectionDelta);
+					ResultOut.bSelectionMissed = false;
+				}
 			}
 		}
 	}
@@ -542,7 +598,8 @@ bool UE::Geometry::EnumerateTriangleSelectionElements(
 	TFunctionRef<void(int32, const FVector3d&)> VertexFunc,
 	TFunctionRef<void(int32, const FSegment3d&)> EdgeFunc,
 	TFunctionRef<void(int32, const FTriangle3d&)> TriangleFunc,
-	const FTransform* ApplyTransform
+	const FTransform* ApplyTransform,
+	bool bMapFacesToEdgeLoops
 )
 {
 	if ( ensure( MeshSelection.TopologyType == EGeometryTopologyType::Triangle ) == false )
@@ -565,7 +622,17 @@ bool UE::Geometry::EnumerateTriangleSelectionElements(
 					B = ApplyTransform->TransformPosition(B);
 					C = ApplyTransform->TransformPosition(C);
 				}
-				TriangleFunc(TriangleID, FTriangle3d(A, B, C));
+				if (bMapFacesToEdgeLoops)
+				{
+					FIndex3i Edges = Mesh.GetTriEdges((int32)TriangleID);
+					EdgeFunc(Edges.A, FSegment3d(A, B));
+					EdgeFunc(Edges.B, FSegment3d(B, C));
+					EdgeFunc(Edges.C, FSegment3d(C, A));
+				}
+				else
+				{
+					TriangleFunc(TriangleID, FTriangle3d(A, B, C));
+				}
 			}
 		}
 	}
@@ -619,7 +686,8 @@ bool UE::Geometry::EnumeratePolygroupSelectionElements(
 	TFunctionRef<void(int32, const FVector3d&)> VertexFunc,
 	TFunctionRef<void(int32, const FSegment3d&)> EdgeFunc,
 	TFunctionRef<void(int32, const FTriangle3d&)> TriangleFunc,
-	const FTransform* ApplyTransform
+	const FTransform* ApplyTransform,
+	bool bMapFacesToEdgeLoops
 )
 {
 	if ( ensure( MeshSelection.TopologyType == EGeometryTopologyType::Polygroup ) == false )
@@ -627,25 +695,65 @@ bool UE::Geometry::EnumeratePolygroupSelectionElements(
 		return false;
 	}
 
+	auto ProcessGroupEdgeID = [GroupTopology, &Mesh, ApplyTransform, &EdgeFunc](int32 GroupEdgeID)
+	{
+		for (int32 EdgeID : GroupTopology->GetGroupEdgeEdges(GroupEdgeID))
+		{
+			FVector3d A, B;
+			Mesh.GetEdgeV(EdgeID, A, B);
+			if (ApplyTransform != nullptr)
+			{
+				A = ApplyTransform->TransformPosition(A);
+				B = ApplyTransform->TransformPosition(B);
+			}
+			EdgeFunc(EdgeID, FSegment3d(A, B));
+		}
+	};
+
 	if (MeshSelection.ElementType == EGeometryElementType::Face)
 	{
-		for (uint64 EncodedID : MeshSelection.Selection)
+		if (bMapFacesToEdgeLoops)
 		{
-			FGeoSelectionID SelectionID(EncodedID);
-			int32 SeedTriangleID = (int32)SelectionID.GeometryID, GroupID = (int32)SelectionID.TopologyID;
-			if (Mesh.IsTriangle(SeedTriangleID))
+			TArray<int32> GroupEdgeIDs;
+			for (uint64 EncodedID : MeshSelection.Selection)
 			{
-				for (int32 TriangleID : GroupTopology->GetGroupFaces(GroupID))
+				int32 GroupID = FGeoSelectionID(EncodedID).TopologyID;
+				if ( const FGroupTopology::FGroup* Group = GroupTopology->FindGroupByID(GroupID) )
 				{
-					FVector3d A, B, C;
-					Mesh.GetTriVertices((int32)TriangleID, A, B, C);
-					if (ApplyTransform != nullptr)
+					for (const FGroupTopology::FGroupBoundary& Boundary : Group->Boundaries)
 					{
-						A = ApplyTransform->TransformPosition(A);
-						B = ApplyTransform->TransformPosition(B);
-						C = ApplyTransform->TransformPosition(C);
+						for (int32 GroupEdgeID : Boundary.GroupEdges)
+						{
+							GroupEdgeIDs.AddUnique(GroupEdgeID);
+						}
+					}
+				}
+			}
+			for (int32 GroupEdgeID : GroupEdgeIDs)
+			{
+				ProcessGroupEdgeID(GroupEdgeID);
+			}
+		}
+		else
+		{
+			for (uint64 EncodedID : MeshSelection.Selection)
+			{
+				FGeoSelectionID SelectionID(EncodedID);
+				int32 SeedTriangleID = (int32)SelectionID.GeometryID, GroupID = (int32)SelectionID.TopologyID;
+				if (Mesh.IsTriangle(SeedTriangleID))
+				{
+					for (int32 TriangleID : GroupTopology->GetGroupFaces(GroupID))
+					{
+						FVector3d A, B, C;
+						Mesh.GetTriVertices((int32)TriangleID, A, B, C);
+						if (ApplyTransform != nullptr)
+						{
+							A = ApplyTransform->TransformPosition(A);
+							B = ApplyTransform->TransformPosition(B);
+							C = ApplyTransform->TransformPosition(C);
+						}
+						TriangleFunc(TriangleID, FTriangle3d(A, B, C));
 					}
-					TriangleFunc(TriangleID, FTriangle3d(A, B, C));
 				}
 			}
 		}
@@ -659,17 +767,7 @@ bool UE::Geometry::EnumeratePolygroupSelectionElements(
 			if (Mesh.IsEdge(SeedEdgeID))
 			{
 				int32 GroupEdgeID = GroupTopology->FindGroupEdgeID(SeedEdgeID);
-				for (int32 EdgeID : GroupTopology->GetGroupEdgeEdges(GroupEdgeID))
-				{
-					FVector3d A, B;
-					Mesh.GetEdgeV(EdgeID, A, B);
-					if (ApplyTransform != nullptr)
-					{
-						A = ApplyTransform->TransformPosition(A);
-						B = ApplyTransform->TransformPosition(B);
-					}
-					EdgeFunc(EdgeID, FSegment3d(A, B));
-				}
+				ProcessGroupEdgeID(GroupEdgeID);
 			}
 		}
 	}
diff --git a/Engine/Plugins/Runtime/GeometryProcessing/Source/DynamicMesh/Public/Selections/GeometrySelection.h b/Engine/Plugins/Runtime/GeometryProcessing/Source/DynamicMesh/Public/Selections/GeometrySelection.h
--- a/Engine/Plugins/Runtime/GeometryProcessing/Source/DynamicMesh/Public/Selections/GeometrySelection.h
+++ b/Engine/Plugins/Runtime/GeometryProcessing/Source/DynamicMesh/Public/Selections/GeometrySelection.h
@@ -191,6 +191,22 @@ struct DYNAMICMESH_API FGeometrySelection
 };
 
 
+/**
+ * FGeometrySelectionHitQueryConfig defines the desired settings for a "hit query" on
+ * selected meshes/objects. 
+ */
+struct DYNAMICMESH_API FGeometrySelectionHitQueryConfig
+{
+	/** Type of object topology to query */
+	EGeometryTopologyType TopologyType = EGeometryTopologyType::Polygroup;
+	/** Type of object topological element to query */
+	EGeometryElementType ElementType = EGeometryElementType::Face;
+	
+	/** If true, only "visible" elements are considered, otherwise query may return obscured elements (eg, hidden edges occluded by mesh surface) */
+	bool bOnlyVisible = true;
+};
+
+
 
 /**
  * Type of change, relative to a FGeometrySelection
@@ -292,19 +308,34 @@ public:
 	 * Initialize the Editor with the given Selection. The
 	 * TargetSelectionIn must live longer than the FGeometrySelectionEditor
 	 */
-	void Initialize(FGeometrySelection* TargetSelectionIn);
+	void Initialize(
+		FGeometrySelection* TargetSelectionIn,
+		const FGeometrySelectionHitQueryConfig& QueryConfigIn);
 
 	/** @return the Element Type of the Target Selection */
 	EGeometryElementType GetElementType() const { return TargetSelection->ElementType; }
 	/** @return the Topology Type of the Target Selection */
 	EGeometryTopologyType GetTopologyType() const { return TargetSelection->TopologyType; }
 
+	/** @return the active configuration for this Selection Editor, eg what type of element/topology is being selected. This is redundant w/ the above fields.  */
+	const FGeometrySelectionHitQueryConfig& GetQueryConfig() const { return QueryConfig; }
+
+	/** 
+	 * Update the active QueryConfig for this SelectionEditor. This is necessary to keep it in sync w/ the TargetSelection, if
+	 * the ElementType or TopologyType are modified. This perhaps should be revisited as they basically
+	 * always need to be the same...
+	 */
+	void UpdateQueryConfig(const FGeometrySelectionHitQueryConfig& NewConfig);
+
 	/** @return true if the given ID is currently selected in the Target Selection */
 	bool IsSelected(uint64 ID) const;
 
 	/** Clear the Target Selection and return change information in DeltaOut */
 	void ClearSelection(FGeometrySelectionDelta& DeltaOut);
 
+	/** Access the Selection object this Editor is modifying */
+	const FGeometrySelection& GetSelection() const { return *TargetSelection; }
+
 	/** Add the items in the List to the Target Selection and return change information in DeltaOut */
 	template<typename ListType>
 	bool Select(const ListType& List, FGeometrySelectionDelta& DeltaOut)
@@ -345,6 +376,29 @@ public:
 
 protected:
 	FGeometrySelection* TargetSelection = nullptr;
+	FGeometrySelectionHitQueryConfig QueryConfig;
+};
+
+
+/**
+ * FGeometrySelectionPreview is a combined FGeometrySelection and FGeometrySelectionEditor.
+ * The purpose of this class is to support things like hover-highlighting, where the Selection 
+ * system needs to be queried but the result is going to be discarded. FGeometrySelectionPreview can 
+ * only be constructed based on a parent FGeometrySelectionEditor, from which it derives it's
+ * configuration, which is then how the rest of the system knows what geometry to provide as a preview/etc
+ */
+class DYNAMICMESH_API FGeometrySelectionPreview : public FGeometrySelectionEditor
+{
+public:
+	FGeometrySelection PreviewSelection;
+
+	FGeometrySelectionPreview(const FGeometrySelectionEditor& ActiveEditor)
+	{
+		PreviewSelection.InitializeTypes(ActiveEditor.GetSelection());
+		Initialize(&PreviewSelection, ActiveEditor.GetQueryConfig());
+	}
+
+	bool IsEmpty() const { return PreviewSelection.IsEmpty(); }
 };
 
 
diff --git a/Engine/Plugins/Runtime/GeometryProcessing/Source/DynamicMesh/Public/Selections/GeometrySelectionUtil.h b/Engine/Plugins/Runtime/GeometryProcessing/Source/DynamicMesh/Public/Selections/GeometrySelectionUtil.h
--- a/Engine/Plugins/Runtime/GeometryProcessing/Source/DynamicMesh/Public/Selections/GeometrySelectionUtil.h
+++ b/Engine/Plugins/Runtime/GeometryProcessing/Source/DynamicMesh/Public/Selections/GeometrySelectionUtil.h
@@ -19,6 +19,18 @@ class FColliderMesh;
 class FGroupTopology;
 struct FGroupTopologySelection;
 
+
+/**
+ * Test if SelectionA and SelectionB are the same selection.
+ * This is currently relatively expensive on Polygroup selections due to how they are encoded
+ * @return true if the selections are iddentical
+ */
+DYNAMICMESH_API bool AreSelectionsIdentical(
+	const FGeometrySelection& SelectionA,
+	const FGeometrySelection& SelectionB);
+
+
+
 /**
  * Assuming that the uint64 values in the GeometrySelection are encoded FGeoSelectionID's, 
  * find the item that has a matching TopologyID, ignoring the GeometryID.
@@ -146,7 +158,8 @@ DYNAMICMESH_API bool EnumerateTriangleSelectionElements(
 	TFunctionRef<void(int32, const FVector3d&)> VertexFunc,
 	TFunctionRef<void(int32, const FSegment3d&)> EdgeFunc,
 	TFunctionRef<void(int32, const FTriangle3d&)> TriangleFunc,
-	const FTransform* ApplyTransform = nullptr
+	const FTransform* ApplyTransform = nullptr,
+	bool bMapFacesToEdgeLoops = false
 );
 /**
  * Call VertexFunc/EdgeFunc/TriangleFunc for the mesh vertices/edges/triangles identified by MeshSelection,
@@ -163,7 +176,8 @@ DYNAMICMESH_API bool EnumeratePolygroupSelectionElements(
 	TFunctionRef<void(int32, const FVector3d&)> VertexFunc,
 	TFunctionRef<void(int32, const FSegment3d&)> EdgeFunc,
 	TFunctionRef<void(int32, const FTriangle3d&)> TriangleFunc,
-	const FTransform* ApplyTransform = nullptr
+	const FTransform* ApplyTransform = nullptr,
+	bool bMapFacesToEdgeLoops = false
 );
 
 
diff --git a/Engine/Plugins/Runtime/MeshModelingToolset/Source/ModelingComponents/Private/Selection/DynamicMeshSelector.cpp b/Engine/Plugins/Runtime/MeshModelingToolset/Source/ModelingComponents/Private/Selection/DynamicMeshSelector.cpp
--- a/Engine/Plugins/Runtime/MeshModelingToolset/Source/ModelingComponents/Private/Selection/DynamicMeshSelector.cpp
+++ b/Engine/Plugins/Runtime/MeshModelingToolset/Source/ModelingComponents/Private/Selection/DynamicMeshSelector.cpp
@@ -10,6 +10,8 @@
 #include "Changes/MeshVertexChange.h"
 #include "DynamicMesh/ColliderMesh.h"
 #include "GroupTopology.h"
+#include "Spatial/SegmentTree3.h"
+#include "ToolSceneQueriesUtil.h"
 
 using namespace UE::Geometry;
 
@@ -45,6 +47,7 @@ void FDynamicMeshSelector::RegisterMeshChangedHandler()
 	TargetMesh_OnMeshChangedHandle = TargetMesh->OnMeshChanged().AddLambda([this](UDynamicMesh* Mesh, FDynamicMeshChangeInfo ChangeInfo)
 	{
 		ColliderMesh.Reset();
+		GroupEdgeSegmentTree.Reset();
 
 		if (ChangeInfo.Type != EDynamicMeshChangeType::MeshVertexChange)
 		{
@@ -65,6 +68,8 @@ void FDynamicMeshSelector::Shutdown()
 		TargetMesh_OnMeshChangedHandle.Reset();
 		
 		ColliderMesh.Reset();
+		GroupTopology.Reset();
+		GroupEdgeSegmentTree.Reset();
 	}
 	TargetMesh = nullptr;
 }
@@ -81,6 +86,7 @@ bool FDynamicMeshSelector::Sleep()
 
 	ColliderMesh.Reset();
 	GroupTopology.Reset();
+	GroupEdgeSegmentTree.Reset();
 
 	return true;
 }
@@ -98,13 +104,21 @@ bool FDynamicMeshSelector::Restore()
 }
 
 
+
+
+
+
 bool FDynamicMeshSelector::RayHitTest(
-	const FRay3d& WorldRay,
+	const FWorldRayQueryInfo& RayInfo,
+	FGeometrySelectionHitQueryConfig QueryConfig,
 	FInputRayHit& HitResultOut )
 {
-	FTransformSRT3d WorldTransform = GetWorldTransformFunc();
 	HitResultOut = FInputRayHit();
-	FRay3d LocalRay = WorldTransform.InverseTransformRay(WorldRay);
+
+	// todo: do we need to sometimes ignore surface hits here? maybe for Volumes in (eg) edge mode...
+
+	FTransformSRT3d WorldTransform = GetWorldTransformFunc();
+	FRay3d LocalRay = WorldTransform.InverseTransformRay(RayInfo.WorldRay);
 	double RayHitT; int32 HitTriangleID; FVector3d HitBaryCoords;
 	if (GetColliderMesh()->FindNearestHitTriangle(LocalRay, RayHitT, HitTriangleID, HitBaryCoords))
 	{
@@ -112,35 +126,147 @@ bool FDynamicMeshSelector::RayHitTest(
 		HitResultOut.HitIdentifier = HitTriangleID;
 		HitResultOut.HitOwner = TargetMesh.Get();
 		FVector3d WorldPosition = WorldTransform.TransformPosition(LocalRay.PointAt(RayHitT));
-		HitResultOut.HitDepth = WorldRay.GetParameter(WorldPosition);
+		HitResultOut.HitDepth = RayInfo.WorldRay.GetParameter(WorldPosition);
+	}
+	else if ( QueryConfig.TopologyType == EGeometryTopologyType::Polygroup && QueryConfig.ElementType == EGeometryElementType::Edge )
+	{
+		auto EdgeHitToleranceTest = [RayInfo, WorldTransform](int32 ID, const FVector3d& A, const FVector3d& B)
+		{
+			return ToolSceneQueriesUtil::PointSnapQuery(RayInfo.CameraState,
+				WorldTransform.TransformPosition(A), WorldTransform.TransformPosition(B));
+		};
+		FSegmentTree3::FSegment Segment;
+		FSegmentTree3::FRayNearestSegmentInfo SegmentInfo;
+		if (GetGroupEdgeSpatial()->FindNearestVisibleSegmentHitByRay(LocalRay, EdgeHitToleranceTest, Segment, SegmentInfo))
+		{
+			HitResultOut.bHit = true;
+			HitResultOut.HitIdentifier = Segment.ID;
+			HitResultOut.HitOwner = TargetMesh.Get();
+			FVector3d WorldPosition = WorldTransform.TransformPosition(LocalRay.PointAt(SegmentInfo.RayParam));
+			HitResultOut.HitDepth = RayInfo.WorldRay.GetParameter(WorldPosition);
+		}
 	}
 	return HitResultOut.bHit;
 }
 
 
 void FDynamicMeshSelector::UpdateSelectionViaRaycast(
-	const FRay3d& WorldRay,
+	const FWorldRayQueryInfo& RayInfo,
 	FGeometrySelectionEditor& SelectionEditor,
 	const FGeometrySelectionUpdateConfig& UpdateConfig,
 	FGeometrySelectionUpdateResult& ResultOut)
 {
-	FRay3d LocalRay = GetWorldTransformFunc().InverseTransformRay(WorldRay);
+	if (SelectionEditor.GetTopologyType() == EGeometryTopologyType::Triangle)
+	{
+		UpdateSelectionViaRaycast_MeshTopology(RayInfo, SelectionEditor, UpdateConfig, ResultOut);
+		return;
+	}
+	check(SelectionEditor.GetTopologyType() == EGeometryTopologyType::Polygroup);	// assume this...
 
-	if (SelectionEditor.GetTopologyType() == EGeometryTopologyType::Polygroup)
+	if (SelectionEditor.GetElementType() == EGeometryElementType::Edge)
+	{
+		UpdateSelectionViaRaycast_GroupEdges(RayInfo, SelectionEditor, UpdateConfig, ResultOut);
+	}
+	else
 	{
+		FRay3d LocalRay = GetWorldTransformFunc().InverseTransformRay(RayInfo.WorldRay);
 		UE::Geometry::UpdateGroupSelectionViaRaycast(
 			GetColliderMesh(), GetGroupTopology(), &SelectionEditor,
 			LocalRay, UpdateConfig, ResultOut);
 	}
-	else
+}
+
+void FDynamicMeshSelector::UpdateSelectionViaRaycast_MeshTopology(
+	const FWorldRayQueryInfo& RayInfo,
+	FGeometrySelectionEditor& SelectionEditor,
+	const FGeometrySelectionUpdateConfig& UpdateConfig,
+	FGeometrySelectionUpdateResult& ResultOut)
+{
+	FRay3d LocalRay = GetWorldTransformFunc().InverseTransformRay(RayInfo.WorldRay);
+	UE::Geometry::UpdateTriangleSelectionViaRaycast(
+		GetColliderMesh(), &SelectionEditor,
+		LocalRay, UpdateConfig, ResultOut);
+}
+
+
+void FDynamicMeshSelector::UpdateSelectionViaRaycast_GroupEdges(
+	const FWorldRayQueryInfo& RayInfo,
+	FGeometrySelectionEditor& SelectionEditor,
+	const FGeometrySelectionUpdateConfig& UpdateConfig,
+	FGeometrySelectionUpdateResult& ResultOut)
+{
+	FTransformSRT3d WorldTransform = GetWorldTransformFunc();
+	FRay3d LocalRay = WorldTransform.InverseTransformRay(RayInfo.WorldRay);
+
+	auto EdgeHitToleranceTest = [RayInfo, WorldTransform](int32 ID, const FVector3d& A, const FVector3d& B)
 	{
-		UE::Geometry::UpdateTriangleSelectionViaRaycast(
-			GetColliderMesh(), &SelectionEditor,
-			LocalRay, UpdateConfig, ResultOut);
+		return ToolSceneQueriesUtil::PointSnapQuery(RayInfo.CameraState,
+			WorldTransform.TransformPosition(A), WorldTransform.TransformPosition(B));
+	};
+
+	if (SelectionEditor.GetQueryConfig().bOnlyVisible)
+	{
+		// Compute mesh hit and then find nearest-segment to the mesh hit point. This filters out 'backface' hits
+		// and avoids complexity of determing "nearest visible" segment. However this will not work for segments
+		// on the mesh silhouette, or if the mesh surface is hidden, so this path should be optional...maybe controlled by the the SelectionEditor?
+		double RayHitT; int32 HitTriangleID; FVector3d HitBaryCoords;
+		bool bSurfaceMeshHit = GetColliderMesh()->FindNearestHitTriangle(LocalRay, RayHitT, HitTriangleID, HitBaryCoords);
+		if (bSurfaceMeshHit)
+		{
+			FVector3d SurfaceHitPos = LocalRay.PointAt(RayHitT);
+			FSegmentTree3::FSegment Segment;
+			if (GetGroupEdgeSpatial()->FindNearestSegment(SurfaceHitPos, Segment))
+			{
+				FVector3d SegmentPoint = Segment.Segment.NearestPoint(SurfaceHitPos);
+				if (EdgeHitToleranceTest(0, SurfaceHitPos, SegmentPoint))
+				{
+					int32 GroupEdgeID = GetGroupTopology()->FindGroupEdgeID(Segment.ID);
+					if (GroupEdgeID >= 0)
+					{
+						FMeshTriEdgeID TriEdgeID;
+						GetDynamicMesh()->ProcessMesh([&](const FDynamicMesh3& Mesh) { TriEdgeID = Mesh.GetTriEdgeIDFromEdgeID(Segment.ID); });
+						ResultOut.bSelectionModified = UpdateSelectionWithNewElements(&SelectionEditor, UpdateConfig.ChangeType, TArray<uint64>{(uint64)TriEdgeID.Encoded()}, & ResultOut.SelectionDelta);
+						ResultOut.bSelectionMissed = false;
+					}
+				}
+			}
+			// if we took this path we are not going to consider ray-nearest point
+			return;
+		}
 	}
+
+	// find nearest segment to ray
+	FSegmentTree3::FSegment Segment;
+	FSegmentTree3::FRayNearestSegmentInfo SegmentInfo;
+	if (GetGroupEdgeSpatial()->FindNearestVisibleSegmentHitByRay(LocalRay, EdgeHitToleranceTest, Segment, SegmentInfo))
+	{
+		int32 GroupEdgeID = GetGroupTopology()->FindGroupEdgeID(Segment.ID);
+		if (GroupEdgeID >= 0)
+		{
+			FMeshTriEdgeID TriEdgeID;
+			GetDynamicMesh()->ProcessMesh([&](const FDynamicMesh3& Mesh) { TriEdgeID = Mesh.GetTriEdgeIDFromEdgeID(Segment.ID); });
+			FGeoSelectionID SelectionID(TriEdgeID.Encoded(), GroupEdgeID);
+			ResultOut.bSelectionModified = UpdateSelectionWithNewElements(&SelectionEditor, UpdateConfig.ChangeType, TArray<uint64>{SelectionID.Encoded()}, & ResultOut.SelectionDelta);
+			ResultOut.bSelectionMissed = false;
+		}
+	}
+}
+
+
+
+
+void FDynamicMeshSelector::GetSelectionPreviewForRaycast(
+	const FWorldRayQueryInfo& RayInfo,
+	FGeometrySelectionEditor& PreviewEditor)
+{
+	FGeometrySelectionUpdateResult UpdateResult;
+	UpdateSelectionViaRaycast(RayInfo, PreviewEditor,
+		FGeometrySelectionUpdateConfig(),	// defaults to add
+		UpdateResult);
 }
 
 
+
 void FDynamicMeshSelector::GetSelectionFrame(const FGeometrySelection& Selection, FFrame3d& SelectionFrame, bool bTransformToWorld)
 {
 	if (Selection.TopologyType == EGeometryTopologyType::Polygroup)
@@ -208,11 +334,12 @@ void FDynamicMeshSelector::AccumulateSelectionBounds(const FGeometrySelection& S
 
 
 
-void FDynamicMeshSelector::AccumulateSelectionElements(const FGeometrySelection& Selection, FGeometrySelectionElements& Elements, bool bTransformToWorld)
+void FDynamicMeshSelector::AccumulateSelectionElements(const FGeometrySelection& Selection, FGeometrySelectionElements& Elements, bool bTransformToWorld, bool bIsForPreview)
 {
 	FTransform UseWorldTransform = GetLocalToWorldTransform();
 	const FTransform* ApplyTransform = (bTransformToWorld) ? &UseWorldTransform : nullptr;
 
+	bool bMapFacesToEdges = bIsForPreview;
 	if (Selection.TopologyType == EGeometryTopologyType::Polygroup)
 	{
 		const FGroupTopology* Topology = GetGroupTopology();
@@ -222,7 +349,7 @@ void FDynamicMeshSelector::AccumulateSelectionElements(const FGeometrySelection&
 				[&](uint32 VertexID, const FVector3d& Point) { Elements.Points.Add(Point); },
 				[&](uint32 EdgeID, const FSegment3d& Segment) { Elements.Segments.Add(Segment); },
 				[&](uint32 TriangleID, const FTriangle3d& Triangle) { Elements.Triangles.Add(Triangle); },
-				ApplyTransform
+				ApplyTransform, bMapFacesToEdges
 			);
 		});
 	}
@@ -234,7 +361,7 @@ void FDynamicMeshSelector::AccumulateSelectionElements(const FGeometrySelection&
 				[&](uint32 VertexID, const FVector3d& Point) { Elements.Points.Add(Point); },
 				[&](uint32 EdgeID, const FSegment3d& Segment) { Elements.Segments.Add(Segment); },
 				[&](uint32 TriangleID, const FTriangle3d& Triangle) { Elements.Triangles.Add(Triangle); },
-				ApplyTransform
+				ApplyTransform, bMapFacesToEdges
 			);
 		});
 	}
@@ -281,6 +408,36 @@ const FGroupTopology* FDynamicMeshSelector::GetGroupTopology()
 
 
 
+void FDynamicMeshSelector::UpdateGroupEdgeSegmentTree()
+{
+	const FGroupTopology* UseGroupTopology = GetGroupTopology();
+
+	TargetMesh->ProcessMesh([&](const UE::Geometry::FDynamicMesh3& SourceMesh)
+	{
+		TArray<int32> EdgeIDs;
+		for (const FGroupTopology::FGroupEdge& GroupEdge : UseGroupTopology->Edges)
+		{
+			EdgeIDs.Append(GroupEdge.Span.Edges);
+		}
+
+		GroupEdgeSegmentTree = MakePimpl<FSegmentTree3>();
+		GroupEdgeSegmentTree->Build(EdgeIDs, [&SourceMesh](int32 EdgeID)
+		{
+			return FSegment3d( SourceMesh.GetEdgePoint(EdgeID, 0), SourceMesh.GetEdgePoint(EdgeID, 1) );
+		}, SourceMesh.EdgeCount() );
+	});
+}
+
+const FSegmentTree3* FDynamicMeshSelector::GetGroupEdgeSpatial()
+{
+	if (!GroupEdgeSegmentTree.IsValid())
+	{
+		UpdateGroupEdgeSegmentTree();
+	}
+	return GroupEdgeSegmentTree.Get();
+}
+
+
 bool FDynamicMeshComponentSelectorFactory::CanBuildForTarget(FGeometryIdentifier TargetIdentifier) const
 {
 	if (TargetIdentifier.TargetType == FGeometryIdentifier::ETargetType::PrimitiveComponent)
diff --git a/Engine/Plugins/Runtime/MeshModelingToolset/Source/ModelingComponents/Private/Selection/GeometrySelectionManager.cpp b/Engine/Plugins/Runtime/MeshModelingToolset/Source/ModelingComponents/Private/Selection/GeometrySelectionManager.cpp
--- a/Engine/Plugins/Runtime/MeshModelingToolset/Source/ModelingComponents/Private/Selection/GeometrySelectionManager.cpp
+++ b/Engine/Plugins/Runtime/MeshModelingToolset/Source/ModelingComponents/Private/Selection/GeometrySelectionManager.cpp
@@ -4,8 +4,11 @@
 #include "Engine/Engine.h"
 #include "Selection/DynamicMeshSelector.h"
 #include "Selection/SelectionEditInteractiveCommand.h"
+#include "InteractiveToolsContext.h"
+#include "InteractiveToolManager.h"
 #include "ToolContextInterfaces.h"
 #include "ToolDataVisualizer.h"
+#include "Selections/GeometrySelectionUtil.h"
 
 // for debug drawing
 #include "SceneManagement.h"		
@@ -16,6 +19,13 @@
 
 #include UE_INLINE_GENERATED_CPP_BY_NAME(GeometrySelectionManager)
 
+static TAutoConsoleVariable<int32> CVarGeometrySelectionManager_FullSelectionHoverHighlights(
+	TEXT("modeling.Selection.FullHoverHighlights"),
+	0,
+	TEXT("Use full selection hover highlights instead of simplified highlights")
+);
+
+
 
 using namespace UE::Geometry;
 
@@ -109,7 +119,7 @@ void UGeometrySelectionManager::SetSelectionElementTypeInternal(EGeometryElement
 		for (TSharedPtr<FGeometrySelectionTarget> Target : ActiveTargetReferences)
 		{
 			Target->Selection.ElementType = SelectionElementType;
-			// do we need to update/rebuild Editors here? They currently do not have state...
+			Target->SelectionEditor->UpdateQueryConfig(GetCurrentSelectionQueryConfig());
 		}
 	}
 }
@@ -152,7 +162,7 @@ void UGeometrySelectionManager::SetMeshTopologyModeInternal(EMeshTopologyMode Ne
 		for (TSharedPtr<FGeometrySelectionTarget> Target : ActiveTargetReferences)
 		{
 			Target->Selection.TopologyType = GetSelectionTopologyType();
-			// do we need to update/rebuild Editors here? They currently do not have state...
+			Target->SelectionEditor->UpdateQueryConfig(GetCurrentSelectionQueryConfig());
 		}
 	}
 }
@@ -199,6 +209,16 @@ EGeometryTopologyType UGeometrySelectionManager::GetSelectionTopologyType() cons
 
 
 
+UE::Geometry::FGeometrySelectionHitQueryConfig UGeometrySelectionManager::GetCurrentSelectionQueryConfig() const
+{
+	FGeometrySelectionHitQueryConfig HitQueryConfig;
+	HitQueryConfig.TopologyType = GetSelectionTopologyType();
+	HitQueryConfig.ElementType = GetSelectionElementType();
+	HitQueryConfig.bOnlyVisible = true;
+	return HitQueryConfig;
+}
+
+
 bool UGeometrySelectionManager::HasSelection() const
 {
 	for (TSharedPtr<FGeometrySelectionTarget> Target : ActiveTargetReferences)
@@ -388,6 +408,7 @@ TSharedPtr<UGeometrySelectionManager::FGeometrySelectionTarget> UGeometrySelecti
 		// ensure these are current, as they may have changed while Target was asleep
 		FoundTarget->Selection.ElementType = GetSelectionElementType();
 		FoundTarget->Selection.TopologyType = GetSelectionTopologyType();
+		FoundTarget->SelectionEditor->UpdateQueryConfig(GetCurrentSelectionQueryConfig());
 
 		return FoundTarget;
 	}
@@ -408,7 +429,8 @@ TSharedPtr<UGeometrySelectionManager::FGeometrySelectionTarget> UGeometrySelecti
 	SelectionTarget->Selection.TopologyType = GetSelectionTopologyType();
 
 	SelectionTarget->SelectionEditor = MakeUnique<FGeometrySelectionEditor>();
-	SelectionTarget->SelectionEditor->Initialize(&SelectionTarget->Selection);
+	FGeometrySelectionHitQueryConfig HitQueryConfig = GetCurrentSelectionQueryConfig();
+	SelectionTarget->SelectionEditor->Initialize(&SelectionTarget->Selection, HitQueryConfig);
 
 	if (SelectionTarget->Selector->SupportsSleep())
 	{
@@ -439,7 +461,16 @@ bool UGeometrySelectionManager::RayHitTest(
 		return false;
 	}
 
-	bool bHit = ActiveTargetReferences[0]->Selector->RayHitTest(WorldRay, HitResultOut);
+	IGeometrySelector::FWorldRayQueryInfo RayQueryInfo;
+	RayQueryInfo.WorldRay = WorldRay;
+	IToolsContextQueriesAPI* QueryAPI = this->ToolsContext->ToolManager->GetContextQueriesAPI();
+	QueryAPI->GetCurrentViewState(RayQueryInfo.CameraState);
+
+	// currently only going to support one object, not sure how to support more yet...
+	FGeometrySelectionTarget* Target = ActiveTargetReferences[0].Get();
+
+	FGeometrySelectionHitQueryConfig HitQueryConfig = GetCurrentSelectionQueryConfig();
+	bool bHit = Target->Selector->RayHitTest(RayQueryInfo, HitQueryConfig, HitResultOut);
 	if (bHit)
 	{
 		HitResultOut.HitOwner = ActiveTargetReferences[0].Get();
@@ -497,8 +528,13 @@ void UGeometrySelectionManager::UpdateSelectionViaRaycast(
 	// currently only going to support one object, not sure how to support more yet...
 	FGeometrySelectionTarget* Target = ActiveTargetReferences[0].Get();
 
+	IGeometrySelector::FWorldRayQueryInfo RayQueryInfo;
+	RayQueryInfo.WorldRay = WorldRay;
+	IToolsContextQueriesAPI* QueryAPI = this->ToolsContext->ToolManager->GetContextQueriesAPI();
+	QueryAPI->GetCurrentViewState(RayQueryInfo.CameraState);
+
 	Target->Selector->UpdateSelectionViaRaycast(
-		WorldRay, *Target->SelectionEditor, UpdateConfig, ResultOut
+		RayQueryInfo, *Target->SelectionEditor, UpdateConfig, ResultOut
 	);
 
 	if (ResultOut.bSelectionModified)
@@ -522,6 +558,45 @@ void UGeometrySelectionManager::UpdateSelectionViaRaycast(
 
 
 
+bool UGeometrySelectionManager::UpdateSelectionPreviewViaRaycast(
+	const FRay3d& WorldRay )
+{
+	if (ActiveTargetReferences.Num() == 0)
+	{
+		return false;
+	}
+	bool bUseSimplifiedPreviewHighlight = (CVarGeometrySelectionManager_FullSelectionHoverHighlights.GetValueOnGameThread() == 0);
+
+	// currently only going to support one object, not sure how to support more yet...
+
+	FGeometrySelectionTarget* Target = ActiveTargetReferences[0].Get();
+
+	IGeometrySelector::FWorldRayQueryInfo RayQueryInfo;
+	RayQueryInfo.WorldRay = WorldRay;
+	IToolsContextQueriesAPI* QueryAPI = this->ToolsContext->ToolManager->GetContextQueriesAPI();
+	QueryAPI->GetCurrentViewState(RayQueryInfo.CameraState);
+
+	FGeometrySelectionPreview NewPreview( *(Target->SelectionEditor) );
+	Target->Selector->GetSelectionPreviewForRaycast( RayQueryInfo, NewPreview );
+	if ( ! UE::Geometry::AreSelectionsIdentical(NewPreview.PreviewSelection, ActivePreviewSelection) )
+	{
+		ActivePreviewSelection = MoveTemp(NewPreview.PreviewSelection);
+		CachedPreviewRenderElements.Reset();
+		if (ActivePreviewSelection.IsEmpty() == false)
+		{
+			Target->Selector->AccumulateSelectionElements(ActivePreviewSelection, CachedPreviewRenderElements, true, bUseSimplifiedPreviewHighlight);
+		}
+	}
+
+	return (ActivePreviewSelection.IsEmpty() == false);
+}
+
+void UGeometrySelectionManager::ClearSelectionPreview()
+{
+	ActivePreviewSelection.Selection.Reset();
+	CachedPreviewRenderElements.Reset();
+}
+
 
 
 bool UGeometrySelectionManager::GetSelectionBounds(FGeometrySelectionBounds& BoundsOut) const
@@ -814,6 +889,8 @@ void UGeometrySelectionManager::RevertChange(IGeometrySelectionChange* Change)
 void UGeometrySelectionManager::OnTargetGeometryModified(IGeometrySelector* Selector)
 {
 	bSelectionRenderCachesDirty = true;
+
+	ClearActivePreview();
 }
 
 
@@ -822,6 +899,8 @@ void UGeometrySelectionManager::UpdateSelectionRenderCacheOnTargetChange()
 	CachedSelectionRenderElements.Reset();
 	CachedSelectionRenderElements.SetNum(ActiveTargetReferences.Num());
 	bSelectionRenderCachesDirty = true;
+
+	ClearActivePreview();
 }
 
 void UGeometrySelectionManager::RebuildSelectionRenderCaches()
@@ -839,13 +918,20 @@ void UGeometrySelectionManager::RebuildSelectionRenderCaches()
 		FGeometrySelectionElements& Elements = CachedSelectionRenderElements[k];
 		Elements.Reset();
 
-		Target->Selector->AccumulateSelectionElements(Target->Selection, Elements, true);
+		Target->Selector->AccumulateSelectionElements(Target->Selection, Elements, true, false);
 	}
 
 	bSelectionRenderCachesDirty = false;
 }
 
 
+void UGeometrySelectionManager::ClearActivePreview()
+{
+	ActivePreviewSelection.Reset();
+	CachedPreviewRenderElements.Reset();
+}
+
+
 void UGeometrySelectionManager::DebugPrintSelection()
 {
 	if (ActiveTargetReferences.Num() == 0)
@@ -901,13 +987,13 @@ void UGeometrySelectionManager::DebugRender(IToolsContextRenderAPI* RenderAPI)
 		Visualizer.bDepthTested = false;
 		Visualizer.BeginFrame(RenderAPI);
 
-		Visualizer.SetLineParameters(FLinearColor(0,1,0,1), 2.0f);
+		Visualizer.SetLineParameters(FLinearColor(0,0.3f,0.95f,1), 3.0f);
 		for (const FSegment3d& Segment : Elements.Segments)
 		{
 			Visualizer.DrawLine(Segment.StartPoint(), Segment.EndPoint());
 		}
 
-		Visualizer.SetPointParameters(FLinearColor(0,0,1,1), 10.0f);
+		Visualizer.SetPointParameters(FLinearColor(0,0.3f,0.95f,1), 10.0f);
 		for (const FVector3d& Point : Elements.Points)
 		{
 			Visualizer.DrawPoint(Point);
@@ -916,6 +1002,45 @@ void UGeometrySelectionManager::DebugRender(IToolsContextRenderAPI* RenderAPI)
 		Visualizer.EndFrame();
 	}
 
+	// draw selection preview
+	{
+		FPrimitiveDrawInterface* CurrentPDI = RenderAPI->GetPrimitiveDrawInterface();
+		FDynamicMeshBuilder MeshBuilder(CurrentPDI->View->GetFeatureLevel());
+		int32 DepthPriority = SDPG_World; // SDPG_Foreground;  // SDPG_World
+		FVector2f UVs[3] = { FVector2f(0,0), FVector2f(0,1), FVector2f(1,1) };
+		FVector3f Normal = FVector3f(0, 0, 1);
+		FVector3f Tangent = FVector3f(1, 0, 0);
+		for (const FTriangle3d& Triangle : CachedPreviewRenderElements.Triangles)
+		{
+			int32 V0 = MeshBuilder.AddVertex(FDynamicMeshVertex((FVector3f)Triangle.V[0], Tangent, Normal, UVs[0], FColor::White));
+			int32 V1 = MeshBuilder.AddVertex(FDynamicMeshVertex((FVector3f)Triangle.V[1], Tangent, Normal, UVs[1], FColor::White));
+			int32 V2 = MeshBuilder.AddVertex(FDynamicMeshVertex((FVector3f)Triangle.V[2], Tangent, Normal, UVs[2], FColor::White));
+			MeshBuilder.AddTriangle(V0, V1, V2);
+		}
+		//FMaterialRenderProxy* MaterialRenderProxy = TriangleMaterial->GetRenderProxy();		// currently does not work, material does not render
+		FMaterialRenderProxy* MaterialRenderProxy = GEngine->ConstraintLimitMaterialX->GetRenderProxy();
+		MeshBuilder.Draw(CurrentPDI, FMatrix::Identity, MaterialRenderProxy, DepthPriority, false, false);
+
+
+		FToolDataVisualizer Visualizer;
+		Visualizer.bDepthTested = false;
+		Visualizer.BeginFrame(RenderAPI);
+
+		Visualizer.SetLineParameters(FLinearColor(1, 1, 0, 1), 1.0f);
+		for (const FSegment3d& Segment : CachedPreviewRenderElements.Segments)
+		{
+			Visualizer.DrawLine(Segment.StartPoint(), Segment.EndPoint());
+		}
+
+		Visualizer.SetPointParameters(FLinearColor(1, 1, 0, 1), 5.0f);
+		for (const FVector3d& Point : CachedPreviewRenderElements.Points)
+		{
+			Visualizer.DrawPoint(Point);
+		}
+
+		Visualizer.EndFrame();
+	}
+
 }
 
 
diff --git a/Engine/Plugins/Runtime/MeshModelingToolset/Source/ModelingComponents/Public/Selection/DynamicMeshSelector.h b/Engine/Plugins/Runtime/MeshModelingToolset/Source/ModelingComponents/Public/Selection/DynamicMeshSelector.h
--- a/Engine/Plugins/Runtime/MeshModelingToolset/Source/ModelingComponents/Public/Selection/DynamicMeshSelector.h
+++ b/Engine/Plugins/Runtime/MeshModelingToolset/Source/ModelingComponents/Public/Selection/DynamicMeshSelector.h
@@ -16,7 +16,7 @@ class FDynamicMeshSelectionTransformer;
 class FMeshVertexChangeBuilder;
 PREDECLARE_GEOMETRY(class FGroupTopology);
 PREDECLARE_GEOMETRY(class FColliderMesh);
-
+PREDECLARE_GEOMETRY(class FSegmentTree3);
 
 /**
  * FDynamicMeshSelector is an implementation of IGeometrySelector for a UDynamicMesh.
@@ -67,16 +67,20 @@ public:
 	}
 
 	virtual bool RayHitTest(
-		const FRay3d& WorldRay,
+		const FWorldRayQueryInfo& RayInfo,
+		UE::Geometry::FGeometrySelectionHitQueryConfig QueryConfig,
 		FInputRayHit& HitResultOut) override;
 
 
 	virtual void UpdateSelectionViaRaycast(
-		const FRay3d& WorldRay,
+		const FWorldRayQueryInfo& RayInfo,
 		FGeometrySelectionEditor& SelectionEditor,
 		const FGeometrySelectionUpdateConfig& UpdateConfig,
 		FGeometrySelectionUpdateResult& ResultOut) override;
 
+	virtual void GetSelectionPreviewForRaycast(
+		const FWorldRayQueryInfo& RayInfo,
+		FGeometrySelectionEditor& PreviewEditor) override;
 
 	virtual FTransform GetLocalToWorldTransform() const
 	{
@@ -85,7 +89,7 @@ public:
 
 	virtual void GetSelectionFrame(const FGeometrySelection& Selection, UE::Geometry::FFrame3d& SelectionFrame, bool bTransformToWorld) override;
 	virtual void AccumulateSelectionBounds(const FGeometrySelection& Selection, FGeometrySelectionBounds& BoundsInOut, bool bTransformToWorld) override;
-	virtual void AccumulateSelectionElements(const FGeometrySelection& Selection, FGeometrySelectionElements& Elements, bool bTransformToWorld) override;
+	virtual void AccumulateSelectionElements(const FGeometrySelection& Selection, FGeometrySelectionElements& Elements, bool bTransformToWorld, bool bIsForPreview) override;
 
 	virtual IGeometrySelectionTransformer* InitializeTransformation(const FGeometrySelection& Selection) override;
 	virtual void ShutdownTransformation(IGeometrySelectionTransformer* Transformer) override;
@@ -115,12 +119,35 @@ protected:
 	void UpdateGroupTopology();
 	const UE::Geometry::FGroupTopology* GetGroupTopology();
 
+	//
+	// GroupEdgeSegmentTree stores a hit-testable AABBTree for the polygroup edges (depends on GroupTopology)
+	//
+	TPimplPtr<UE::Geometry::FSegmentTree3> GroupEdgeSegmentTree;
+	void UpdateGroupEdgeSegmentTree();
+	const UE::Geometry::FSegmentTree3* GetGroupEdgeSpatial();
+
 	// support for sleep/restore
 	TWeakObjectPtr<UDynamicMesh> SleepingTargetMesh = nullptr;
 
 
 	TPimplPtr<FDynamicMeshSelectionTransformer> ActiveTransformer;
 
+
+
+	virtual void UpdateSelectionViaRaycast_GroupEdges(
+		const FWorldRayQueryInfo& RayInfo,
+		FGeometrySelectionEditor& SelectionEditor,
+		const FGeometrySelectionUpdateConfig& UpdateConfig,
+		FGeometrySelectionUpdateResult& ResultOut);
+
+	virtual void UpdateSelectionViaRaycast_MeshTopology(
+		const FWorldRayQueryInfo& RayInfo,
+		FGeometrySelectionEditor& SelectionEditor,
+		const FGeometrySelectionUpdateConfig& UpdateConfig,
+		FGeometrySelectionUpdateResult& ResultOut);
+
+
+
 	// give Transformer access to internals 
 	friend class FDynamicMeshSelectionTransformer;
 };
diff --git a/Engine/Plugins/Runtime/MeshModelingToolset/Source/ModelingComponents/Public/Selection/GeometrySelectionManager.h b/Engine/Plugins/Runtime/MeshModelingToolset/Source/ModelingComponents/Public/Selection/GeometrySelectionManager.h
--- a/Engine/Plugins/Runtime/MeshModelingToolset/Source/ModelingComponents/Public/Selection/GeometrySelectionManager.h
+++ b/Engine/Plugins/Runtime/MeshModelingToolset/Source/ModelingComponents/Public/Selection/GeometrySelectionManager.h
@@ -178,6 +178,19 @@ public:
 	FModelingSelectionInteraction_SelectionModified OnSelectionModified;
 
 
+
+	//
+	// Hover/Preview support
+	//
+public:
+
+	virtual bool UpdateSelectionPreviewViaRaycast(
+		const FRay3d& WorldRay
+	);
+
+	virtual void ClearSelectionPreview();
+
+
 	//
 	// Selection queries
 	//
@@ -293,6 +306,8 @@ protected:
 	EMeshTopologyMode MeshTopologyMode = EMeshTopologyMode::None;
 	void SetMeshTopologyModeInternal(EMeshTopologyMode NewTopologyMode);
 
+	UE::Geometry::FGeometrySelectionHitQueryConfig GetCurrentSelectionQueryConfig() const;
+
 	// ITF references
 
 	UPROPERTY()
@@ -352,6 +367,16 @@ protected:
 	void UpdateSelectionRenderCacheOnTargetChange();
 	void RebuildSelectionRenderCaches();
 
+
+	//
+	// 3D geometry for active hover/preview highlight
+	// note: currently only supporting single target here, will need to be refactored to handle multiple targets...
+	//
+	FGeometrySelection ActivePreviewSelection;
+	FGeometrySelectionElements CachedPreviewRenderElements;
+	void ClearActivePreview();
+
+
 	// various change types need internal access
 
 	friend class FGeometrySelectionManager_SelectionTypeChange;
diff --git a/Engine/Plugins/Runtime/MeshModelingToolset/Source/ModelingComponents/Public/Selection/GeometrySelector.h b/Engine/Plugins/Runtime/MeshModelingToolset/Source/ModelingComponents/Public/Selection/GeometrySelector.h
--- a/Engine/Plugins/Runtime/MeshModelingToolset/Source/ModelingComponents/Public/Selection/GeometrySelector.h
+++ b/Engine/Plugins/Runtime/MeshModelingToolset/Source/ModelingComponents/Public/Selection/GeometrySelector.h
@@ -6,6 +6,7 @@
 #include "Selections/GeometrySelection.h"
 #include "FrameTypes.h"
 #include "InputState.h"
+#include "ToolContextInterfaces.h"
 
 class IToolsContextTransactionsAPI;
 
@@ -255,13 +256,20 @@ public:
 	 */
 	virtual FGeometryIdentifier GetIdentifier() const = 0;
 
+	struct FWorldRayQueryInfo
+	{
+		FRay3d WorldRay;
+		FViewCameraState CameraState;
+	};
+
 	/**
 	 * Check for intersection between a world-space Ray with the Selector's target object and return a FInputRayHit result.
 	 * RayHitTest() must return true for other raycast-based functions like UpdateSelectionViaRaycast() to be called.
 	 * @return true on hit, false on miss
 	 */
 	virtual bool RayHitTest(
-		const FRay3d& WorldRay,
+		const FWorldRayQueryInfo& RayInfo,
+		UE::Geometry::FGeometrySelectionHitQueryConfig QueryConfig,
 		FInputRayHit& HitResultOut
 	) = 0;
 
@@ -271,12 +279,22 @@ public:
 	 * Information about any changes actually made to the selection via the SelectionEditor are returned in ResultOut
 	 */
 	virtual void UpdateSelectionViaRaycast(	
-		const FRay3d& WorldRay,
+		const FWorldRayQueryInfo& RayInfo,
 		FGeometrySelectionEditor& SelectionEditor,
 		const FGeometrySelectionUpdateConfig& UpdateConfig,
 		FGeometrySelectionUpdateResult& ResultOut 
 	) = 0;
 
+	/**
+	 * Query the Selector for a potential selection based on world-space Ray. The intention of this query
+	 * is that it be used for "selection preview", eg like hover-highlighting. The PreviewEditor will
+	 * always be updated in 'Add' mode, and currently no result information will be returned. 
+	 */
+	virtual void GetSelectionPreviewForRaycast(
+		const FWorldRayQueryInfo& RayInfo,
+		FGeometrySelectionEditor& PreviewEditor
+	) = 0;
+
 	/**
 	 * @return the World transform for the Selector's target object
 	 */
@@ -299,8 +317,9 @@ public:
 	 * Accumulate geometric elements (currently 3D triangles, line segments, and points) for the provided Selection in the provided ElementsInOut. 
 	 * ElementsInOut is not cleared.
 	 * @param bTransformToWorld if true each geometric element will be transformed to World space, based on GetLocalToWorldTransform()
+	 * @param bIsForPreview if true, geometry is being collected for a preview of selection. Selector may return simplified geometry in this case
 	 */
-	virtual void AccumulateSelectionElements(const FGeometrySelection& Selection, FGeometrySelectionElements& ElementsInOut, bool bTransformToWorld) = 0;
+	virtual void AccumulateSelectionElements(const FGeometrySelection& Selection, FGeometrySelectionElements& ElementsInOut, bool bTransformToWorld, bool bIsForPreview) = 0;
 
 
 	/**
