ModelingMode: Mesh Element Selection system evolution.
Refactor most of FDynamicMeshSelector into FDynamicMeshSelectorBase, FDynamicMeshSelector is now a subclass.
Move the DynamicMeshTransformer from cpp to FBasicDynamicMeshSelectionTransformer in header.
Add FVolumeSelector and FStaticMeshSelector which are derived from FBaseDynamicMeshSelector.
Add ability to filter selection queries/edits by TopologyID in FGeometrySelectionEditor. This is necesary for polygroup and polygroup-edge selections as their encoding includes an arbitrary TriangleID paired w/ the TopologyID.
UGeometrySelectionManager now manages this TopologyID filtering on selection editors it creates/updates
Add concept of "Tracked" selection changes to UGeometrySelectionManager, allows a single selection FChange to be constructed incrementally, eg like with a brush. Add a variant of raycast-selection-update that can be run inside a tracked change.

diff --git a/Engine/Plugins/Runtime/GeometryProcessing/Source/DynamicMesh/Private/Selections/GeometrySelection.cpp b/Engine/Plugins/Runtime/GeometryProcessing/Source/DynamicMesh/Private/Selections/GeometrySelection.cpp
--- a/Engine/Plugins/Runtime/GeometryProcessing/Source/DynamicMesh/Private/Selections/GeometrySelection.cpp
+++ b/Engine/Plugins/Runtime/GeometryProcessing/Source/DynamicMesh/Private/Selections/GeometrySelection.cpp
@@ -2,24 +2,24 @@
 
 #include "Selections/GeometrySelection.h"
 #include "DynamicMesh/DynamicMesh3.h"
+#include "Algo/Find.h"
 
 using namespace UE::Geometry;
 
-void FGeometrySelectionEditor::Initialize(FGeometrySelection* TargetSelectionIn, const FGeometrySelectionHitQueryConfig& QueryConfigIn)
+void FGeometrySelectionEditor::Initialize(FGeometrySelection* TargetSelectionIn, const FGeometrySelectionHitQueryConfig& QueryConfigIn, bool bEnableTopologyIDFilteringIn)
 {
 	check(TargetSelectionIn != nullptr);
-	check(QueryConfigIn.ElementType == TargetSelectionIn->ElementType);
-	check(QueryConfigIn.TopologyType == TargetSelectionIn->TopologyType);
 	TargetSelection = TargetSelectionIn;
-	QueryConfig = QueryConfigIn;
+	UpdateQueryConfig(QueryConfigIn, bEnableTopologyIDFilteringIn);
 }
 
 
-void FGeometrySelectionEditor::UpdateQueryConfig(const FGeometrySelectionHitQueryConfig& NewConfig)
+void FGeometrySelectionEditor::UpdateQueryConfig(const FGeometrySelectionHitQueryConfig& NewConfig, bool bEnableTopologyIDFilteringIn)
 {
 	check(NewConfig.ElementType == TargetSelection->ElementType);
 	check(NewConfig.TopologyType == TargetSelection->TopologyType);
 	QueryConfig = NewConfig;
+	bEnableTopologyIDFiltering = bEnableTopologyIDFilteringIn;
 }
 
 void FGeometrySelectionEditor::ClearSelection(FGeometrySelectionDelta& DeltaOut)
@@ -35,9 +35,47 @@ void FGeometrySelectionEditor::ClearSelection(FGeometrySelectionDelta& DeltaOut)
 
 bool FGeometrySelectionEditor::IsSelected(uint64 ID) const
 {
+	if (bEnableTopologyIDFiltering)
+	{
+		bool bFound;
+		RemapToExistingTopologyID(ID, bFound);
+		return bFound;
+	}
 	return TargetSelection->Selection.Contains(ID);
 }
 
+bool FGeometrySelectionEditor::RemoveFromSelection(uint64 ID)
+{
+	if (bEnableTopologyIDFiltering)
+	{
+		bool bFound;
+		ID = RemapToExistingTopologyID(ID, bFound);
+		if (!bFound)
+		{
+			return false;
+		}
+	}
+	int32 NumRemoved = TargetSelection->Selection.Remove(ID);
+	checkSlow(NumRemoved == 0 || NumRemoved == 1);
+	return (NumRemoved > 0);
+}
+
+uint64 FGeometrySelectionEditor::RemapToExistingTopologyID(uint64 ID, bool& bFound) const
+{
+	bFound = false;
+	uint32 TopologyID = FGeoSelectionID(ID).TopologyID;
+	const uint64* Found = Algo::FindByPredicate(TargetSelection->Selection, [&](uint64 Item)
+	{
+		return FGeoSelectionID(Item).TopologyID == TopologyID;
+	});
+	if (Found != nullptr)
+	{
+		bFound = true;
+		return *Found;
+	}
+	return ID;
+}
+
 bool FGeometrySelectionEditor::Replace(const FGeometrySelection& NewSelection, FGeometrySelectionDelta& DeltaOut)
 {
 	// currently hard-enforcing this...maybe we can allow switches? 
diff --git a/Engine/Plugins/Runtime/GeometryProcessing/Source/DynamicMesh/Public/Selections/GeometrySelection.h b/Engine/Plugins/Runtime/GeometryProcessing/Source/DynamicMesh/Public/Selections/GeometrySelection.h
--- a/Engine/Plugins/Runtime/GeometryProcessing/Source/DynamicMesh/Public/Selections/GeometrySelection.h
+++ b/Engine/Plugins/Runtime/GeometryProcessing/Source/DynamicMesh/Public/Selections/GeometrySelection.h
@@ -300,6 +300,17 @@ struct DYNAMICMESH_API FGeometrySelectionUpdateResult
  * a FGeometrySelection. The various functions can be used to add or remove
  * to the Selection, while also tracking what changed, returned via
  * FGeometrySelectionDelta structs. 
+ * 
+ * In some cases (eg Polygroup selections on faces and edges), only the TopologyIDs
+ * are unique, and the same TopologyID may be paired with an arbitrary TriangleID/EdgeID.
+ * FGeometrySelection is currently not aware of this distinction, however it means when
+ * (for example) adding to a selection, a new uint64 ID should only be added if the 
+ * TopologyID is unique, ie the ElementID bit should be ignored. FGeometrySelectionEditor can
+ * do this, by enabling TopologyIDFiltering in the various setup/config functions.
+ * 
+ * Note, however, this means that FGeometrySelectionEditor.IsSelected() must be used
+ * to determine whether an ID is selected, rather than Selection.Contains()
+ * 
  */
 class DYNAMICMESH_API FGeometrySelectionEditor
 {
@@ -310,7 +321,8 @@ public:
 	 */
 	void Initialize(
 		FGeometrySelection* TargetSelectionIn,
-		const FGeometrySelectionHitQueryConfig& QueryConfigIn);
+		const FGeometrySelectionHitQueryConfig& QueryConfigIn,
+		bool bEnableTopologyIDFiltering);
 
 	/** @return the Element Type of the Target Selection */
 	EGeometryElementType GetElementType() const { return TargetSelection->ElementType; }
@@ -320,12 +332,14 @@ public:
 	/** @return the active configuration for this Selection Editor, eg what type of element/topology is being selected. This is redundant w/ the above fields.  */
 	const FGeometrySelectionHitQueryConfig& GetQueryConfig() const { return QueryConfig; }
 
+	bool GetIsTopologyIDFilteringEnabled() const { return bEnableTopologyIDFiltering; }
+
 	/** 
 	 * Update the active QueryConfig for this SelectionEditor. This is necessary to keep it in sync w/ the TargetSelection, if
 	 * the ElementType or TopologyType are modified. This perhaps should be revisited as they basically
 	 * always need to be the same...
 	 */
-	void UpdateQueryConfig(const FGeometrySelectionHitQueryConfig& NewConfig);
+	void UpdateQueryConfig(const FGeometrySelectionHitQueryConfig& NewConfig, bool bEnableTopologyIDFilteringIn);
 
 	/** @return true if the given ID is currently selected in the Target Selection */
 	bool IsSelected(uint64 ID) const;
@@ -333,6 +347,9 @@ public:
 	/** Clear the Target Selection and return change information in DeltaOut */
 	void ClearSelection(FGeometrySelectionDelta& DeltaOut);
 
+	/** Remove ID from Selection */
+	bool RemoveFromSelection(uint64 ID);
+
 	/** Access the Selection object this Editor is modifying */
 	const FGeometrySelection& GetSelection() const { return *TargetSelection; }
 
@@ -343,7 +360,7 @@ public:
 		int32 NumAdded = 0;
 		for (uint64 ID : List)
 		{
-			if (TargetSelection->Selection.Contains(ID) == false)
+			if (IsSelected(ID) == false)
 			{
 				TargetSelection->Selection.Add(ID);
 				DeltaOut.Added.Add(ID);
@@ -360,12 +377,11 @@ public:
 		int32 TotalRemoved = 0;
 		for (uint64 ID : List)
 		{
-			int32 NumRemoved = TargetSelection->Selection.Remove(ID);
-			if (NumRemoved > 0)
+			if (RemoveFromSelection(ID))
 			{
 				DeltaOut.Removed.Add(ID);
+				TotalRemoved++;
 			}
-			TotalRemoved += NumRemoved;
 		}
 		return (TotalRemoved > 0);
 	}
@@ -377,6 +393,9 @@ public:
 protected:
 	FGeometrySelection* TargetSelection = nullptr;
 	FGeometrySelectionHitQueryConfig QueryConfig;
+
+	bool bEnableTopologyIDFiltering = false;
+	uint64 RemapToExistingTopologyID(uint64 ID, bool& bFoundOut) const;
 };
 
 
@@ -395,7 +414,7 @@ public:
 	FGeometrySelectionPreview(const FGeometrySelectionEditor& ActiveEditor)
 	{
 		PreviewSelection.InitializeTypes(ActiveEditor.GetSelection());
-		Initialize(&PreviewSelection, ActiveEditor.GetQueryConfig());
+		Initialize(&PreviewSelection, ActiveEditor.GetQueryConfig(), ActiveEditor.GetIsTopologyIDFilteringEnabled());
 	}
 
 	bool IsEmpty() const { return PreviewSelection.IsEmpty(); }
diff --git a/Engine/Plugins/Runtime/MeshModelingToolset/Source/ModelingComponents/Private/Selection/DynamicMeshSelector.cpp b/Engine/Plugins/Runtime/MeshModelingToolset/Source/ModelingComponents/Private/Selection/DynamicMeshSelector.cpp
--- a/Engine/Plugins/Runtime/MeshModelingToolset/Source/ModelingComponents/Private/Selection/DynamicMeshSelector.cpp
+++ b/Engine/Plugins/Runtime/MeshModelingToolset/Source/ModelingComponents/Private/Selection/DynamicMeshSelector.cpp
@@ -15,20 +15,20 @@
 
 using namespace UE::Geometry;
 
-#define LOCTEXT_NAMESPACE "FDynamicMeshSelector"
+#define LOCTEXT_NAMESPACE "FBaseDynamicMeshSelector"
 
 
-FDynamicMeshSelector::~FDynamicMeshSelector()
+FBaseDynamicMeshSelector::~FBaseDynamicMeshSelector()
 {
 	if (TargetMesh.IsValid())
 	{
-		ensureMsgf(false, TEXT("FDynamicMeshSelector was not properly Shutdown!"));
-		FDynamicMeshSelector::Shutdown();
+		ensureMsgf(false, TEXT("FBaseDynamicMeshSelector was not properly Shutdown!"));
+		FBaseDynamicMeshSelector::Shutdown();
 	}
 }
 
 
-void FDynamicMeshSelector::Initialize(
+void FBaseDynamicMeshSelector::Initialize(
 	FGeometryIdentifier SourceGeometryIdentifierIn,
 	UDynamicMesh* TargetMeshIn, 
 	TUniqueFunction<UE::Geometry::FTransformSRT3d()> GetWorldTransformFuncIn)
@@ -42,25 +42,30 @@ void FDynamicMeshSelector::Initialize(
 	RegisterMeshChangedHandler();
 }
 
-void FDynamicMeshSelector::RegisterMeshChangedHandler()
+void FBaseDynamicMeshSelector::RegisterMeshChangedHandler()
 {
 	TargetMesh_OnMeshChangedHandle = TargetMesh->OnMeshChanged().AddLambda([this](UDynamicMesh* Mesh, FDynamicMeshChangeInfo ChangeInfo)
 	{
-		ColliderMesh.Reset();
-		GroupEdgeSegmentTree.Reset();
+		InvalidateOnMeshChange(ChangeInfo);
+	});
+}
 
-		if (ChangeInfo.Type != EDynamicMeshChangeType::MeshVertexChange)
-		{
-			GroupTopology.Reset();
-		}
+void FBaseDynamicMeshSelector::InvalidateOnMeshChange(FDynamicMeshChangeInfo ChangeInfo)
+{
+	ColliderMesh.Reset();
+	GroupEdgeSegmentTree.Reset();
 
-		// publish geometry-modified event
-		NotifyGeometryModified();
-	});
+	if (ChangeInfo.Type != EDynamicMeshChangeType::MeshVertexChange)
+	{
+		GroupTopology.Reset();
+	}
+
+	// publish geometry-modified event
+	NotifyGeometryModified();
 }
 
 
-void FDynamicMeshSelector::Shutdown()
+void FBaseDynamicMeshSelector::Shutdown()
 {
 	if (TargetMesh.IsValid())
 	{
@@ -74,7 +79,7 @@ void FDynamicMeshSelector::Shutdown()
 	TargetMesh = nullptr;
 }
 
-bool FDynamicMeshSelector::Sleep()
+bool FBaseDynamicMeshSelector::Sleep()
 {
 	check(TargetMesh.IsValid());
 
@@ -91,7 +96,7 @@ bool FDynamicMeshSelector::Sleep()
 	return true;
 }
 
-bool FDynamicMeshSelector::Restore()
+bool FBaseDynamicMeshSelector::Restore()
 {
 	check(SleepingTargetMesh.IsValid());
 
@@ -108,7 +113,7 @@ bool FDynamicMeshSelector::Restore()
 
 
 
-bool FDynamicMeshSelector::RayHitTest(
+bool FBaseDynamicMeshSelector::RayHitTest(
 	const FWorldRayQueryInfo& RayInfo,
 	FGeometrySelectionHitQueryConfig QueryConfig,
 	FInputRayHit& HitResultOut )
@@ -150,7 +155,7 @@ bool FDynamicMeshSelector::RayHitTest(
 }
 
 
-void FDynamicMeshSelector::UpdateSelectionViaRaycast(
+void FBaseDynamicMeshSelector::UpdateSelectionViaRaycast(
 	const FWorldRayQueryInfo& RayInfo,
 	FGeometrySelectionEditor& SelectionEditor,
 	const FGeometrySelectionUpdateConfig& UpdateConfig,
@@ -176,7 +181,7 @@ void FDynamicMeshSelector::UpdateSelectionViaRaycast(
 	}
 }
 
-void FDynamicMeshSelector::UpdateSelectionViaRaycast_MeshTopology(
+void FBaseDynamicMeshSelector::UpdateSelectionViaRaycast_MeshTopology(
 	const FWorldRayQueryInfo& RayInfo,
 	FGeometrySelectionEditor& SelectionEditor,
 	const FGeometrySelectionUpdateConfig& UpdateConfig,
@@ -189,7 +194,7 @@ void FDynamicMeshSelector::UpdateSelectionViaRaycast_MeshTopology(
 }
 
 
-void FDynamicMeshSelector::UpdateSelectionViaRaycast_GroupEdges(
+void FBaseDynamicMeshSelector::UpdateSelectionViaRaycast_GroupEdges(
 	const FWorldRayQueryInfo& RayInfo,
 	FGeometrySelectionEditor& SelectionEditor,
 	const FGeometrySelectionUpdateConfig& UpdateConfig,
@@ -255,7 +260,7 @@ void FDynamicMeshSelector::UpdateSelectionViaRaycast_GroupEdges(
 
 
 
-void FDynamicMeshSelector::GetSelectionPreviewForRaycast(
+void FBaseDynamicMeshSelector::GetSelectionPreviewForRaycast(
 	const FWorldRayQueryInfo& RayInfo,
 	FGeometrySelectionEditor& PreviewEditor)
 {
@@ -267,7 +272,7 @@ void FDynamicMeshSelector::GetSelectionPreviewForRaycast(
 
 
 
-void FDynamicMeshSelector::GetSelectionFrame(const FGeometrySelection& Selection, FFrame3d& SelectionFrame, bool bTransformToWorld)
+void FBaseDynamicMeshSelector::GetSelectionFrame(const FGeometrySelection& Selection, FFrame3d& SelectionFrame, bool bTransformToWorld)
 {
 	if (Selection.TopologyType == EGeometryTopologyType::Polygroup)
 	{
@@ -294,7 +299,7 @@ void FDynamicMeshSelector::GetSelectionFrame(const FGeometrySelection& Selection
 	}
 }
 
-void FDynamicMeshSelector::AccumulateSelectionBounds(const FGeometrySelection& Selection, FGeometrySelectionBounds& BoundsInOut, bool bTransformToWorld)
+void FBaseDynamicMeshSelector::AccumulateSelectionBounds(const FGeometrySelection& Selection, FGeometrySelectionBounds& BoundsInOut, bool bTransformToWorld)
 {
 	FTransform UseTransform = (bTransformToWorld) ? GetLocalToWorldTransform() : FTransform::Identity;
 	UE::Geometry::FAxisAlignedBox3d TargetWorldBounds = UE::Geometry::FAxisAlignedBox3d::Empty();
@@ -334,7 +339,7 @@ void FDynamicMeshSelector::AccumulateSelectionBounds(const FGeometrySelection& S
 
 
 
-void FDynamicMeshSelector::AccumulateSelectionElements(const FGeometrySelection& Selection, FGeometrySelectionElements& Elements, bool bTransformToWorld, bool bIsForPreview)
+void FBaseDynamicMeshSelector::AccumulateSelectionElements(const FGeometrySelection& Selection, FGeometrySelectionElements& Elements, bool bTransformToWorld, bool bIsForPreview)
 {
 	FTransform UseWorldTransform = GetLocalToWorldTransform();
 	const FTransform* ApplyTransform = (bTransformToWorld) ? &UseWorldTransform : nullptr;
@@ -368,7 +373,7 @@ void FDynamicMeshSelector::AccumulateSelectionElements(const FGeometrySelection&
 }
 
 
-void FDynamicMeshSelector::UpdateColliderMesh()
+void FBaseDynamicMeshSelector::UpdateColliderMesh()
 {
 	// should we transform to world??
 	TargetMesh->ProcessMesh([&](const UE::Geometry::FDynamicMesh3& SourceMesh)
@@ -379,7 +384,7 @@ void FDynamicMeshSelector::UpdateColliderMesh()
 	});
 }
 
-const FColliderMesh* FDynamicMeshSelector::GetColliderMesh()
+const FColliderMesh* FBaseDynamicMeshSelector::GetColliderMesh()
 {
 	if (!ColliderMesh.IsValid())
 	{
@@ -389,7 +394,7 @@ const FColliderMesh* FDynamicMeshSelector::GetColliderMesh()
 }
 
 
-void FDynamicMeshSelector::UpdateGroupTopology()
+void FBaseDynamicMeshSelector::UpdateGroupTopology()
 {
 	// TODO would be preferable to not have to use the raw mesh pointer here, however
 	// FGroupTopology currently needs the pointer to do mesh queries
@@ -397,7 +402,7 @@ void FDynamicMeshSelector::UpdateGroupTopology()
 	GroupTopology = MakePimpl<FGroupTopology>(TargetMesh->GetMeshPtr(), true);
 }
 
-const FGroupTopology* FDynamicMeshSelector::GetGroupTopology()
+const FGroupTopology* FBaseDynamicMeshSelector::GetGroupTopology()
 {
 	if (!GroupTopology.IsValid())
 	{
@@ -408,7 +413,7 @@ const FGroupTopology* FDynamicMeshSelector::GetGroupTopology()
 
 
 
-void FDynamicMeshSelector::UpdateGroupEdgeSegmentTree()
+void FBaseDynamicMeshSelector::UpdateGroupEdgeSegmentTree()
 {
 	const FGroupTopology* UseGroupTopology = GetGroupTopology();
 
@@ -428,7 +433,7 @@ void FDynamicMeshSelector::UpdateGroupEdgeSegmentTree()
 	});
 }
 
-const FSegmentTree3* FDynamicMeshSelector::GetGroupEdgeSpatial()
+const FSegmentTree3* FBaseDynamicMeshSelector::GetGroupEdgeSpatial()
 {
 	if (!GroupEdgeSegmentTree.IsValid())
 	{
@@ -469,40 +474,47 @@ TUniquePtr<IGeometrySelector> FDynamicMeshComponentSelectorFactory::BuildForTarg
 
 
 
-/**
-* FDynamicMeshSelectionTransformer is a basic Transformer implementation for a
-* FDynamicMeshSelector. 
-*/
-class FDynamicMeshSelectionTransformer : public IGeometrySelectionTransformer
+
+IGeometrySelectionTransformer* FDynamicMeshSelector::InitializeTransformation(const FGeometrySelection& Selection)
 {
-public:
-	FDynamicMeshSelector* Selector;
+	check(!ActiveTransformer);
 
-	TArray<int32> MeshVertices;
-	TArray<FVector3d> InitialPositions;
-	TSet<int32> TriangleROI;
-	TSet<int32> OverlayNormals;
+	// If we are transforming a DynamicMeshComponent, we want to defer collision updates, otherwise 
+	// complex collision will be rebuilt every frame
+	FGeometryIdentifier ParentIdentifier = GetSourceGeometryIdentifier();
+	if (ParentIdentifier.ObjectType == FGeometryIdentifier::EObjectType::DynamicMeshComponent)
+	{
+		ParentIdentifier.GetAsComponentType<UDynamicMeshComponent>()->SetTransientDeferCollisionUpdates(true);
+	}
 
-	TArray<FVector3d> UpdatedPositions;
+	ActiveTransformer = MakePimpl<FBasicDynamicMeshSelectionTransformer>();
+	ActiveTransformer->Initialize(this);
+	return ActiveTransformer.Get();
+}
 
-	TPimplPtr<FMeshVertexChangeBuilder> ActiveVertexChange;
+void FDynamicMeshSelector::ShutdownTransformation(IGeometrySelectionTransformer* Transformer)
+{
+	ActiveTransformer.Reset();
 
-	virtual IGeometrySelector* GetSelector() const override
+	FGeometryIdentifier ParentIdentifier = GetSourceGeometryIdentifier();
+	if (ParentIdentifier.ObjectType == FGeometryIdentifier::EObjectType::DynamicMeshComponent)
 	{
-		return Selector;
+		ParentIdentifier.GetAsComponentType<UDynamicMeshComponent>()->SetTransientDeferCollisionUpdates(false);
+		ParentIdentifier.GetAsComponentType<UDynamicMeshComponent>()->UpdateCollision(true);
 	}
+}
 
-	virtual void BeginTransform(const FGeometrySelection& Selection) override;
 
-	virtual void UpdateTransform( TFunctionRef<FVector3d(int32 VertexID, const FVector3d& InitialPosition, const FTransform& WorldTransform)> PositionTransformFunc ) override;
-	void UpdatePendingVertexChange(bool bFinal);
 
-	virtual void EndTransform(IToolsContextTransactionsAPI* TransactionsAPI) override;
-};
 
 
 
-void FDynamicMeshSelectionTransformer::BeginTransform(const FGeometrySelection& Selection)
+void FBasicDynamicMeshSelectionTransformer::Initialize(FBaseDynamicMeshSelector* SelectorIn)
+{
+	Selector = SelectorIn;
+}
+
+void FBasicDynamicMeshSelectionTransformer::BeginTransform(const FGeometrySelection& Selection)
 {
 	const FGroupTopology* UseTopology = nullptr;
 	if (Selection.TopologyType == EGeometryTopologyType::Polygroup)
@@ -551,7 +563,7 @@ void FDynamicMeshSelectionTransformer::BeginTransform(const FGeometrySelection&
 	UpdatePendingVertexChange(false);
 }
 
-void FDynamicMeshSelectionTransformer::UpdateTransform( 
+void FBasicDynamicMeshSelectionTransformer::UpdateTransform( 
 	TFunctionRef<FVector3d(int32 VertexID, const FVector3d& InitialPosition, const FTransform& WorldTransform)> PositionTransformFunc 
 )
 {
@@ -580,7 +592,7 @@ void FDynamicMeshSelectionTransformer::UpdateTransform(
 	UpdatePendingVertexChange(false);
 }
 
-void FDynamicMeshSelectionTransformer::UpdatePendingVertexChange(bool bFinal)
+void FBasicDynamicMeshSelectionTransformer::UpdatePendingVertexChange(bool bFinal)
 {
 	// update the vertex change
 	Selector->GetDynamicMesh()->ProcessMesh([&](const FDynamicMesh3& SourceMesh)
@@ -590,7 +602,7 @@ void FDynamicMeshSelectionTransformer::UpdatePendingVertexChange(bool bFinal)
 	});
 }
 
-void FDynamicMeshSelectionTransformer::EndTransform(IToolsContextTransactionsAPI* TransactionsAPI)
+void FBasicDynamicMeshSelectionTransformer::EndTransform(IToolsContextTransactionsAPI* TransactionsAPI)
 {
 	UpdatePendingVertexChange(true);
 
@@ -601,43 +613,15 @@ void FDynamicMeshSelectionTransformer::EndTransform(IToolsContextTransactionsAPI
 	}
 
 	ActiveVertexChange.Reset();
-}
-
-
-
-
-
 
-
-IGeometrySelectionTransformer* FDynamicMeshSelector::InitializeTransformation(const FGeometrySelection& Selection)
-{
-	check(!ActiveTransformer);
-
-	// If we are transforming a DynamicMeshComponent, we want to defer collision updates, otherwise 
-	// complex collision will be rebuilt every frame
-	FGeometryIdentifier ParentIdentifier = GetSourceGeometryIdentifier();
-	if (ParentIdentifier.ObjectType == FGeometryIdentifier::EObjectType::DynamicMeshComponent)
+	if (OnEndTransformFunc)
 	{
-		ParentIdentifier.GetAsComponentType<UDynamicMeshComponent>()->SetTransientDeferCollisionUpdates(true);
+		OnEndTransformFunc(TransactionsAPI);
 	}
-
-	ActiveTransformer = MakePimpl<FDynamicMeshSelectionTransformer>();
-	ActiveTransformer->Selector = this;
-
-	return ActiveTransformer.Get();
 }
 
-void FDynamicMeshSelector::ShutdownTransformation(IGeometrySelectionTransformer* Transformer)
-{
-	ActiveTransformer.Reset();
 
-	FGeometryIdentifier ParentIdentifier = GetSourceGeometryIdentifier();
-	if (ParentIdentifier.ObjectType == FGeometryIdentifier::EObjectType::DynamicMeshComponent)
-	{
-		ParentIdentifier.GetAsComponentType<UDynamicMeshComponent>()->SetTransientDeferCollisionUpdates(false);
-		ParentIdentifier.GetAsComponentType<UDynamicMeshComponent>()->UpdateCollision(true);
-	}
-}
+
 
 
 
diff --git a/Engine/Plugins/Runtime/MeshModelingToolset/Source/ModelingComponents/Private/Selection/GeometrySelectionManager.cpp b/Engine/Plugins/Runtime/MeshModelingToolset/Source/ModelingComponents/Private/Selection/GeometrySelectionManager.cpp
--- a/Engine/Plugins/Runtime/MeshModelingToolset/Source/ModelingComponents/Private/Selection/GeometrySelectionManager.cpp
+++ b/Engine/Plugins/Runtime/MeshModelingToolset/Source/ModelingComponents/Private/Selection/GeometrySelectionManager.cpp
@@ -51,6 +51,11 @@ void UGeometrySelectionManager::Shutdown()
 	ToolsContext = nullptr;
 	TransactionsAPI = nullptr;
 
+	for (TSharedPtr<FGeometrySelectionTarget> Target : ActiveTargetReferences)
+	{
+		SleepOrShutdownTarget(Target.Get(), false);
+	}
+
 	ResetTargetCache();
 	ActiveTargetReferences.Reset();
 	ActiveTargetMap.Reset();
@@ -119,7 +124,8 @@ void UGeometrySelectionManager::SetSelectionElementTypeInternal(EGeometryElement
 		for (TSharedPtr<FGeometrySelectionTarget> Target : ActiveTargetReferences)
 		{
 			Target->Selection.ElementType = SelectionElementType;
-			Target->SelectionEditor->UpdateQueryConfig(GetCurrentSelectionQueryConfig());
+			bool bEnableTopologyFilter = (Target->Selection.TopologyType == EGeometryTopologyType::Polygroup && Target->Selection.ElementType != EGeometryElementType::Vertex);
+			Target->SelectionEditor->UpdateQueryConfig(GetCurrentSelectionQueryConfig(), bEnableTopologyFilter);
 		}
 	}
 }
@@ -142,6 +148,7 @@ void UGeometrySelectionManager::SetSelectionElementType(EGeometryElementType New
 		TUniquePtr<FGeometrySelectionManager_SelectionTypeChange> TypeChange = MakeUnique<FGeometrySelectionManager_SelectionTypeChange>();
 		TypeChange->FromElementType = SelectionElementType;
 		TypeChange->ToElementType = NewElementType;
+		TypeChange->FromTopologyMode = TypeChange->ToTopologyMode = MeshTopologyMode;		// no-op
 		GetTransactionsAPI()->AppendChange(this, MoveTemp(TypeChange), LOCTEXT("ChangeElementType", "Selection Type"));
 
 		SetSelectionElementTypeInternal(NewElementType);
@@ -162,7 +169,8 @@ void UGeometrySelectionManager::SetMeshTopologyModeInternal(EMeshTopologyMode Ne
 		for (TSharedPtr<FGeometrySelectionTarget> Target : ActiveTargetReferences)
 		{
 			Target->Selection.TopologyType = GetSelectionTopologyType();
-			Target->SelectionEditor->UpdateQueryConfig(GetCurrentSelectionQueryConfig());
+			bool bEnableTopologyFilter = (Target->Selection.TopologyType == EGeometryTopologyType::Polygroup && Target->Selection.ElementType != EGeometryElementType::Vertex);
+			Target->SelectionEditor->UpdateQueryConfig(GetCurrentSelectionQueryConfig(), bEnableTopologyFilter);
 		}
 	}
 }
@@ -184,7 +192,8 @@ void UGeometrySelectionManager::SetMeshTopologyMode(EMeshTopologyMode NewTopolog
 		// so that when we undo, we change to the correct type before we restore
 		TUniquePtr<FGeometrySelectionManager_SelectionTypeChange> TypeChange = MakeUnique<FGeometrySelectionManager_SelectionTypeChange>();
 		TypeChange->FromTopologyMode = MeshTopologyMode;
-		TypeChange->FromTopologyMode = NewTopologyMode;
+		TypeChange->ToTopologyMode = NewTopologyMode;
+		TypeChange->FromElementType = TypeChange->ToElementType = SelectionElementType;		// no-op
 		GetTransactionsAPI()->AppendChange(this, MoveTemp(TypeChange),LOCTEXT("ChangeSelectionMode", "Selection Mode"));
 
 		SetMeshTopologyModeInternal(NewTopologyMode);
@@ -408,7 +417,8 @@ TSharedPtr<UGeometrySelectionManager::FGeometrySelectionTarget> UGeometrySelecti
 		// ensure these are current, as they may have changed while Target was asleep
 		FoundTarget->Selection.ElementType = GetSelectionElementType();
 		FoundTarget->Selection.TopologyType = GetSelectionTopologyType();
-		FoundTarget->SelectionEditor->UpdateQueryConfig(GetCurrentSelectionQueryConfig());
+		bool bEnableTopologyFilter = (FoundTarget->Selection.TopologyType == EGeometryTopologyType::Polygroup && FoundTarget->Selection.ElementType != EGeometryElementType::Vertex);
+		FoundTarget->SelectionEditor->UpdateQueryConfig(GetCurrentSelectionQueryConfig(), bEnableTopologyFilter);
 
 		return FoundTarget;
 	}
@@ -430,7 +440,8 @@ TSharedPtr<UGeometrySelectionManager::FGeometrySelectionTarget> UGeometrySelecti
 
 	SelectionTarget->SelectionEditor = MakeUnique<FGeometrySelectionEditor>();
 	FGeometrySelectionHitQueryConfig HitQueryConfig = GetCurrentSelectionQueryConfig();
-	SelectionTarget->SelectionEditor->Initialize(&SelectionTarget->Selection, HitQueryConfig);
+	bool bEnableTopologyFilter = (HitQueryConfig.TopologyType == EGeometryTopologyType::Polygroup && HitQueryConfig.ElementType != EGeometryElementType::Vertex);
+	SelectionTarget->SelectionEditor->Initialize(&SelectionTarget->Selection, HitQueryConfig, bEnableTopologyFilter);
 
 	if (SelectionTarget->Selector->SupportsSleep())
 	{
@@ -557,6 +568,108 @@ void UGeometrySelectionManager::UpdateSelectionViaRaycast(
 }
 
 
+bool UGeometrySelectionManager::CanBeginTrackedSelectionChange() const
+{
+	return ActiveTargetReferences.Num() > 0 && bInTrackedSelectionChange == false;
+}
+
+bool UGeometrySelectionManager::BeginTrackedSelectionChange(FGeometrySelectionUpdateConfig UpdateConfig, bool bClearOnBegin)
+{
+	if (ensureMsgf(CanBeginTrackedSelectionChange(), TEXT("Cannot begin Selection Change - validate CanBeginTrackedSelectionChange() before calling BeginTrackedSelectionChange()")) == false)
+	{
+		return false;
+	}
+
+	GetTransactionsAPI()->BeginUndoTransaction(LOCTEXT("ChangeSelection", "Change Selection"));
+	bInTrackedSelectionChange = true;
+
+	// currently only going to support one object, not sure how to support more yet...
+	FGeometrySelectionTarget* Target = ActiveTargetReferences[0].Get();
+
+	ActiveTrackedUpdateConfig = UpdateConfig;
+	bSelectionModifiedDuringTrackedChange = false;
+
+	// if we are doing a Replace selection, we want to clear on initialization...
+	InitialTrackedDelta = FGeometrySelectionDelta();
+	if (bClearOnBegin)
+	{
+		Target->SelectionEditor->ClearSelection(InitialTrackedDelta);
+		bSelectionModifiedDuringTrackedChange = true;
+	}
+
+	ActiveTrackedSelection = Target->Selection;
+	ActiveTrackedDelta = FGeometrySelectionDelta();
+
+	if (bClearOnBegin && InitialTrackedDelta.IsEmpty() == false)
+	{
+		bSelectionRenderCachesDirty = true;
+		OnSelectionModified.Broadcast();
+	}
+
+	return true;
+}
+
+void UGeometrySelectionManager::AccumulateSelectionUpdate_Raycast(
+	const FRay3d& WorldRay,
+	FGeometrySelectionUpdateResult& ResultOut )
+{
+	if (!ensure(bInTrackedSelectionChange)) return;
+
+	// currently only going to support one object, not sure how to support more yet...
+	FGeometrySelectionTarget* Target = ActiveTargetReferences[0].Get();
+
+	IGeometrySelector::FWorldRayQueryInfo RayQueryInfo;
+	RayQueryInfo.WorldRay = WorldRay;
+	IToolsContextQueriesAPI* QueryAPI = this->ToolsContext->ToolManager->GetContextQueriesAPI();
+	QueryAPI->GetCurrentViewState(RayQueryInfo.CameraState);
+
+	Target->Selector->UpdateSelectionViaRaycast(
+		RayQueryInfo, *Target->SelectionEditor, ActiveTrackedUpdateConfig, ResultOut	);
+
+	if (ResultOut.bSelectionModified)
+	{
+		bSelectionModifiedDuringTrackedChange = true;
+		ActiveTrackedDelta.Added.Append( ResultOut.SelectionDelta.Added );
+		ActiveTrackedDelta.Removed.Append( ResultOut.SelectionDelta.Removed );
+
+		bSelectionRenderCachesDirty = true;
+		OnSelectionModified.Broadcast();
+	}
+}
+
+void UGeometrySelectionManager::EndTrackedSelectionChange()
+{
+	if ( ensure(bInTrackedSelectionChange) )
+	{
+		if (bSelectionModifiedDuringTrackedChange)
+		{
+			FGeometrySelectionTarget* Target = ActiveTargetReferences[0].Get();
+
+			if (InitialTrackedDelta.IsEmpty() == false)
+			{
+				TUniquePtr<FGeometrySelectionDeltaChange> InitialDeltaChange = MakeUnique<FGeometrySelectionDeltaChange>();
+				InitialDeltaChange->Identifier = Target->SelectionIdentifer;
+				InitialDeltaChange->Delta = MoveTemp(InitialTrackedDelta);
+				GetTransactionsAPI()->AppendChange(this, MoveTemp(InitialDeltaChange), LOCTEXT("ChangeSelection", "Change Selection"));
+			}
+
+			if (ActiveTrackedDelta.IsEmpty() == false)
+			{
+				TUniquePtr<FGeometrySelectionDeltaChange> AccumDeltaChange = MakeUnique<FGeometrySelectionDeltaChange>();
+				AccumDeltaChange->Identifier = Target->SelectionIdentifer;
+				AccumDeltaChange->Delta = MoveTemp(ActiveTrackedDelta);
+				GetTransactionsAPI()->AppendChange(this, MoveTemp(AccumDeltaChange), LOCTEXT("ChangeSelection", "Change Selection"));
+			}
+		}
+
+		GetTransactionsAPI()->EndUndoTransaction();
+		bInTrackedSelectionChange = false;
+	}
+}
+
+
+
+
 
 bool UGeometrySelectionManager::UpdateSelectionPreviewViaRaycast(
 	const FRay3d& WorldRay )
diff --git a/Engine/Plugins/Runtime/MeshModelingToolset/Source/ModelingComponents/Public/Selection/DynamicMeshSelector.h b/Engine/Plugins/Runtime/MeshModelingToolset/Source/ModelingComponents/Public/Selection/DynamicMeshSelector.h
--- a/Engine/Plugins/Runtime/MeshModelingToolset/Source/ModelingComponents/Public/Selection/DynamicMeshSelector.h
+++ b/Engine/Plugins/Runtime/MeshModelingToolset/Source/ModelingComponents/Public/Selection/DynamicMeshSelector.h
@@ -12,24 +12,25 @@
 #include "TransformTypes.h"
 
 
-class FDynamicMeshSelectionTransformer;
+class FBasicDynamicMeshSelectionTransformer;
 class FMeshVertexChangeBuilder;
 PREDECLARE_GEOMETRY(class FGroupTopology);
 PREDECLARE_GEOMETRY(class FColliderMesh);
 PREDECLARE_GEOMETRY(class FSegmentTree3);
 
 /**
- * FDynamicMeshSelector is an implementation of IGeometrySelector for a UDynamicMesh.
- * Note that the Selector itself does *not* require that the target object be a UDynamicMeshComponent.
+ * FBaseDynamicMeshSelector is an implementation of IGeometrySelector for a UDynamicMesh.
+ * Note that the Selector itself does *not* require that the target object be a UDynamicMeshComponent,
+ * and subclasses of FBaseDynamicMeshSelector are used for both Volumes and StaticMeshComponents.
  * Access to the World transform is provided by a TUniqueFunction set up in the Factory.
  */
-class MODELINGCOMPONENTS_API FDynamicMeshSelector : public IGeometrySelector
+class MODELINGCOMPONENTS_API FBaseDynamicMeshSelector : public IGeometrySelector
 {
 public:
-	virtual ~FDynamicMeshSelector();
+	virtual ~FBaseDynamicMeshSelector();
 
 	/**
-	 * Initialize the FDynamicMeshSelector for a given source/target UDynamicMesh.
+	 * Initialize the FBaseDynamicMeshSelector for a given source/target UDynamicMesh.
 	 * @param SourceGeometryIdentifier identifier for the object that the TargetMesh came from (eg DynamicMeshComponent or other UDynamicMesh source)
 	 * @param TargetMesh the target UDynamicMesh
 	 * @param GetWorldTransformFunc function that provides the Local to World Transform
@@ -91,19 +92,21 @@ public:
 	virtual void AccumulateSelectionBounds(const FGeometrySelection& Selection, FGeometrySelectionBounds& BoundsInOut, bool bTransformToWorld) override;
 	virtual void AccumulateSelectionElements(const FGeometrySelection& Selection, FGeometrySelectionElements& Elements, bool bTransformToWorld, bool bIsForPreview) override;
 
-	virtual IGeometrySelectionTransformer* InitializeTransformation(const FGeometrySelection& Selection) override;
-	virtual void ShutdownTransformation(IGeometrySelectionTransformer* Transformer) override;
-
+public:
+	// these need to be public for the Transformers...can we do it another way?
+	UDynamicMesh* GetDynamicMesh() const { return TargetMesh.Get(); }
+	const UE::Geometry::FGroupTopology* GetGroupTopology();
 
 protected:
 	FGeometryIdentifier SourceGeometryIdentifier;
 	TUniqueFunction<UE::Geometry::FTransformSRT3d()> GetWorldTransformFunc;
 
 	TWeakObjectPtr<UDynamicMesh> TargetMesh;
-	UDynamicMesh* GetDynamicMesh() const { return TargetMesh.Get(); }
+	//UDynamicMesh* GetDynamicMesh() const { return TargetMesh.Get(); }		// temporarily public above
 
 	FDelegateHandle TargetMesh_OnMeshChangedHandle;
 	void RegisterMeshChangedHandler();
+	void InvalidateOnMeshChange(FDynamicMeshChangeInfo ChangeInfo);
 
 	//
 	// FColliderMesh is used to store a hit-testable AABBTree independent of the UDynamicMesh 
@@ -117,7 +120,7 @@ protected:
 	//
 	TPimplPtr<UE::Geometry::FGroupTopology> GroupTopology;
 	void UpdateGroupTopology();
-	const UE::Geometry::FGroupTopology* GetGroupTopology();
+	//const UE::Geometry::FGroupTopology* GetGroupTopology();		// temporarily public above
 
 	//
 	// GroupEdgeSegmentTree stores a hit-testable AABBTree for the polygroup edges (depends on GroupTopology)
@@ -130,10 +133,6 @@ protected:
 	TWeakObjectPtr<UDynamicMesh> SleepingTargetMesh = nullptr;
 
 
-	TPimplPtr<FDynamicMeshSelectionTransformer> ActiveTransformer;
-
-
-
 	virtual void UpdateSelectionViaRaycast_GroupEdges(
 		const FWorldRayQueryInfo& RayInfo,
 		FGeometrySelectionEditor& SelectionEditor,
@@ -146,13 +145,29 @@ protected:
 		const FGeometrySelectionUpdateConfig& UpdateConfig,
 		FGeometrySelectionUpdateResult& ResultOut);
 
+};
+
+
+/**
+ * FDynamicMeshSelector is an implementation of FBaseDynamicMeshSelector meant to be used
+ * with UDynamicMeshComponents.
+ */
+class MODELINGCOMPONENTS_API FDynamicMeshSelector : public FBaseDynamicMeshSelector
+{
+public:
 
+	virtual IGeometrySelectionTransformer* InitializeTransformation(const FGeometrySelection& Selection) override;
+	virtual void ShutdownTransformation(IGeometrySelectionTransformer* Transformer) override;
+
+protected:
+	TPimplPtr<FBasicDynamicMeshSelectionTransformer> ActiveTransformer;
 
 	// give Transformer access to internals 
-	friend class FDynamicMeshSelectionTransformer;
+	friend class FBasicDynamicMeshSelectionTransformer;
 };
 
 
+
 /**
  * FDynamicMeshComponentSelectorFactory constructs FDynamicMeshSelector instances 
  * for UDynamicMeshComponents
@@ -166,6 +181,44 @@ public:
 };
 
 
+/**
+* BasicDynamicMeshSelectionTransformer is a basic Transformer implementation that can be
+* used with a FBaseDynamicMeshSelector. This Transformer moves the selected vertices and 
+* nothing else (ie no polygroup-based soft deformation)
+*/
+class MODELINGCOMPONENTS_API FBasicDynamicMeshSelectionTransformer : public IGeometrySelectionTransformer
+{
+public:
+	virtual void Initialize(FBaseDynamicMeshSelector* Selector);
+
+	virtual IGeometrySelector* GetSelector() const override
+	{
+		return Selector;
+	}
+
+	virtual void BeginTransform(const UE::Geometry::FGeometrySelection& Selection) override;
+
+	virtual void UpdateTransform( TFunctionRef<FVector3d(int32 VertexID, const FVector3d& InitialPosition, const FTransform& WorldTransform)> PositionTransformFunc ) override;
+	void UpdatePendingVertexChange(bool bFinal);
+
+	virtual void EndTransform(IToolsContextTransactionsAPI* TransactionsAPI) override;
+
+	TFunction<void(IToolsContextTransactionsAPI* TransactionsAPI)> OnEndTransformFunc;
+
+protected:
+	FBaseDynamicMeshSelector* Selector;
+
+	TArray<int32> MeshVertices;
+	TArray<FVector3d> InitialPositions;
+	TSet<int32> TriangleROI;
+	TSet<int32> OverlayNormals;
+
+	TArray<FVector3d> UpdatedPositions;
+
+	TPimplPtr<FMeshVertexChangeBuilder> ActiveVertexChange;
+
+};
+
 
 
 
diff --git a/Engine/Plugins/Runtime/MeshModelingToolset/Source/ModelingComponents/Public/Selection/GeometrySelectionManager.h b/Engine/Plugins/Runtime/MeshModelingToolset/Source/ModelingComponents/Public/Selection/GeometrySelectionManager.h
--- a/Engine/Plugins/Runtime/MeshModelingToolset/Source/ModelingComponents/Public/Selection/GeometrySelectionManager.h
+++ b/Engine/Plugins/Runtime/MeshModelingToolset/Source/ModelingComponents/Public/Selection/GeometrySelectionManager.h
@@ -170,6 +170,44 @@ public:
 	);
 
 
+	//
+	// Support for more complex selection changes that might (eg) occur over multiple frames, 
+	// or be computed externally. The usage pattern is:
+	//   - verify that CanBeginTrackedSelectionChange() returns true
+	//   - call BeginTrackedSelectionChange(), this opens transacation
+	//      - modify selection here, eg via multiple calls to AccumulateSelectionUpdate_Raycast
+	//   - call EndTrackedSelectionChange() to emit changes and close transaction
+
+	/** @return true if a tracked selection change can be initialized */
+	virtual bool CanBeginTrackedSelectionChange() const;
+
+	/** @return true if an active tracked selection change is in flight */
+	virtual bool IsInTrackedSelectionChange() const { return bInTrackedSelectionChange; }
+
+	/**
+	 * Begin a tracked selection change. CanBeginTrackedSelectionChange() must return true to call this function.
+	 * EndTrackedSelectionChange() must be called to close the selection change.
+	 * @param UpdateConfig how the selection will be modified. Currently this must be consistent across the entire tracked change (ie, only add, only subtract, etc)
+	 * @param bClearOnBegin if true, selection will be initially cleared (eg Replace-style)
+	 */
+	virtual bool BeginTrackedSelectionChange(FGeometrySelectionUpdateConfig UpdateConfig, bool bClearOnBegin);
+
+	/**
+	 * Update the tracked selection change via a single Raycast, using the active UpdateConfig mode passed to BeginTrackedSelectionChange
+	 * @param ResultOut information on any element selection modifications is returned here* 
+	 */
+	virtual void AccumulateSelectionUpdate_Raycast(
+		const FRay3d& WorldRay,
+		FGeometrySelectionUpdateResult& ResultOut
+	);	
+
+	/**
+	 * Close an active tracked selection change. 
+	 * This will emit one or more FChanges for the selection modifications, and then close the open Transaction
+	 */
+	virtual void EndTrackedSelectionChange();
+
+
 	DECLARE_MULTICAST_DELEGATE(FModelingSelectionInteraction_SelectionModified);
 	/**
 	 * OnSelectionModified is broadcast if the selection is modified via the above functions.
@@ -383,6 +421,15 @@ protected:
 	friend class FGeometrySelectionManager_ActiveTargetsChange;
 
 
+	// support for complex selection changes that are driven externally
+	bool bInTrackedSelectionChange = false;
+	FGeometrySelectionUpdateConfig ActiveTrackedUpdateConfig;
+	FGeometrySelection ActiveTrackedSelection;
+	UE::Geometry::FGeometrySelectionDelta InitialTrackedDelta;
+	UE::Geometry::FGeometrySelectionDelta ActiveTrackedDelta;
+	bool bSelectionModifiedDuringTrackedChange = false;
+
+
 	// legacy / to-deprecate/remove
 	
 	// this is to support GetActiveSingleSelectionConverted_Legacy, which will eventually be removed
diff --git a/Engine/Plugins/Runtime/MeshModelingToolset/Source/ModelingComponents/Public/Selection/GeometrySelector.h b/Engine/Plugins/Runtime/MeshModelingToolset/Source/ModelingComponents/Public/Selection/GeometrySelector.h
--- a/Engine/Plugins/Runtime/MeshModelingToolset/Source/ModelingComponents/Public/Selection/GeometrySelector.h
+++ b/Engine/Plugins/Runtime/MeshModelingToolset/Source/ModelingComponents/Public/Selection/GeometrySelector.h
@@ -40,6 +40,7 @@ struct MODELINGCOMPONENTS_API FGeometryIdentifier
 		DynamicMeshComponent = 1,
 		DynamicMesh = 2,
 		BrushComponent = 3,
+		StaticMeshComponent = 4,
 
 		UserDefinedBase = 64
 	};
diff --git a/Engine/Plugins/Runtime/MeshModelingToolset/Source/ModelingComponentsEditorOnly/Private/Selection/StaticMeshSelector.cpp b/Engine/Plugins/Runtime/MeshModelingToolset/Source/ModelingComponentsEditorOnly/Private/Selection/StaticMeshSelector.cpp
new file mode 100644
--- /dev/null
+++ b/Engine/Plugins/Runtime/MeshModelingToolset/Source/ModelingComponentsEditorOnly/Private/Selection/StaticMeshSelector.cpp
@@ -0,0 +1,171 @@
+// Copyright Epic Games, Inc. All Rights Reserved.
+
+
+#include "Selection/StaticMeshSelector.h"
+#include "Selection/GeometrySelector.h"
+#include "Selections/GeometrySelection.h"
+#include "Selections/GeometrySelectionUtil.h"
+
+#include "DynamicMesh/MeshIndexUtil.h"
+#include "DynamicMesh/MeshNormals.h"
+#include "DynamicMesh/MeshTransforms.h"
+#include "Changes/MeshVertexChange.h"
+
+#include "Components/StaticMeshComponent.h"
+#include "Engine/StaticMesh.h"
+#include "MeshDescriptionToDynamicMesh.h"
+#include "DynamicMeshToMeshDescription.h"
+
+#include "RenderingThread.h"
+
+using namespace UE::Geometry;
+
+#define LOCTEXT_NAMESPACE "FStaticMeshSelector"
+
+
+bool FStaticMeshSelector::Initialize(
+	FGeometryIdentifier SourceGeometryIdentifierIn)
+{
+	if (SourceGeometryIdentifierIn.TargetType != FGeometryIdentifier::ETargetType::PrimitiveComponent ||
+		SourceGeometryIdentifierIn.ObjectType != FGeometryIdentifier::EObjectType::StaticMeshComponent)
+	{
+		return false;
+	}
+	StaticMeshComponent = Cast<UStaticMeshComponent>(SourceGeometryIdentifierIn.TargetObject);
+	if (StaticMeshComponent == nullptr)
+	{
+		return false;
+	}
+	StaticMesh = StaticMeshComponent->GetStaticMesh();
+	if (StaticMesh == nullptr)
+	{
+		return false;
+	}
+
+	// To know if the static mesh changed, we will listen to this OnMeshChanged event.
+	// It is currently a very large hammer, though...
+	StaticMesh_OnMeshChangedHandle = StaticMesh->OnMeshChanged.AddLambda([this]()
+	{
+		CopyFromStaticMesh();
+		InvalidateOnMeshChange(FDynamicMeshChangeInfo());
+	});
+
+	// create new transient UDyanmicMesh that will be owned by this Selector
+	LocalTargetMesh.Reset( NewObject<UDynamicMesh>() );
+
+	CopyFromStaticMesh();
+
+	FBaseDynamicMeshSelector::Initialize(SourceGeometryIdentifierIn, LocalTargetMesh.Get(), 
+		[this]() { return IsValid(this->StaticMeshComponent) ? (UE::Geometry::FTransformSRT3d)this->StaticMeshComponent->GetComponentTransform() : FTransformSRT3d::Identity(); });
+
+	return true;
+}
+
+void FStaticMeshSelector::Shutdown()
+{
+	StaticMesh->OnMeshChanged.Remove(StaticMesh_OnMeshChangedHandle);
+	StaticMesh_OnMeshChangedHandle.Reset();
+
+	LocalTargetMesh.Reset();
+	FBaseDynamicMeshSelector::Shutdown();
+}
+
+
+
+IGeometrySelectionTransformer* FStaticMeshSelector::InitializeTransformation(const FGeometrySelection& Selection)
+{
+	check(!ActiveTransformer);
+
+	ActiveTransformer = MakePimpl<FBasicDynamicMeshSelectionTransformer>();
+	ActiveTransformer->Initialize(this);
+	ActiveTransformer->OnEndTransformFunc = [this](IToolsContextTransactionsAPI*) 
+	{ 
+		CommitMeshTransform();
+	};
+	return ActiveTransformer.Get();
+}
+
+void FStaticMeshSelector::ShutdownTransformation(IGeometrySelectionTransformer* Transformer)
+{
+	ActiveTransformer.Reset();
+}
+
+
+void FStaticMeshSelector::CopyFromStaticMesh()
+{
+	int32 UseLODIndex = 0;
+	const FMeshDescription* SourceMesh = StaticMesh->GetMeshDescription(UseLODIndex);
+	if (SourceMesh != nullptr)
+	{
+		FDynamicMesh3 NewMesh;
+		FMeshDescriptionToDynamicMesh Converter;
+		Converter.Convert(SourceMesh, NewMesh, false /*AssetOptions.bRequestTangents*/ );
+
+		LocalTargetMesh->EditMesh([&](FDynamicMesh3& EditMesh)
+		{
+			EditMesh = MoveTemp(NewMesh);
+		});
+	}
+	else
+	{
+		LocalTargetMesh->ResetToCube();
+	}
+}
+
+
+void FStaticMeshSelector::CommitMeshTransform()
+{
+	int32 UseLODIndex = 0;
+	FlushRenderingCommands();
+
+	// emit transaction here??
+
+	// mark as modified
+	StaticMesh->Modify();
+
+	FMeshDescription* MeshDescription = StaticMesh->GetMeshDescription(UseLODIndex);
+	if (MeshDescription == nullptr)
+	{
+		MeshDescription = StaticMesh->CreateMeshDescription(UseLODIndex);
+	}
+	StaticMesh->ModifyMeshDescription(UseLODIndex);
+
+	FConversionToMeshDescriptionOptions ConversionOptions;
+	FDynamicMeshToMeshDescription Converter(ConversionOptions);
+	TargetMesh->ProcessMesh([&](const FDynamicMesh3& ReadMesh)
+	{
+		Converter.Convert(&ReadMesh, *MeshDescription, false /*bCopyTangents*/ );
+	});
+
+	StaticMesh->CommitMeshDescription(UseLODIndex);
+	StaticMesh->PostEditChange();
+}
+
+
+
+
+bool FStaticMeshComponentSelectorFactory::CanBuildForTarget(FGeometryIdentifier TargetIdentifier) const
+{
+	if (TargetIdentifier.TargetType == FGeometryIdentifier::ETargetType::PrimitiveComponent)
+	{
+		UStaticMeshComponent* Component = TargetIdentifier.GetAsComponentType<UStaticMeshComponent>();
+		if (Component != nullptr)
+		{
+			return true;
+		}
+	}
+	return false;
+}
+
+TUniquePtr<IGeometrySelector> FStaticMeshComponentSelectorFactory::BuildForTarget(FGeometryIdentifier TargetIdentifier) const
+{
+	UStaticMeshComponent* Component = TargetIdentifier.GetAsComponentType<UStaticMeshComponent>();
+	check(Component != nullptr);
+
+	TUniquePtr<FStaticMeshSelector> Selector = MakeUnique<FStaticMeshSelector>();
+	bool bInitialized = Selector->Initialize(TargetIdentifier);
+	return (bInitialized) ? MoveTemp(Selector) : TUniquePtr<IGeometrySelector>();
+}
+
+
+#undef LOCTEXT_NAMESPACE 
\ No newline at end of file
diff --git a/Engine/Plugins/Runtime/MeshModelingToolset/Source/ModelingComponentsEditorOnly/Private/Selection/VolumeSelector.cpp b/Engine/Plugins/Runtime/MeshModelingToolset/Source/ModelingComponentsEditorOnly/Private/Selection/VolumeSelector.cpp
new file mode 100644
--- /dev/null
+++ b/Engine/Plugins/Runtime/MeshModelingToolset/Source/ModelingComponentsEditorOnly/Private/Selection/VolumeSelector.cpp
@@ -0,0 +1,136 @@
+// Copyright Epic Games, Inc. All Rights Reserved.
+
+
+#include "Selection/VolumeSelector.h"
+#include "Selection/GeometrySelector.h"
+#include "Selections/GeometrySelection.h"
+#include "Selections/GeometrySelectionUtil.h"
+
+#include "DynamicMesh/MeshIndexUtil.h"
+#include "DynamicMesh/MeshNormals.h"
+#include "DynamicMesh/MeshTransforms.h"
+#include "Changes/MeshVertexChange.h"
+
+#include "GameFramework/Volume.h"
+#include "Components/BrushComponent.h"
+#include "ConversionUtils/VolumeToDynamicMesh.h"
+#include "ConversionUtils/DynamicMeshToVolume.h"
+
+using namespace UE::Geometry;
+
+#define LOCTEXT_NAMESPACE "FVolumeSelector"
+
+
+
+
+bool FVolumeSelector::Initialize(
+	FGeometryIdentifier SourceGeometryIdentifierIn)
+{
+	if (SourceGeometryIdentifierIn.TargetType != FGeometryIdentifier::ETargetType::PrimitiveComponent ||
+		SourceGeometryIdentifierIn.ObjectType != FGeometryIdentifier::EObjectType::BrushComponent)
+	{
+		return false;
+	}
+	BrushComponent = Cast<UBrushComponent>(SourceGeometryIdentifierIn.TargetObject);
+	if (BrushComponent == nullptr)
+	{
+		return false;
+	}
+	ParentVolume = Cast<AVolume>(BrushComponent->GetOwner());
+	if (ParentVolume == nullptr)
+	{
+		return false;
+	}
+	
+	LocalTargetMesh.Reset( NewObject<UDynamicMesh>() );
+
+	LocalTargetMesh->EditMesh([&](FDynamicMesh3& EditMesh)
+	{
+		UE::Conversion::FVolumeToMeshOptions ConvertOptions;
+		UE::Conversion::BrushComponentToDynamicMesh(BrushComponent, EditMesh, ConvertOptions);
+	});
+
+	FBaseDynamicMeshSelector::Initialize(SourceGeometryIdentifierIn, LocalTargetMesh.Get(), 
+		[this]() { return IsValid(this->BrushComponent) ? (UE::Geometry::FTransformSRT3d)this->BrushComponent->GetComponentTransform() : FTransformSRT3d::Identity(); });
+
+	return true;
+}
+
+
+void FVolumeSelector::Shutdown()
+{
+	LocalTargetMesh.Reset();
+	FBaseDynamicMeshSelector::Shutdown();
+}
+
+
+IGeometrySelectionTransformer* FVolumeSelector::InitializeTransformation(const FGeometrySelection& Selection)
+{
+	check(!ActiveTransformer);
+
+	ActiveTransformer = MakePimpl<FBasicDynamicMeshSelectionTransformer>();
+	ActiveTransformer->Initialize(this);
+	ActiveTransformer->OnEndTransformFunc = [this](IToolsContextTransactionsAPI*) 
+	{ 
+		CommitMeshTransform();
+	};
+	return ActiveTransformer.Get();
+}
+
+void FVolumeSelector::ShutdownTransformation(IGeometrySelectionTransformer* Transformer)
+{
+	ActiveTransformer.Reset();
+}
+
+
+void FVolumeSelector::CommitMeshTransform()
+{
+	FTransform SetTransform = ParentVolume->GetActorTransform();
+
+	ParentVolume->Modify();
+	BrushComponent->Modify();
+
+	UE::Conversion::FMeshToVolumeOptions ConvertOptions;
+	FDynamicMesh3 TmpMesh;
+	TargetMesh->ProcessMesh([&](const UE::Geometry::FDynamicMesh3& SourceMesh)
+	{
+		TmpMesh = SourceMesh;
+	});
+
+	// This conversion will (somehow) mess w/ the transform on the Component/Volume? What seems like
+	// is happening is that the mesh is interpreted as being in world-space, and the Actor transform recomputed,
+	// so basically it moves to the Origin. So we need to re-set the previous Transform afterwards.
+	UE::Conversion::DynamicMeshToVolume(TmpMesh, ParentVolume, ConvertOptions);
+	ParentVolume->SetActorTransform(SetTransform);
+
+	ParentVolume->PostEditChange();
+}
+
+
+
+
+bool FBrushComponentSelectorFactory::CanBuildForTarget(FGeometryIdentifier TargetIdentifier) const
+{
+	if (TargetIdentifier.TargetType == FGeometryIdentifier::ETargetType::PrimitiveComponent)
+	{
+		UBrushComponent* Component = TargetIdentifier.GetAsComponentType<UBrushComponent>();
+		if (Component && Cast<AVolume>(Component->GetOwner()) != nullptr)
+		{
+			return true;
+		}
+	}
+	return false;
+}
+
+TUniquePtr<IGeometrySelector> FBrushComponentSelectorFactory::BuildForTarget(FGeometryIdentifier TargetIdentifier) const
+{
+	UBrushComponent* Component = TargetIdentifier.GetAsComponentType<UBrushComponent>();
+	check(Component != nullptr);
+
+	TUniquePtr<FVolumeSelector> Selector = MakeUnique<FVolumeSelector>();
+	bool bInitialized = Selector->Initialize(TargetIdentifier);
+	return (bInitialized) ? MoveTemp(Selector) : TUniquePtr<IGeometrySelector>();
+}
+
+
+#undef LOCTEXT_NAMESPACE 
\ No newline at end of file
diff --git a/Engine/Plugins/Runtime/MeshModelingToolset/Source/ModelingComponentsEditorOnly/Public/Selection/StaticMeshSelector.h b/Engine/Plugins/Runtime/MeshModelingToolset/Source/ModelingComponentsEditorOnly/Public/Selection/StaticMeshSelector.h
new file mode 100644
--- /dev/null
+++ b/Engine/Plugins/Runtime/MeshModelingToolset/Source/ModelingComponentsEditorOnly/Public/Selection/StaticMeshSelector.h
@@ -0,0 +1,68 @@
+// Copyright Epic Games, Inc. All Rights Reserved.
+
+#pragma once
+
+#include "CoreMinimal.h"
+#include "Selection/GeometrySelector.h"
+#include "Selection/DynamicMeshSelector.h"
+
+#include "UObject/StrongObjectPtr.h"
+
+
+class UDynamicMesh;
+class UStaticMesh;
+class UStaticMeshComponent;
+
+
+class MODELINGCOMPONENTSEDITORONLY_API FStaticMeshSelector : public FBaseDynamicMeshSelector
+{
+public:
+	using FBaseDynamicMeshSelector::Initialize;
+
+	virtual bool Initialize(
+		FGeometryIdentifier SourceGeometryIdentifier);
+
+	//
+	// IGeometrySelector API implementation
+	//
+
+	virtual void Shutdown() override;
+	// disable sleep on StaticMesh Selector until we can properly track changes...
+	virtual bool SupportsSleep() const override { return false; }
+	//virtual bool Sleep();
+	//virtual bool Restore();
+
+
+
+	virtual IGeometrySelectionTransformer* InitializeTransformation(const FGeometrySelection& Selection) override;
+	virtual void ShutdownTransformation(IGeometrySelectionTransformer* Transformer) override;
+
+protected:
+	UStaticMeshComponent* StaticMeshComponent = nullptr;
+	UStaticMesh* StaticMesh = nullptr;
+	FDelegateHandle StaticMesh_OnMeshChangedHandle;
+
+	// TODO: this is not a great design, it would be better if something external could own this mesh...
+	TStrongObjectPtr<UDynamicMesh> LocalTargetMesh;
+	void CopyFromStaticMesh();
+
+	TPimplPtr<FBasicDynamicMeshSelectionTransformer> ActiveTransformer;
+	void CommitMeshTransform();
+};
+
+
+
+/**
+ * FStaticMeshComponentSelectorFactory constructs FStaticMeshSelector instances 
+ * for UBrushComponents
+ */
+class MODELINGCOMPONENTSEDITORONLY_API FStaticMeshComponentSelectorFactory : public IGeometrySelectorFactory
+{
+public:
+	virtual bool CanBuildForTarget(FGeometryIdentifier TargetIdentifier) const;
+	virtual TUniquePtr<IGeometrySelector> BuildForTarget(FGeometryIdentifier TargetIdentifier) const;
+};
+
+
+
+
diff --git a/Engine/Plugins/Runtime/MeshModelingToolset/Source/ModelingComponentsEditorOnly/Public/Selection/VolumeSelector.h b/Engine/Plugins/Runtime/MeshModelingToolset/Source/ModelingComponentsEditorOnly/Public/Selection/VolumeSelector.h
new file mode 100644
--- /dev/null
+++ b/Engine/Plugins/Runtime/MeshModelingToolset/Source/ModelingComponentsEditorOnly/Public/Selection/VolumeSelector.h
@@ -0,0 +1,66 @@
+// Copyright Epic Games, Inc. All Rights Reserved.
+
+#pragma once
+
+#include "CoreMinimal.h"
+#include "Selection/GeometrySelector.h"
+#include "Selection/DynamicMeshSelector.h"
+
+#include "UObject/StrongObjectPtr.h"
+
+
+class UDynamicMesh;
+class AVolume;
+class UBrushComponent;
+
+
+class MODELINGCOMPONENTSEDITORONLY_API FVolumeSelector : public FBaseDynamicMeshSelector
+{
+public:
+	using FBaseDynamicMeshSelector::Initialize;
+
+	virtual bool Initialize(
+		FGeometryIdentifier SourceGeometryIdentifier);
+
+	//
+	// IGeometrySelector API implementation
+	//
+
+	virtual void Shutdown() override;
+	// disable sleep on VolumeSelector until we can properly track changes...
+	virtual bool SupportsSleep() const override { return false; }
+	//virtual bool Sleep();
+	//virtual bool Restore();
+
+
+	virtual IGeometrySelectionTransformer* InitializeTransformation(const FGeometrySelection& Selection) override;
+	virtual void ShutdownTransformation(IGeometrySelectionTransformer* Transformer) override;
+
+protected:
+	AVolume* ParentVolume = nullptr;
+	UBrushComponent* BrushComponent = nullptr;
+
+	// TODO: this is not a great design, it would be better if something external could own this mesh...
+	TStrongObjectPtr<UDynamicMesh> LocalTargetMesh;
+
+	TPimplPtr<FBasicDynamicMeshSelectionTransformer> ActiveTransformer;
+	void CommitMeshTransform();
+};
+
+
+
+
+/**
+ * FVolumeComponentSelectorFactory constructs FVolumeSelector instances 
+ * for UBrushComponents
+ */
+class MODELINGCOMPONENTSEDITORONLY_API FBrushComponentSelectorFactory : public IGeometrySelectorFactory
+{
+public:
+	virtual bool CanBuildForTarget(FGeometryIdentifier TargetIdentifier) const;
+	virtual TUniquePtr<IGeometrySelector> BuildForTarget(FGeometryIdentifier TargetIdentifier) const;
+};
+
+
+
+
diff --git a/Engine/Source/Runtime/InteractiveToolsFramework/Private/InputBehaviorSet.cpp b/Engine/Source/Runtime/InteractiveToolsFramework/Private/InputBehaviorSet.cpp
--- a/Engine/Source/Runtime/InteractiveToolsFramework/Private/InputBehaviorSet.cpp
+++ b/Engine/Source/Runtime/InteractiveToolsFramework/Private/InputBehaviorSet.cpp
@@ -101,15 +101,15 @@ void UInputBehaviorSet::RemoveAll()
 
 void UInputBehaviorSet::CollectWantsCapture(const FInputDeviceState& input, TArray<FInputCaptureRequest>& result)
 {
-	for (auto b : Behaviors)
+	for (FBehaviorInfo& BehaviorInfo : Behaviors)
 	{
 		// only call WantsCapture if the Behavior supports the current input device
-		if (SupportsInputType(b.Behavior, input)) 
+		if (SupportsInputType(BehaviorInfo.Behavior, input)) 
 		{
-			FInputCaptureRequest request = b.Behavior->WantsCapture(input);
+			FInputCaptureRequest request = BehaviorInfo.Behavior->WantsCapture(input);
 			if ( request.Type != EInputCaptureRequestType::Ignore ) 
 			{
-				request.Owner = b.Source;
+				request.Owner = BehaviorInfo.Source;
 				result.Add(request);
 			}
 		}
@@ -120,15 +120,15 @@ void UInputBehaviorSet::CollectWantsCapture(const FInputDeviceState& input, TArr
 
 void UInputBehaviorSet::CollectWantsHoverCapture(const FInputDeviceState& input, TArray<FInputCaptureRequest>& result)
 {
-	for (auto b : Behaviors)
+	for (FBehaviorInfo& BehaviorInfo : Behaviors)
 	{
 		// only call WantsCapture if the Behavior supports the current input device
-		if (b.Behavior->WantsHoverEvents() && SupportsInputType(b.Behavior, input))
+		if (BehaviorInfo.Behavior->WantsHoverEvents() && SupportsInputType(BehaviorInfo.Behavior, input))
 		{
-			FInputCaptureRequest request = b.Behavior->WantsHoverCapture(input);
+			FInputCaptureRequest request = BehaviorInfo.Behavior->WantsHoverCapture(input);
 			if (request.Type != EInputCaptureRequestType::Ignore)
 			{
-				request.Owner = b.Source;
+				request.Owner = BehaviorInfo.Source;
 				result.Add(request);
 			}
 		}
