Geometry - Modeling Components:  conversion from SkinnedMeshComponent to DynamicMesh.
           GeometryScriptintCore : add the conversion to CopyMeshFromComponent

diff --git a/Engine/Plugins/Experimental/GeometryScripting/Source/GeometryScriptingCore/Private/SceneUtilityFunctions.cpp b/Engine/Plugins/Experimental/GeometryScripting/Source/GeometryScriptingCore/Private/SceneUtilityFunctions.cpp
--- a/Engine/Plugins/Experimental/GeometryScripting/Source/GeometryScriptingCore/Private/SceneUtilityFunctions.cpp
+++ b/Engine/Plugins/Experimental/GeometryScripting/Source/GeometryScriptingCore/Private/SceneUtilityFunctions.cpp
@@ -12,9 +12,11 @@
 #include "Engine/SkeletalMesh.h"
 #include "Components/StaticMeshComponent.h"
 #include "Components/SkeletalMeshComponent.h"
+#include "Components/SkinnedMeshComponent.h"
 #include "Components/BrushComponent.h"
 #include "Components/DynamicMeshComponent.h"
 #include "ConversionUtils/VolumeToDynamicMesh.h"
+#include "ConversionUtils/SkinnedMeshToDynamicMesh.h"
 #include "Physics/ComponentCollisionUtil.h"
 
 #include UE_INLINE_GENERATED_CPP_BY_NAME(SceneUtilityFunctions)
@@ -40,7 +42,35 @@ UDynamicMesh* UGeometryScriptLibrary_SceneUtilityFunctions::CopyMeshFromComponen
 {
 	Outcome = EGeometryScriptOutcomePins::Failure;
 
-	if (UStaticMeshComponent* StaticMeshComponent = Cast<UStaticMeshComponent>(Component))
+	if (USkinnedMeshComponent* SkinnedMeshComponent = Cast<USkinnedMeshComponent>(Component))
+	{
+		LocalToWorld = SkinnedMeshComponent->GetComponentTransform();
+
+		const int32 NumLODs = SkinnedMeshComponent->GetNumLODs();
+		const int32 RequestedLOD = Options.RequestedLOD.LODIndex;
+		if (RequestedLOD < 0 || RequestedLOD > NumLODs - 1)
+		{
+			UE::Geometry::AppendError(Debug, EGeometryScriptErrorType::InvalidInputs, LOCTEXT("CopyMeshFromComponent_MissingSkinnedMeshComponentLOD", "CopyMeshFromComponent: SkinnedMeshComponent requested LOD does not exist"));
+		}
+		else
+		{ 
+			USkinnedAsset* SkinnedAsset = SkinnedMeshComponent->GetSkinnedAsset();
+			if (SkinnedAsset)
+			{
+				FDynamicMesh3 NewMesh;
+				UE::Conversion::SkinnedMeshComponentToDynamicMesh(*SkinnedMeshComponent, NewMesh, RequestedLOD, Options.bWantTangents);
+				NewMesh.DiscardTriangleGroups();
+				ToDynamicMesh->SetMesh(MoveTemp(NewMesh));
+				Outcome = EGeometryScriptOutcomePins::Success;
+			}
+			else
+			{
+				UE::Geometry::AppendError(Debug, EGeometryScriptErrorType::InvalidInputs, LOCTEXT("CopyMeshFromComponent_MissingSkinnedAsset", "CopyMeshFromComponent: SkinnedMeshComponent has a null SkinnedAsset"));
+			}
+		}
+
+	}
+	else if (UStaticMeshComponent* StaticMeshComponent = Cast<UStaticMeshComponent>(Component))
 	{
 		LocalToWorld = StaticMeshComponent->GetComponentTransform();
 		UStaticMesh* StaticMesh = StaticMeshComponent->GetStaticMesh();
diff --git a/Engine/Plugins/Runtime/MeshModelingToolset/Source/ModelingComponents/Private/ConversionUtils/SkinnedMeshToDynamicMesh.cpp b/Engine/Plugins/Runtime/MeshModelingToolset/Source/ModelingComponents/Private/ConversionUtils/SkinnedMeshToDynamicMesh.cpp
new file mode 100644
--- /dev/null
+++ b/Engine/Plugins/Runtime/MeshModelingToolset/Source/ModelingComponents/Private/ConversionUtils/SkinnedMeshToDynamicMesh.cpp
@@ -0,0 +1,352 @@
+// Copyright Epic Games, Inc. All Rights Reserved.
+
+#include "ConversionUtils/SkinnedMeshToDynamicMesh.h"
+
+
+#include "Components/SkinnedMeshComponent.h"
+#include "CoreMinimal.h"
+#include "DynamicMesh/DynamicMesh3.h"
+#include "DynamicMesh/DynamicMeshAttributeSet.h"
+#include "Engine/SkinnedAssetCommon.h"
+#include "Engine/SkinnedAsset.h"
+#include "SkeletalRenderPublic.h"
+#include "MeshConversion/Public/ToDynamicMesh.h"
+#include "Rendering/SkeletalMeshRenderData.h"
+
+
+namespace UE 
+{
+namespace Conversion
+{
+
+
+// Wrapper that allows us to use ToDynamicMesh converter
+class FSkinnedComponentWrapper
+{
+public:
+
+	typedef int32 TriIDType;
+	typedef int32 VertIDType;
+	typedef int32 WedgeIDType;
+
+	typedef int32 UVIDType;
+	typedef int32 NormalIDType;
+	typedef int32 ColorIDType;
+
+
+	FSkinnedComponentWrapper(USkinnedMeshComponent& SkinnedMeshComponent, int32 RequestedLOD, bool UseDisabledSections, bool bWantTangents)
+	: Component(&SkinnedMeshComponent)
+	, LOD(RequestedLOD)
+	, bUseDisabledSections(UseDisabledSections)
+	, bHasTangents(bWantTangents)
+	, Asset(SkinnedMeshComponent.GetSkinnedAsset())
+	, SrcLODInfo(SkinnedMeshComponent.GetSkinnedAsset()->GetLODInfo(RequestedLOD))
+	, SkeletalMeshRenderData(SkinnedMeshComponent.MeshObject->GetSkeletalMeshRenderData())
+	, IndexBuffer(SkinnedMeshComponent.MeshObject->GetSkeletalMeshRenderData().LODRenderData[RequestedLOD].MultiSizeIndexContainer.GetIndexBuffer())
+
+	{
+
+		// This step does the actual skinning (and isn't as const as the api suggests..)
+		Component->GetCPUSkinnedVertices(SkinnedVertices, RequestedLOD);
+		const FSkeletalMeshLODRenderData& LODData  = SkeletalMeshRenderData.LODRenderData[RequestedLOD];
+		
+
+		const int32 NumSections = LODData.RenderSections.Num();
+		
+		auto SkipSection = [&](const FSkelMeshRenderSection& SkelMeshSection)
+							{
+								return ( !bUseDisabledSections && !SkinnedMeshComponent.IsMaterialSectionShown(SkelMeshSection.MaterialIndex, RequestedLOD));
+							};
+
+		// pre-compute number of tris & verts
+		int32 NumTris = 0;
+		int32 NumVerts = 0;		
+		{
+					
+			for (int32 SectionIndex = 0; SectionIndex < NumSections; ++SectionIndex)
+			{
+				const FSkelMeshRenderSection& SkelMeshSection = LODData.RenderSections[SectionIndex];
+
+				if (SkipSection(SkelMeshSection))
+				{
+					continue;
+				}
+				NumTris += SkelMeshSection.NumTriangles;
+				NumVerts += SkelMeshSection.NumVertices;
+			}
+		}
+		
+		// construct a list of all valid VertIDs for this mesh.
+		VertIDs.Reserve(NumVerts);
+		for (const  FSkelMeshRenderSection& Section : LODData.RenderSections)
+		{
+			if (SkipSection(Section))
+			{
+				continue;
+			}
+
+			const int32 BaseVertexIndex = static_cast<int32>( Section.BaseVertexIndex);
+			const int32 NumSectionVtx = static_cast<int32>( Section.NumVertices );
+			for (int32 VtxIndex = BaseVertexIndex;  VtxIndex < NumSectionVtx + BaseVertexIndex;  ++VtxIndex)
+			{
+				VertIDs.Add(VtxIndex);
+			}
+		}
+
+		// construct list of all valid TriIDs for this mesh.  
+		TriIDs.Reserve(NumTris);
+		TriIDtoSectionID.Init(-1, NumTris);
+		for (int32 s = 0; s < NumSections; ++s)
+		{
+			const  FSkelMeshRenderSection& Section = LODData.RenderSections[s];
+			if (SkipSection(Section))
+			{
+				continue; 
+			}
+			const int32 BaseIndex = Section.BaseIndex;
+			const int32 BaseTriIdx = BaseIndex / 3;
+			const int32 NumSectionTris = static_cast<int32>(Section.NumTriangles);
+			for (int32 j = BaseTriIdx; j < BaseTriIdx + NumSectionTris; ++j)
+			{
+				TriIDs.Add(j);
+				TriIDtoSectionID[j] = s;
+			}
+			
+		}
+	}
+
+	int32 NumTris() const
+	{
+		return TriIDs.Num();
+	}
+	int32 NumVerts() const
+	{
+		return VertIDs.Num();
+	}
+
+	int32 NumUVLayers() const
+	{
+		const FSkeletalMeshLODRenderData& LODData  = SkeletalMeshRenderData.LODRenderData[LOD];
+		return static_cast<int32>( LODData.StaticVertexBuffers.StaticMeshVertexBuffer.GetNumTexCoords() );
+	}
+
+	// --"Vertex Buffer" info
+	const TArray<VertIDType>& GetVertIDs() const
+	{
+		return VertIDs;
+	}
+
+	FVector3d GetPosition(VertIDType VtxID) const
+	{
+		return static_cast<FVector3d>(SkinnedVertices[VtxID].Position);
+	}
+
+	// --"Index Buffer" info
+	const TArray<TriIDType>& GetTriIDs() const
+	{
+		return TriIDs;
+	}
+
+	bool GetTri(TriIDType TriID, VertIDType& VID0, VertIDType& VID1, VertIDType& VID2) const
+	{
+		int32 Offset = TriID * 3;
+
+		// index buffer for this LOD.
+		const FRawStaticIndexBuffer16or32Interface& Indices = *IndexBuffer;
+		auto GetVID = [&Indices](int32 i){ return static_cast<VertIDType>(Indices.Get(i));};
+
+		VID0 = GetVID( Offset    );
+		VID1 = GetVID( Offset + 1);
+		VID2 = GetVID( Offset + 2);
+
+		// maybe check IndexBuffer.Num() > Offset+2 ? 
+		return true;
+	}
+
+
+	bool HasNormals() const { return true; }
+	bool HasTangents() const { return bHasTangents; }
+	bool HasBiTangents() const { return bHasTangents; }
+	
+	bool HasColors() const
+	{
+		const FSkeletalMeshLODRenderData& LODData  = SkeletalMeshRenderData.LODRenderData[LOD];
+		return  (LODData.StaticVertexBuffers.ColorVertexBuffer.IsInitialized() && LODData.StaticVertexBuffers.ColorVertexBuffer.GetNumVertices() > 0);
+	}
+
+	//-- Access to per-wedge attributes --//
+	void GetWedgeIDs(const TriIDType& TriID, WedgeIDType& WID0, WedgeIDType& WID1, WedgeIDType& WID2) const
+	{
+		int32 Offset = 3 * TriID;
+		WID0 = Offset;
+		WID1 = Offset + 1;
+		WID2 = Offset + 2;
+	}
+
+	FVector2f GetWedgeUV(int32 UVLayerIndex, WedgeIDType WID) const
+	{
+		const FFinalSkinVertex& SkinnedVertex = GetWedgeVertexInstance(WID);
+		return static_cast<FVector2f>( SkinnedVertex.TextureCoordinates[UVLayerIndex] );
+	}
+
+	FVector3f GetWedgeNormal(WedgeIDType WID) const
+	{
+		const FFinalSkinVertex& SkinnedVertex = GetWedgeVertexInstance(WID);
+		return SkinnedVertex.TangentZ.ToFVector3f();
+	}
+
+	FVector3f GetWedgeTangent(WedgeIDType WID) const
+	{
+		const FFinalSkinVertex& SkinnedVertex = GetWedgeVertexInstance(WID);
+		return SkinnedVertex.TangentX.ToFVector3f();
+	}
+
+	FVector3f GetWedgeBiTangent(WedgeIDType WID) const
+	{
+		const FFinalSkinVertex& SkinnedVertex = GetWedgeVertexInstance(WID);
+		const FVector3f TangentX = SkinnedVertex.TangentX.ToFVector3f();
+		const FVector3f TangentZ = SkinnedVertex.TangentZ.ToFVector3f();
+		const float OrientationSign = SkinnedVertex.TangentZ.ToFVector4f().W;
+		const FVector3f TangentY = (TangentZ.Cross(TangentX)).GetSafeNormal() * OrientationSign;
+		return TangentY;
+	}
+
+	FVector4f GetWedgeColor(WedgeIDType WID) const
+	{
+		FLinearColor LinearColor = FLinearColor::White;
+		if (HasColors())
+		{ 
+			const int32 VID = GetVertID(WID);
+			const FSkeletalMeshLODRenderData& LODData  = SkeletalMeshRenderData.LODRenderData[LOD];
+			const FColor& Color = LODData.StaticVertexBuffers.ColorVertexBuffer.VertexColor(VID);
+			LinearColor =  Color.ReinterpretAsLinear();
+		}
+		return FVector4f(LinearColor.R, LinearColor.G, LinearColor.B, LinearColor.A);
+	}
+
+
+	int32 GetMaterialIndex(TriIDType TriID) const
+	{
+		int32 SectionID = TriIDtoSectionID[TriID];
+		check(SectionID > -1);
+
+		const FSkeletalMeshLODRenderData& LODData  = SkeletalMeshRenderData.LODRenderData[LOD];
+		const  FSkelMeshRenderSection& Section = LODData.RenderSections[SectionID];
+		
+		int32 MaterialIndex = Section.MaterialIndex;
+		// use the remapping of material indices if there is a valid value
+		if (SrcLODInfo->LODMaterialMap.IsValidIndex(SectionID) && SrcLODInfo->LODMaterialMap[SectionID] != INDEX_NONE)
+		{
+			MaterialIndex = FMath::Clamp<int32>(SrcLODInfo->LODMaterialMap[SectionID], 0, Asset->GetMaterials().Num() - 1);
+		}
+
+		return MaterialIndex;
+	}
+
+
+
+	//-- null implementation of shared attributes: Skeletal mesh model doesn't use these --//
+	const TArray<int32>& GetUVIDs(int32 LayerID) const { return EmptyArray; }
+	FVector2f GetUV(int32 LayerID, UVIDType UVID) const { check(0); return FVector2f(); }
+	bool GetUVTri(int32 LayerID, const TriIDType&, UVIDType& ID0, UVIDType& ID1, UVIDType& ID2) const { ID0 = ID1 = ID2 = UVIDType(-1); return false; }
+
+	const TArray<int32>& GetNormalIDs() const { return EmptyArray; }
+	FVector3f GetNormal(NormalIDType ID) const { check(0); return FVector3f(); }
+	bool GetNormalTri(const TriIDType&, NormalIDType& ID0, NormalIDType& ID1, NormalIDType& ID2) const { ID0 = ID1 = ID2 = NormalIDType(-1); return false; }
+
+	const TArray<int32>& GetTangentIDs() const { return EmptyArray; }
+	FVector3f GetTangent(NormalIDType ID) const { check(0); return FVector3f(); }
+	bool GetTangentTri(const TriIDType&, NormalIDType& ID0, NormalIDType& ID1, NormalIDType& ID2) const { ID0 = ID1 = ID2 = NormalIDType(-1); return false; }
+
+
+	const TArray<int32>& GetBiTangentIDs() const { return EmptyArray; }
+	FVector3f GetBiTangent(NormalIDType ID) const { check(0); return FVector3f(); }
+	bool GetBiTangentTri(const TriIDType&, NormalIDType& ID0, NormalIDType& ID1, NormalIDType& ID2) const { ID0 = ID1 = ID2 = NormalIDType(-1); return false; }
+
+	const TArray<int32>& GetColorIDs() const { return EmptyArray; }
+	FVector4f GetColor(ColorIDType ID) const { check(0); return FVector4f(); }
+	bool GetColorTri(const TriIDType&, ColorIDType& ID0, ColorIDType& ID1, ColorIDType& ID2) const { ID0 = ID1 = ID2 = ColorIDType(-1); return false; }
+
+
+private:
+
+	inline const FFinalSkinVertex& GetWedgeVertexInstance(WedgeIDType WID) const
+	{
+		int32 VertID = GetVertID(WID);
+		return SkinnedVertices[VertID];
+	}
+
+	// Convert the WedgeID the ID of the corresponding position vertex.
+	inline const VertIDType GetVertID(int32 WedgeID) const
+	{
+		const FRawStaticIndexBuffer16or32Interface& Indices = *IndexBuffer;
+		return static_cast<VertIDType>(Indices.Get(WedgeID));
+	}
+
+private:
+
+	TArray<FFinalSkinVertex>  SkinnedVertices;
+	USkinnedMeshComponent* Component;
+	int32                  LOD;
+	bool                   bUseDisabledSections;
+	bool                   bHasTangents;
+
+	USkinnedAsset* Asset;
+	const FSkeletalMeshLODInfo* SrcLODInfo;
+	const FSkeletalMeshRenderData& SkeletalMeshRenderData;
+	const FRawStaticIndexBuffer16or32Interface* IndexBuffer;
+	
+
+	TArray<int32> VertIDs;
+	TArray<int32> TriIDs;
+
+	TArray<int32> TriIDtoSectionID;
+	TArray<int32> EmptyArray;
+};
+
+
+
+
+
+void SkinnedMeshComponentToDynamicMesh(USkinnedMeshComponent& SkinnedMeshComponent, Geometry::FDynamicMesh3& MeshOut, int32 RequestedLOD, bool bWantTangents)
+{
+
+	MeshOut.Clear();
+
+	const int32 NumLODs = SkinnedMeshComponent.GetNumLODs();
+	const bool bIsValidSkinnedMeshComponent = SkinnedMeshComponent.MeshObject && SkinnedMeshComponent.IsVisible();
+
+	if (!bIsValidSkinnedMeshComponent)
+	{
+		return;
+	}
+
+	if (RequestedLOD < 0 || RequestedLOD  > NumLODs - 1)
+	{
+		return;
+	}
+
+	if (!SkinnedMeshComponent.GetSkinnedAsset())
+	{
+		return;
+	}
+
+
+	constexpr bool bUseDisabledSections = false;
+	FSkinnedComponentWrapper SkinnedComponentWrapper(SkinnedMeshComponent, RequestedLOD, bUseDisabledSections, bWantTangents);
+
+	// do the actual conversion.
+	{
+		constexpr bool bCopyTangents = true;
+		auto TriToGroupID    = [&SkinnedComponentWrapper](const int32& SrcTriID)->int32 { return (SkinnedComponentWrapper.GetMaterialIndex(SrcTriID) + 1); };
+		auto TriToMaterialID = [&SkinnedComponentWrapper](const int32& SrcTriID)->int32 { return SkinnedComponentWrapper.GetMaterialIndex(SrcTriID); };
+		UE::Geometry::TToDynamicMesh<FSkinnedComponentWrapper> SkinnedComponentConverter;
+
+		SkinnedComponentConverter.Convert(MeshOut, SkinnedComponentWrapper, TriToGroupID, TriToMaterialID, bCopyTangents);
+	}
+
+}
+
+} // end namespace Conversion
+} // end namespace UE
+
diff --git a/Engine/Plugins/Runtime/MeshModelingToolset/Source/ModelingComponents/Public/ConversionUtils/SkinnedMeshToDynamicMesh.h b/Engine/Plugins/Runtime/MeshModelingToolset/Source/ModelingComponents/Public/ConversionUtils/SkinnedMeshToDynamicMesh.h
new file mode 100644
--- /dev/null
+++ b/Engine/Plugins/Runtime/MeshModelingToolset/Source/ModelingComponents/Public/ConversionUtils/SkinnedMeshToDynamicMesh.h
@@ -0,0 +1,30 @@
+// Copyright Epic Games, Inc. All Rights Reserved.
+
+#pragma once
+
+#include "CoreMinimal.h"
+
+class USkinnedMeshComponent;
+
+namespace UE {
+namespace Geometry{ class FDynamicMesh3; };
+
+namespace Conversion {
+
+	/**
+	* Converts a SkinnedMeshComponent to a DynamicMesh.  In doing so, it will create attribute layers for the vertex normals, tangents, UVs and colors defined on the SkinnedMeshComponent, but 
+	* the resulting mesh will have mesh (edges) seams along any attribute seam.  You may want to weld the dynamic mesh after this conversion.
+	* Per-triangle Material IDs are transfered directly from the SkinnedMeshComponent, and per-triangle Group IDs are assigned to these regions ( specifically the Group ID for a triangle will be the 
+	* Material ID + 1).
+	* 
+	* If the SkinnedComponent is hidden then the result mesh will be empty, likewise material sections of the SkinnedMeshCompnent that are set to hidden will be ignored during conversion.
+	*
+	* @param SkinnedMeshComponent - the input component. Note, the constiness.  Internally this calls some const methods that really aren't const, and some non-const functions that should be..
+	* @param MeshOut       - the result mesh.  Note, this function does not append the new mesh - the MeshOut is cleared by the function before populating it.
+	* @param RequestedLOD  - the LOD to be converted, if a non-existent LOD is requested the result mesh will be empty.
+	* @param bWantTangents - controls if tangents are transfered
+	*/
+	void MODELINGCOMPONENTS_API SkinnedMeshComponentToDynamicMesh(USkinnedMeshComponent& SkinnedMeshComponent, Geometry::FDynamicMesh3& MeshOut, int32 RequestedLOD,  bool bWantTangents);
+
+} // end namespace Geometry
+} // end namespace UE
\ No newline at end of file
