Geometry - SkinnedMeshToDynamicMesh - fix linux editor build error
#rnx

diff --git a/Engine/Plugins/Runtime/MeshModelingToolset/Source/ModelingComponents/Private/ConversionUtils/SkinnedMeshToDynamicMesh.cpp b/Engine/Plugins/Runtime/MeshModelingToolset/Source/ModelingComponents/Private/ConversionUtils/SkinnedMeshToDynamicMesh.cpp
--- a/Engine/Plugins/Runtime/MeshModelingToolset/Source/ModelingComponents/Private/ConversionUtils/SkinnedMeshToDynamicMesh.cpp
+++ b/Engine/Plugins/Runtime/MeshModelingToolset/Source/ModelingComponents/Private/ConversionUtils/SkinnedMeshToDynamicMesh.cpp
@@ -267,6 +267,11 @@ public:
 	FVector4f GetColor(ColorIDType ID) const { check(0); return FVector4f(); }
 	bool GetColorTri(const TriIDType&, ColorIDType& ID0, ColorIDType& ID1, ColorIDType& ID2) const { ID0 = ID1 = ID2 = ColorIDType(-1); return false; }
 
+	int32 NumWeightMapLayers() const { return 0; }
+	float GetVertexWeight(int32 WeightMapIndex, int32 SrcVertID) const { return 0; }
+	FName GetWeightMapName(int32 WeightMapIndex) const { return FName(); }
+
+
 
 private:
 
