Merge Tool: Remove useless 'Reset to Default' button for the 'Write To' field in UBaseCreateFromSelectedHandleSourceProperties.

This isn't specific to the merge tool, it is caused by the base create from selected tool from which the merge tool inherits.

diff --git a/Engine/Plugins/Runtime/MeshModelingToolset/Source/ModelingComponents/Public/BaseTools/BaseCreateFromSelectedTool.h b/Engine/Plugins/Runtime/MeshModelingToolset/Source/ModelingComponents/Public/BaseTools/BaseCreateFromSelectedTool.h
--- a/Engine/Plugins/Runtime/MeshModelingToolset/Source/ModelingComponents/Public/BaseTools/BaseCreateFromSelectedTool.h
+++ b/Engine/Plugins/Runtime/MeshModelingToolset/Source/ModelingComponents/Public/BaseTools/BaseCreateFromSelectedTool.h
@@ -54,7 +54,7 @@ class MODELINGCOMPONENTS_API UBaseCreateFromSelectedHandleSourceProperties : pub
 	GENERATED_BODY()
 public:
 	/** Defines the object the tool output is written to. */
-	UPROPERTY(EditAnywhere, Category = OutputObject, meta = (DisplayName = "Write To"))
+	UPROPERTY(EditAnywhere, Category = OutputObject, meta = (DisplayName = "Write To", NoResetToDefault))
 	EBaseCreateFromSelectedTargetType OutputWriteTo = EBaseCreateFromSelectedTargetType::NewObject;
 
 	/** Base name of the newly generated object to which the output is written to. */
