Modeling Mode Plane Visualization: Expose some options in the modeling mode preferences allowing the user to determine how planes in various tools should appear.

Created UModelingComponentsEditorSettings, a subclass of UDeveloperSettings.

diff --git a/Engine/Plugins/Runtime/MeshModelingToolset/Source/ModelingComponents/Private/Drawing/MeshDebugDrawing.cpp b/Engine/Plugins/Runtime/MeshModelingToolset/Source/ModelingComponents/Private/Drawing/MeshDebugDrawing.cpp
--- a/Engine/Plugins/Runtime/MeshModelingToolset/Source/ModelingComponents/Private/Drawing/MeshDebugDrawing.cpp
+++ b/Engine/Plugins/Runtime/MeshModelingToolset/Source/ModelingComponents/Private/Drawing/MeshDebugDrawing.cpp
@@ -157,7 +157,7 @@ void MeshDebugDraw::DrawHierarchicalGrid(
 		FMath::GridSnap(LocalFrame.Origin.X, GridScales[0]),
 		FMath::GridSnap(LocalFrame.Origin.Y, GridScales[0]),
 		LocalFrame.Origin.Z);
-	UE::Geometry::FFrame3d GridFrame(GridOrigin);
+	UE::Geometry::FFrame3d GridFrame(GridOrigin, LocalFrame.Rotation);
 
 	// Draw each level of the grid
 	for (int32 Level = 0; Level < Levels; ++Level)
diff --git a/Engine/Plugins/Runtime/MeshModelingToolset/Source/ModelingComponents/Private/Mechanics/ConstructionPlaneMechanic.cpp b/Engine/Plugins/Runtime/MeshModelingToolset/Source/ModelingComponents/Private/Mechanics/ConstructionPlaneMechanic.cpp
--- a/Engine/Plugins/Runtime/MeshModelingToolset/Source/ModelingComponents/Private/Mechanics/ConstructionPlaneMechanic.cpp
+++ b/Engine/Plugins/Runtime/MeshModelingToolset/Source/ModelingComponents/Private/Mechanics/ConstructionPlaneMechanic.cpp
@@ -8,6 +8,7 @@
 #include "Drawing/MeshDebugDrawing.h"
 #include "ToolSceneQueriesUtil.h"
 #include "SceneQueries/SceneSnappingManager.h"
+#include "ModelingComponentsSettings.h"
 
 #include UE_INLINE_GENERATED_CPP_BY_NAME(ConstructionPlaneMechanic)
 
@@ -147,18 +148,53 @@ void UConstructionPlaneMechanic::Tick(float DeltaTime)
  
 void UConstructionPlaneMechanic::Render(IToolsContextRenderAPI* RenderAPI)
 {
-	if (bShowGrid)
+	const UModelingComponentsEditorSettings* Settings = GetMutableDefault<UModelingComponentsEditorSettings>();
+	if (bShowGrid && Settings != nullptr)
 	{
-		FViewCameraState CameraState = RenderAPI->GetCameraState();
-		float PDIScale = CameraState.GetPDIScalingFactor();
-
-		int32 NumGridLines = 10;
+		const FViewCameraState CameraState = RenderAPI->GetCameraState();
+		const float PDIScale = CameraState.GetPDIScalingFactor();
+		
+		const FFrame3d DrawFrame(PlaneTransformProxy->GetTransform().GetLocation(), Plane.Rotation);
+		const int32 NumGridLines = Settings->NumGridLines;
+		const float GridThickness = 0.75f * PDIScale;
+		constexpr FColor GridColor(128, 128, 128, 32);
+		static constexpr bool bDepthTested = false;
 		FPrimitiveDrawInterface* PDI = RenderAPI->GetPrimitiveDrawInterface();
-		FColor GridColor(128, 128, 128, 32);
-		float GridThickness = 0.75f*PDIScale;
+		const FTransform DrawTransform = FTransform::Identity;
 
-		FFrame3d DrawFrame(Plane);
-		MeshDebugDraw::DrawSimpleFixedScreenAreaGrid(CameraState, DrawFrame, NumGridLines, 45.0, GridThickness, GridColor, false, PDI, FTransform::Identity);
+		switch (Settings->GridMode)
+		{
+		case EModelingComponentsPlaneVisualizationMode::SimpleGrid:
+		{
+			const double GridSpacing = Settings->GridSpacing;
+			MeshDebugDraw::DrawSimpleGrid(DrawFrame, NumGridLines, GridSpacing, GridThickness, GridColor, bDepthTested, PDI, DrawTransform);
+			break;
+		}
+			
+		case EModelingComponentsPlaneVisualizationMode::HierarchicalGrid:
+		{
+			const double ZoomFactor = Distance(CameraState.Position, DrawFrame.Origin);
+			const FVector WorldMaxBounds(ZoomFactor, ZoomFactor, 0);
+			const FVector WorldMinBounds(-WorldMaxBounds);
+			constexpr int32 Levels = 2;
+
+			TArray<FColor> Colors;
+			Colors.Add(GridColor);
+				
+			const double Scale = Settings->GridScale;
+			constexpr int32 MaxLevelDensity = 500;
+			constexpr int32 Subdivision = 5;
+				
+			MeshDebugDraw::DrawHierarchicalGrid(Scale, ZoomFactor / Scale, MaxLevelDensity, WorldMaxBounds, WorldMinBounds, Levels, Subdivision, Colors, DrawFrame, GridThickness, bDepthTested, PDI, DrawTransform);
+			break;
+		}
+			
+		case EModelingComponentsPlaneVisualizationMode::FixedScreenAreaGrid:
+		{
+			const double VisualAngleSpan = Settings->GridSize * 90.;
+			MeshDebugDraw::DrawSimpleFixedScreenAreaGrid(CameraState, DrawFrame, NumGridLines, VisualAngleSpan, GridThickness, GridColor, bDepthTested, PDI, DrawTransform);
+			break;
+		}
+		}
 	}
-}
-
+}
\ No newline at end of file
diff --git a/Engine/Plugins/Runtime/MeshModelingToolset/Source/ModelingComponents/Private/ModelingComponentsSettings.cpp b/Engine/Plugins/Runtime/MeshModelingToolset/Source/ModelingComponents/Private/ModelingComponentsSettings.cpp
--- a/Engine/Plugins/Runtime/MeshModelingToolset/Source/ModelingComponents/Private/ModelingComponentsSettings.cpp
+++ b/Engine/Plugins/Runtime/MeshModelingToolset/Source/ModelingComponents/Private/ModelingComponentsSettings.cpp
@@ -10,12 +10,22 @@
 #if WITH_EDITOR
 FText UModelingComponentsSettings::GetSectionText() const 
 { 
-	return LOCTEXT("ModelingComponentsSettingsName", "Modeling Mode Tools"); 
+	return LOCTEXT("ModelingComponentsProjectSettingsName", "Modeling Mode Tools"); 
 }
 
 FText UModelingComponentsSettings::GetSectionDescription() const
 {
-	return LOCTEXT("ModelingComponentsSettingsDescription", "Configure Tool-level Settings for the Modeling Tools Editor Mode plugin");
+	return LOCTEXT("ModelingComponentsProjectSettingsDescription", "Configure Tool-level Settings for the Modeling Tools Editor Mode plugin");
+}
+
+FText UModelingComponentsEditorSettings::GetSectionText() const 
+{ 
+	return LOCTEXT("ModelingComponentsEditorSettingsName", "Modeling Mode Tools"); 
+}
+
+FText UModelingComponentsEditorSettings::GetSectionDescription() const
+{
+	return LOCTEXT("ModelingComponentsEditorSettingsDescription", "Configure Tool-level Settings for the Modeling Tools Editor Mode plugin");
 }
 #endif
 
diff --git a/Engine/Plugins/Runtime/MeshModelingToolset/Source/ModelingComponents/Public/ModelingComponentsSettings.h b/Engine/Plugins/Runtime/MeshModelingToolset/Source/ModelingComponents/Public/ModelingComponentsSettings.h
--- a/Engine/Plugins/Runtime/MeshModelingToolset/Source/ModelingComponents/Public/ModelingComponentsSettings.h
+++ b/Engine/Plugins/Runtime/MeshModelingToolset/Source/ModelingComponents/Public/ModelingComponentsSettings.h
@@ -57,4 +57,62 @@ public:
 	TEnumAsByte<enum ECollisionTraceFlag> CollisionMode = ECollisionTraceFlag::CTF_UseComplexAsSimple;
 
 
+};
+
+
+
+/** Modeling Components plugin-wide plane visualization modes */
+UENUM()
+enum class EModelingComponentsPlaneVisualizationMode : uint8
+{
+	/** Draw a grid with a fixed size in world space */
+	SimpleGrid,
+	
+	/** Draw a hierarchical grid */
+	HierarchicalGrid,
+	
+	/** Draw a grid with a fixed size in screen space */
+	FixedScreenAreaGrid,
+};
+
+/**
+ * Editor preferences for the Modeling Components plug-in.
+ */
+UCLASS(config=EditorSettings)
+class MODELINGCOMPONENTS_API UModelingComponentsEditorSettings : public UDeveloperSettings
+{
+	GENERATED_BODY()
+
+public:
+	// UDeveloperSettings overrides
+	virtual FName GetContainerName() const { return FName("Editor"); }
+	virtual FName GetCategoryName() const { return FName("Plugins"); }
+	virtual FName GetSectionName() const { return FName("Modeling Mode Tools"); }
+
+#if WITH_EDITOR
+	virtual FText GetSectionText() const override;
+	virtual FText GetSectionDescription() const override;
+#endif
+
+public:
+	/** The type of grid to draw in the viewport for modeling mode tools */
+	UPROPERTY(EditAnywhere, NonTransactional, Category = "Modeling Tools|Work Plane Configuration")
+	EModelingComponentsPlaneVisualizationMode GridMode = EModelingComponentsPlaneVisualizationMode::SimpleGrid;
+
+	/** The number of lines to be drawn for the plane */
+	UPROPERTY(EditAnywhere, NonTransactional, Category = "Modeling Tools|Work Plane Configuration", meta = (ClampMin = 2, EditCondition = "GridMode != EModelingComponentsPlaneVisualizationMode::HierarchicalGrid", EditConditionHides))
+	int NumGridLines = 21;
+
+	/** The space between grid lines in world space */
+	UPROPERTY(EditAnywhere, NonTransactional, Category = "Modeling Tools|Work Plane Configuration", meta = (ClampMin = 0, EditCondition = "GridMode == EModelingComponentsPlaneVisualizationMode::SimpleGrid", EditConditionHides))
+	float GridSpacing = 100.0f;
+
+	/** The base scale used to determine the size of the hierarchical plane */
+	UPROPERTY(EditAnywhere, NonTransactional, Category = "Modeling Tools|Work Plane Configuration", meta = (ClampMin = 1, EditCondition = "GridMode == EModelingComponentsPlaneVisualizationMode::HierarchicalGrid", EditConditionHides))
+	float GridScale = 1.0f;
+
+	/** The fraction of the viewport that the grid should occupy if looking at the plane's center */
+	UPROPERTY(EditAnywhere, NonTransactional, Category = "Modeling Tools|Work Plane Configuration", meta = (ClampMin = 0, ClampMax = 2, EditCondition = "GridMode == EModelingComponentsPlaneVisualizationMode::FixedScreenAreaGrid", EditConditionHides))
+	float GridSize = 0.5;
+
 };
\ No newline at end of file
