MeshModelingToolset: Fix inconsistencies in "PolyGroup" strings in UI and tooltips

diff --git a/Engine/Plugins/Experimental/MeshModelingToolsetExp/Source/MeshModelingToolsEditorOnlyExp/Public/AttributeEditorTool.h b/Engine/Plugins/Experimental/MeshModelingToolsetExp/Source/MeshModelingToolsEditorOnlyExp/Public/AttributeEditorTool.h
--- a/Engine/Plugins/Experimental/MeshModelingToolsetExp/Source/MeshModelingToolsEditorOnlyExp/Public/AttributeEditorTool.h
+++ b/Engine/Plugins/Experimental/MeshModelingToolsetExp/Source/MeshModelingToolsEditorOnlyExp/Public/AttributeEditorTool.h
@@ -261,7 +261,7 @@ public:
 	}
 
 	/** Add a new PolyGroup layer with the given Name */
-	UFUNCTION(CallInEditor, Category = NewAttribute, meta = (DisplayPriority = 3))
+	UFUNCTION(CallInEditor, Category = NewAttribute, meta = (DisplayPriority = 3, DisplayName = "Add PolyGroup Layer"))
 	void AddPolyGroupLayer()
 	{
 		PostAction(EAttributeEditorToolActions::AddPolyGroupLayer);
diff --git a/Engine/Plugins/Experimental/MeshModelingToolsetExp/Source/MeshModelingToolsEditorOnlyExp/Public/SubdividePolyTool.h b/Engine/Plugins/Experimental/MeshModelingToolsetExp/Source/MeshModelingToolsEditorOnlyExp/Public/SubdividePolyTool.h
--- a/Engine/Plugins/Experimental/MeshModelingToolsetExp/Source/MeshModelingToolsEditorOnlyExp/Public/SubdividePolyTool.h
+++ b/Engine/Plugins/Experimental/MeshModelingToolsetExp/Source/MeshModelingToolsEditorOnlyExp/Public/SubdividePolyTool.h
@@ -55,7 +55,7 @@ public:
 	ESubdivisionOutputUVs UVComputationMethod = ESubdivisionOutputUVs::Interpolated;
 
 	/** Assign a new PolyGroup ID to each newly created face */
-	UPROPERTY(EditAnywhere, Category = Settings)
+	UPROPERTY(EditAnywhere, Category = Settings, meta = (DisplayName = "New PolyGroups"))
 	bool bNewPolyGroups = false;
 
 	UPROPERTY(EditAnywhere, Category = Settings)
diff --git a/Engine/Plugins/Experimental/MeshModelingToolsetExp/Source/MeshModelingToolsExp/Public/DrawPolyPathTool.h b/Engine/Plugins/Experimental/MeshModelingToolsetExp/Source/MeshModelingToolsExp/Public/DrawPolyPathTool.h
--- a/Engine/Plugins/Experimental/MeshModelingToolsetExp/Source/MeshModelingToolsExp/Public/DrawPolyPathTool.h
+++ b/Engine/Plugins/Experimental/MeshModelingToolsetExp/Source/MeshModelingToolsExp/Public/DrawPolyPathTool.h
@@ -106,7 +106,7 @@ public:
 	int RadialSlices = 16;
 
 	/** If true, all quads on the path will belong to the same polygon. If false, each quad gets its own polygon. */
-	UPROPERTY(EditAnywhere, Category = Path)
+	UPROPERTY(EditAnywhere, Category = Path, meta = (DisplayName = "Single PolyGroup"))
 	bool bSinglePolyGroup = false;
 
 	/** If and how the drawn path gets extruded */
diff --git a/Engine/Plugins/Runtime/MeshModelingToolset/Source/ModelingOperators/Public/Properties/RecomputeUVsProperties.h b/Engine/Plugins/Runtime/MeshModelingToolset/Source/ModelingOperators/Public/Properties/RecomputeUVsProperties.h
--- a/Engine/Plugins/Runtime/MeshModelingToolset/Source/ModelingOperators/Public/Properties/RecomputeUVsProperties.h
+++ b/Engine/Plugins/Runtime/MeshModelingToolset/Source/ModelingOperators/Public/Properties/RecomputeUVsProperties.h
@@ -30,7 +30,7 @@ enum class ERecomputeUVsPropertiesIslandMode
 {
 	// Values must match UE::Geometry::ERecomputeUVsIslandMode
 
-	/** Use Polygroups defined by Active PolyGroup property to define initial UV islands. */
+	/** Use PolyGroups defined by Active PolyGroup property to define initial UV islands. */
 	PolyGroups = 0 UMETA(DisplayName = "PolyGroups"),
 	/** Use existing UV Layer to define UV islands, i.e. re-solve UV flattening based on existing UVs */
 	ExistingUVs = 1
diff --git a/Engine/Plugins/Runtime/MeshModelingToolset/Source/ModelingOperatorsEditorOnly/Public/Properties/ParameterizeMeshProperties.h b/Engine/Plugins/Runtime/MeshModelingToolset/Source/ModelingOperatorsEditorOnly/Public/Properties/ParameterizeMeshProperties.h
--- a/Engine/Plugins/Runtime/MeshModelingToolset/Source/ModelingOperatorsEditorOnly/Public/Properties/ParameterizeMeshProperties.h
+++ b/Engine/Plugins/Runtime/MeshModelingToolset/Source/ModelingOperatorsEditorOnly/Public/Properties/ParameterizeMeshProperties.h
@@ -54,7 +54,7 @@ public:
 	int TextureResolution = 1024;
 
 	/** Generate new UVs based on polygroups from specified layer. */
-	UPROPERTY(EditAnywhere, Category = UVAtlas, meta = (DisplayName = "Constrain to Polygroups", EditCondition = bPolygroupsEnabled, EditConditionHides, HideEditConditionToggle))
+	UPROPERTY(EditAnywhere, Category = UVAtlas, meta = (DisplayName = "Constrain to PolyGroups", EditCondition = bPolygroupsEnabled, EditConditionHides, HideEditConditionToggle))
 	bool bUsePolygroups;
 
 	/** Layout resulting islands on UDIMs based on polygroups. */
@@ -120,7 +120,7 @@ public:
 	int TextureResolution = 1024;
 
 	/** Generate new UVs based on polygroups from specified layer. */
-	UPROPERTY(EditAnywhere, Category = PatchBuilder, meta = (DisplayName = "Constrain to Polygroups", EditCondition = bPolygroupsEnabled, EditConditionHides, HideEditConditionToggle))
+	UPROPERTY(EditAnywhere, Category = PatchBuilder, meta = (DisplayName = "Constrain to PolyGroups", EditCondition = bPolygroupsEnabled, EditConditionHides, HideEditConditionToggle))
 	bool bUsePolygroups;
 
 	/** Layout resulting islands on UDIMs based on polygroups. */
