MeshModelingToolset: Fix ResetToDefault for Output Object Name in e.g. Merge and Trim tools

diff --git a/Engine/Plugins/Runtime/MeshModelingToolset/Source/MeshModelingTools/Public/CombineMeshesTool.h b/Engine/Plugins/Runtime/MeshModelingToolset/Source/MeshModelingTools/Public/CombineMeshesTool.h
--- a/Engine/Plugins/Runtime/MeshModelingToolset/Source/MeshModelingTools/Public/CombineMeshesTool.h
+++ b/Engine/Plugins/Runtime/MeshModelingToolset/Source/MeshModelingTools/Public/CombineMeshesTool.h
@@ -56,7 +56,7 @@ public:
 
 	/** Name of the existing object to which the output is written to. */
 	UPROPERTY(VisibleAnywhere, Category = OutputObject, meta = (TransientToolProperty, DisplayName = "Name",
-		EditCondition = "bIsDuplicateMode == false && OutputWriteTo != EBaseCreateFromSelectedTargetType::NewObject", EditConditionHides))
+		EditCondition = "bIsDuplicateMode == false && OutputWriteTo != EBaseCreateFromSelectedTargetType::NewObject", EditConditionHides, NoResetToDefault))
 	FString OutputExistingName;
 };
 
diff --git a/Engine/Plugins/Runtime/MeshModelingToolset/Source/ModelingComponents/Public/BaseTools/BaseCreateFromSelectedTool.h b/Engine/Plugins/Runtime/MeshModelingToolset/Source/ModelingComponents/Public/BaseTools/BaseCreateFromSelectedTool.h
--- a/Engine/Plugins/Runtime/MeshModelingToolset/Source/ModelingComponents/Public/BaseTools/BaseCreateFromSelectedTool.h
+++ b/Engine/Plugins/Runtime/MeshModelingToolset/Source/ModelingComponents/Public/BaseTools/BaseCreateFromSelectedTool.h
@@ -64,7 +64,7 @@ public:
 
 	/** Name of the existing object to which the output is written to. */
 	UPROPERTY(VisibleAnywhere, Category = OutputObject, meta = (TransientToolProperty, DisplayName = "Name",
-		EditCondition = "OutputWriteTo != EBaseCreateFromSelectedTargetType::NewObject", EditConditionHides))
+		EditCondition = "OutputWriteTo != EBaseCreateFromSelectedTargetType::NewObject", EditConditionHides, NoResetToDefault))
 	FString OutputExistingName;
 };
 
