ModelingMode: Add support for Gizmo snapping and pivot-repositioning in new Selection mode.

UDragAlignmentMechanic is refactored into FDragAlignmentBase with subclasses UDragAlignmentMechanic and UDragAlignmentInteraction, the latter is usable outside of a Tool (the Mechanic requires a Tool).
FFindSceneSnapPointParams struct in ToolSceneQueriesUtil now takes explicit SnappingManager and CameraState, instead of a Tool that is used to look up those values. Update the one call site (inside of ToolSceneQueriesUtil) that used this struct.
Add a DragAlignmentInteraction to ModelingSelectionInteraction and hook up to TransformGizmo, as well as required DragAlignmentToggleBehavior. Also required adding UModelingSelectionInteraction::Render(), and calling fro ModelingToolsEditorMode

diff --git a/Engine/Plugins/Editor/ModelingToolsEditorMode/Source/ModelingToolsEditorMode/Private/ModelingSelectionInteraction.cpp b/Engine/Plugins/Editor/ModelingToolsEditorMode/Source/ModelingToolsEditorMode/Private/ModelingSelectionInteraction.cpp
--- a/Engine/Plugins/Editor/ModelingToolsEditorMode/Source/ModelingToolsEditorMode/Private/ModelingSelectionInteraction.cpp
+++ b/Engine/Plugins/Editor/ModelingToolsEditorMode/Source/ModelingToolsEditorMode/Private/ModelingSelectionInteraction.cpp
@@ -11,9 +11,10 @@
 #include "BaseBehaviors/SingleClickBehavior.h"
 #include "BaseBehaviors/SingleClickOrDragBehavior.h"
 #include "BaseBehaviors/MouseHoverBehavior.h"
+#include "BaseBehaviors/KeyAsModifierInputBehavior.h"
 
 #include "Mechanics/RectangleMarqueeMechanic.h"
-
+#include "Mechanics/DragAlignmentMechanic.h"
 
 using namespace UE::Geometry;
 
@@ -53,12 +54,20 @@ void UModelingSelectionInteraction::Initialize(
 	UpdateActiveDragMode();
 
 	TransformProxy = NewObject<UTransformProxy>(this);
-	// todo: make this repositionable etc. Maybe make this function a delegate? or allow caller to provide the gizmo?
-	TransformGizmo = UE::TransformGizmoUtil::Create3AxisTransformGizmo(
-		SelectionManager->GetToolsContext()->GizmoManager, this, TEXT("ModelingSelectionInteraction") );
+	TransformGizmo = UE::TransformGizmoUtil::CreateCustomRepositionableTransformGizmo( 
+		SelectionManager->GetToolsContext()->GizmoManager, ETransformGizmoSubElements::FullTranslateRotateScale,
+		this, TEXT("ModelingSelectionInteraction") );
 	TransformGizmo->SetActiveTarget(TransformProxy, SelectionManager->GetToolsContext()->GizmoManager);
 	TransformGizmo->SetVisibility(false);
 
+	DragAlignmentInteraction = NewObject<UDragAlignmentInteraction>(this);
+	DragAlignmentInteraction->Setup( USceneSnappingManager::Find(SelectionManager->GetToolsContext()->GizmoManager) );
+
+	DragAlignmentToggleBehavior = NewObject<UKeyAsModifierInputBehavior>(this);
+	DragAlignmentInteraction->RegisterAsBehaviorTarget(DragAlignmentToggleBehavior);
+	BehaviorSet->Add(DragAlignmentToggleBehavior, this);
+	DragAlignmentInteraction->AddToGizmo(TransformGizmo);
+
 	// listen for change events on transform proxy
 	TransformProxy->OnTransformChanged.AddUObject(this, &UModelingSelectionInteraction::OnGizmoTransformChanged);
 	TransformProxy->OnBeginTransformEdit.AddUObject(this, &UModelingSelectionInteraction::OnBeginGizmoTransform);
@@ -82,6 +91,13 @@ void UModelingSelectionInteraction::Shutdown()
 	}
 }
 
+void UModelingSelectionInteraction::Render(IToolsContextRenderAPI* RenderAPI)
+{
+	if (DragAlignmentInteraction)
+	{
+		DragAlignmentInteraction->Render(RenderAPI);
+	}
+}
 
 void UModelingSelectionInteraction::DrawHUD(FCanvas* Canvas, IToolsContextRenderAPI* RenderAPI)
 {
diff --git a/Engine/Plugins/Editor/ModelingToolsEditorMode/Source/ModelingToolsEditorMode/Private/ModelingToolsEditorMode.cpp b/Engine/Plugins/Editor/ModelingToolsEditorMode/Source/ModelingToolsEditorMode/Private/ModelingToolsEditorMode.cpp
--- a/Engine/Plugins/Editor/ModelingToolsEditorMode/Source/ModelingToolsEditorMode/Private/ModelingToolsEditorMode.cpp
+++ b/Engine/Plugins/Editor/ModelingToolsEditorMode/Source/ModelingToolsEditorMode/Private/ModelingToolsEditorMode.cpp
@@ -1076,6 +1076,8 @@ void UModelingToolsEditorMode::OnToolsContextRender(IToolsContextRenderAPI* Rend
 {
 	if (SelectionInteraction)
 	{
+		SelectionInteraction->Render(RenderAPI);
+
 		// Bake in transform changes. Note that if we do this in OnToolsContextTick, it will 
 		// still block rendering updates if it is too expensive, unless it is only done every second Tick
 		SelectionInteraction->ApplyPendingTransformInteractions();
diff --git a/Engine/Plugins/Editor/ModelingToolsEditorMode/Source/ModelingToolsEditorMode/Public/ModelingSelectionInteraction.h b/Engine/Plugins/Editor/ModelingToolsEditorMode/Source/ModelingToolsEditorMode/Public/ModelingSelectionInteraction.h
--- a/Engine/Plugins/Editor/ModelingToolsEditorMode/Source/ModelingToolsEditorMode/Public/ModelingSelectionInteraction.h
+++ b/Engine/Plugins/Editor/ModelingToolsEditorMode/Source/ModelingToolsEditorMode/Public/ModelingSelectionInteraction.h
@@ -20,6 +20,8 @@ class URectangleMarqueeMechanic;
 class UPathSelectionInteraction;
 class FCanvas;
 class IToolsContextRenderAPI;
+class UDragAlignmentInteraction;
+class UKeyAsModifierInputBehavior;
 
 
 UENUM()
@@ -69,6 +71,7 @@ public:
 	UGeometrySelectionManager* GetSelectionManager() { return SelectionManager; }
 	UE::Geometry::FGeometrySelectionUpdateConfig GetActiveSelectionUpdateConfig() const;
 
+	virtual void Render(IToolsContextRenderAPI* RenderAPI);
 	virtual void DrawHUD(FCanvas* Canvas, IToolsContextRenderAPI* RenderAPI);
 
 
@@ -174,6 +177,12 @@ protected:
 	UPROPERTY()
 	TObjectPtr<UCombinedTransformGizmo> TransformGizmo;
 
+	UPROPERTY()
+	TObjectPtr<UDragAlignmentInteraction> DragAlignmentInteraction;
+
+	UPROPERTY()
+	TObjectPtr<UKeyAsModifierInputBehavior> DragAlignmentToggleBehavior;
+
 	UE::Geometry::FFrame3d InitialGizmoFrame;
 	FVector3d InitialGizmoScale;
 	UE::Geometry::FFrame3d LastUpdateGizmoFrame;
diff --git a/Engine/Plugins/Runtime/MeshModelingToolset/Source/ModelingComponents/Private/Mechanics/DragAlignmentMechanic.cpp b/Engine/Plugins/Runtime/MeshModelingToolset/Source/ModelingComponents/Private/Mechanics/DragAlignmentMechanic.cpp
--- a/Engine/Plugins/Runtime/MeshModelingToolset/Source/ModelingComponents/Private/Mechanics/DragAlignmentMechanic.cpp
+++ b/Engine/Plugins/Runtime/MeshModelingToolset/Source/ModelingComponents/Private/Mechanics/DragAlignmentMechanic.cpp
@@ -11,66 +11,14 @@
 #include "SceneManagement.h" //FPrimitiveDrawInterface
 #include "ToolDataVisualizer.h"
 #include "ToolSceneQueriesUtil.h"
+#include "SceneQueries/SceneSnappingManager.h"
 
 #include UE_INLINE_GENERATED_CPP_BY_NAME(DragAlignmentMechanic)
 
 using namespace UE::Geometry;
 
-void UDragAlignmentMechanic::Setup(UInteractiveTool* ParentToolIn)
-{	
-	UInteractionMechanic::Setup(ParentToolIn);
 
-	// Register modifier listeners.
-	UKeyAsModifierInputBehavior* AlignmentToggleBehavior = NewObject<UKeyAsModifierInputBehavior>();
-	AlignmentToggleBehavior->Initialize(this, AlignmentModifierID, FInputDeviceState::IsCtrlKeyDown);
-	ParentTool->AddInputBehavior(AlignmentToggleBehavior);
-
-	// Set up the function that casts rays into the world.
-	WorldRayCast = [this, ParentToolIn](const FRay& WorldRay, FHitResult& HitResult, 
-		const TArray<const UPrimitiveComponent*>* ComponentsToIgnore,
-		const TArray<const UPrimitiveComponent*>* InvisibleComponentsToInclude)
-	{
-		bool bSuccess = false;
-		FVector3d QueryPoint = (FVector3d)WorldRay.PointAt(1);
-		FVector3d SnapPoint;
-		ToolSceneQueriesUtil::FSnapGeometry SnapGeometry;
-
-		// Try to snap to visible vertices/edges first
-		ToolSceneQueriesUtil::FFindSceneSnapPointParams SnapParams;
-		SnapParams.Tool = ParentTool.Get();
-		SnapParams.Point = &QueryPoint;
-		SnapParams.SnapPointOut = &SnapPoint;
-		SnapParams.bEdges = true;
-		SnapParams.bVertices = true;
-		SnapParams.VisualAngleThreshold = VisualAngleSnapThreshold;
-		SnapParams.SnapGeometryOut = &SnapGeometry;
-		SnapParams.ComponentsToIgnore = ComponentsToIgnore;
-		SnapParams.InvisibleComponentsToInclude = InvisibleComponentsToInclude;
-
-		if (ToolSceneQueriesUtil::FindSceneSnapPoint(SnapParams))
-		{
-			HitResult = FHitResult(WorldRay.Origin, (FVector)SnapPoint);
-			HitResult.ImpactPoint = (FVector)SnapPoint;
-			HitResult.Distance = WorldRay.GetParameter(HitResult.ImpactPoint);
-			HitResult.Item = SnapGeometry.PointCount == 1 ? VERT_HIT_TYPE_ID : EDGE_HIT_TYPE_ID;
-			bSuccess = true;
-		}
-		// If we failed, try to hit a simple triangle
-		else if (ToolSceneQueriesUtil::FindNearestVisibleObjectHit(ParentToolIn, HitResult, WorldRay, ComponentsToIgnore))
-		{
-			HitResult.Item = TRIANGLE_HIT_TYPE_ID;
-			bSuccess = true;
-		}
-
-		return bSuccess;
-	};
-}
-
-void UDragAlignmentMechanic::Shutdown()
-{
-}
-
-void UDragAlignmentMechanic::AddToGizmo(UCombinedTransformGizmo* TransformGizmo, const TArray<const UPrimitiveComponent*>* ComponentsToIgnoreInAlignment, 
+void FDragAlignmentBase::AddToGizmo(UCombinedTransformGizmo* TransformGizmo, const TArray<const UPrimitiveComponent*>* ComponentsToIgnoreInAlignment, 
 	const TArray<const UPrimitiveComponent*>* InvisibleComponentsToIncludeInAlignment)
 {
 	// If we have components to ignore/include, we need a copy of the array so that the alignment
@@ -88,7 +36,7 @@
 
 	// Set the alignment functions.
 	TransformGizmo->SetWorldAlignmentFunctions(
-		[this]() { return bAlignmentToggle; },
+		[this]() { return GetAlignmentModeEnabled(); },
 		[this, ComponentsToIgnorePersistent, ComponentsToIncludePersistent](const FRay& WorldRay, FVector& OutputPoint) {
 			return CastRay(WorldRay, OutputPoint, ComponentsToIgnorePersistent.Get(), ComponentsToIncludePersistent.Get(), true);
 		});
@@ -101,7 +49,7 @@
 	if (URepositionableTransformGizmo* RepositionableGizmo = Cast<URepositionableTransformGizmo>(TransformGizmo))
 	{
 		RepositionableGizmo->SetPivotAlignmentFunctions(
-			[this]() { return bAlignmentToggle; },
+			[this]() { return GetAlignmentModeEnabled(); },
 			[this, ComponentsToIncludePersistent](const FRay& WorldRay, FVector& OutputPoint) {
 				return CastRay(WorldRay, OutputPoint, nullptr, ComponentsToIncludePersistent.Get(), false); // don't ignore anything
 			});
@@ -109,25 +57,25 @@
 
 	// Register listeners to help with rendering. They get us the location that the gizmo ended
 	// up being placed, so that we can draw a line from the hit point to that location.
-	TransformGizmo->ActiveTarget->OnTransformChanged.AddWeakLambda(this, [this](UTransformProxy*, FTransform NewTransform){
+	TransformGizmo->ActiveTarget->OnTransformChanged.AddWeakLambda(GetUObjectContainer(), [this](UTransformProxy*, FTransform NewTransform) {
 		OnGizmoTransformChanged(NewTransform);
 	});
-	TransformGizmo->ActiveTarget->OnPivotChanged.AddWeakLambda(this, [this](UTransformProxy*, FTransform NewTransform) {
+	TransformGizmo->ActiveTarget->OnPivotChanged.AddWeakLambda(GetUObjectContainer(), [this](UTransformProxy*, FTransform NewTransform) {
 		OnGizmoTransformChanged(NewTransform);
-		});
+	});
 
 	// Register a listener to stop drawing once the gizmo is done moving.
-	TransformGizmo->ActiveTarget->OnEndTransformEdit.AddWeakLambda(this, [this](UTransformProxy*) { 
+	TransformGizmo->ActiveTarget->OnEndTransformEdit.AddWeakLambda(GetUObjectContainer(), [this](UTransformProxy*) { 
 		bPreviewEndpointsValid = false; 
 		bWaitingOnProjectedResult = false;
 	});
-	TransformGizmo->ActiveTarget->OnEndPivotEdit.AddWeakLambda(this, [this](UTransformProxy*) {
+	TransformGizmo->ActiveTarget->OnEndPivotEdit.AddWeakLambda(GetUObjectContainer(), [this](UTransformProxy*) {
 		bPreviewEndpointsValid = false;
 		bWaitingOnProjectedResult = false;
 		});
 }
 
-void UDragAlignmentMechanic::AddToGizmo(UIntervalGizmo* IntervalGizmo, const TArray<const UPrimitiveComponent*>* ComponentsToIgnoreInAlignment,
+void FDragAlignmentBase::AddToGizmo(UIntervalGizmo* IntervalGizmo, const TArray<const UPrimitiveComponent*>* ComponentsToIgnoreInAlignment,
 	const TArray<const UPrimitiveComponent*>* InvisibleComponentsToIncludeInAlignment)
 {
 	// If we have components to ignore/include, we need a copy of the array so that the alignment
@@ -145,27 +93,27 @@
 
 	// Set the alignment functions.
 	IntervalGizmo->SetWorldAlignmentFunctions(
-		[this]() { return bAlignmentToggle; },
+		[this]() { return GetAlignmentModeEnabled(); },
 		[this, ComponentsToIgnorePersistent, ComponentsToIncludePersistent](const FRay& WorldRay, FVector& OutputPoint) {
 			return CastRay(WorldRay, OutputPoint, ComponentsToIgnorePersistent.Get(), ComponentsToIncludePersistent.Get(), true);
 		});
 
 	// Register listener to help with rendering. It gets us the final location of the interval endpoint
 	// so that we can draw a line from the hit point to that location.
-	IntervalGizmo->OnIntervalChanged.AddWeakLambda(this, [this](UIntervalGizmo* Gizmo, FVector Direction, float NewParam) {
+	IntervalGizmo->OnIntervalChanged.AddWeakLambda(GetUObjectContainer(), [this](UIntervalGizmo* Gizmo, FVector Direction, float NewParam) {
 		FTransform Transform = Gizmo->GetGizmoTransform();
 		Transform.AddToTranslation(Transform.Rotator().RotateVector(Direction) * NewParam);
 		OnGizmoTransformChanged(Transform);
 		});
 
 	// Register a listener to stop drawing once the gizmo is done moving.
-	IntervalGizmo->OnEndIntervalGizmoEdit.AddWeakLambda(this, [this](UIntervalGizmo*) {
+	IntervalGizmo->OnEndIntervalGizmoEdit.AddWeakLambda(GetUObjectContainer(), [this](UIntervalGizmo*) {
 		bPreviewEndpointsValid = false;
 		bWaitingOnProjectedResult = false;
 		});
 }
 
-void UDragAlignmentMechanic::OnGizmoTransformChanged(FTransform NewTransform)
+void FDragAlignmentBase::OnGizmoTransformChanged(FTransform NewTransform)
 {
 	// Ignore changes that weren't preceded by a successful query (likely programatic)
 	if (bWaitingOnProjectedResult) 
@@ -176,7 +124,7 @@
 	bWaitingOnProjectedResult = false;
 }
 
-void UDragAlignmentMechanic::InitializeDeformedMeshRayCast(TFunction<FDynamicMeshAABBTree3* ()> GetSpatialIn, 
+void FDragAlignmentBase::InitializeDeformedMeshRayCast(TFunction<FDynamicMeshAABBTree3* ()> GetSpatialIn, 
 	const FTransform3d &TargetTransform, const FGroupTopologyDeformer* LinearDeformer)
 {
 	if (LinearDeformer)
@@ -277,7 +225,7 @@
  *  filtered to avoid hitting deformed triangles.
  * @return true if something was hit.
  */
-bool UDragAlignmentMechanic::CastRay(const FRay& WorldRay, FVector& OutputPoint,
+bool FDragAlignmentBase::CastRay(const FRay& WorldRay, FVector& OutputPoint,
 	const TArray<const UPrimitiveComponent*>* ComponentsToIgnore, 
 	const TArray<const UPrimitiveComponent*>* InvisibleComponentsToInclude,
 	bool bUseFilter)
@@ -309,10 +257,60 @@
 	return bHit;
 }
 
-void UDragAlignmentMechanic::Render(IToolsContextRenderAPI* RenderAPI)
+
+
+
+
+
+void FDragAlignmentBase::SetupInternal()
 {
+	// Set up the function that casts rays into the world.
+	WorldRayCast = [this](const FRay& WorldRay, FHitResult& HitResult, 
+		const TArray<const UPrimitiveComponent*>* ComponentsToIgnore,
+		const TArray<const UPrimitiveComponent*>* InvisibleComponentsToInclude)
+	{
+		bool bSuccess = false;
+		FVector3d QueryPoint = (FVector3d)WorldRay.PointAt(1);
+		FVector3d SnapPoint;
+		ToolSceneQueriesUtil::FSnapGeometry SnapGeometry;
+
+		// Try to snap to visible vertices/edges first
+		ToolSceneQueriesUtil::FFindSceneSnapPointParams SnapParams;
+		SnapParams.SnappingManager = GetSnappingManager();
+		SnapParams.CameraState = CachedCameraState;
+		SnapParams.Point = &QueryPoint;
+		SnapParams.SnapPointOut = &SnapPoint;
+		SnapParams.bEdges = true;
+		SnapParams.bVertices = true;
+		SnapParams.VisualAngleThreshold = VisualAngleSnapThreshold;
+		SnapParams.SnapGeometryOut = &SnapGeometry;
+		SnapParams.ComponentsToIgnore = ComponentsToIgnore;
+		SnapParams.InvisibleComponentsToInclude = InvisibleComponentsToInclude;
+
+		if (ToolSceneQueriesUtil::FindSceneSnapPoint(SnapParams))
+		{
+			HitResult = FHitResult(WorldRay.Origin, (FVector)SnapPoint);
+			HitResult.ImpactPoint = (FVector)SnapPoint;
+			HitResult.Distance = WorldRay.GetParameter(HitResult.ImpactPoint);
+			HitResult.Item = SnapGeometry.PointCount == 1 ? VERT_HIT_TYPE_ID : EDGE_HIT_TYPE_ID;
+			bSuccess = true;
+		}
+		// If we failed, try to hit a simple triangle
+		else if (ToolSceneQueriesUtil::FindNearestVisibleObjectHit(GetSnappingManager(), HitResult, WorldRay, ComponentsToIgnore))
+		{
+			HitResult.Item = TRIANGLE_HIT_TYPE_ID;
+			bSuccess = true;
+		}
+
+		return bSuccess;
+	};
+}
+
+
+void FDragAlignmentBase::RenderInternal(IToolsContextRenderAPI* RenderAPI)
+{
 	// Cache the camera state
-	GetParentTool()->GetToolManager()->GetContextQueriesAPI()->GetCurrentViewState(CachedCameraState);
+	CachedCameraState = GetCameraState();
 
 	if (!bPreviewEndpointsValid)
 	{
@@ -363,6 +361,43 @@
 	Renderer.EndFrame();
 }
 
+
+
+
+void UDragAlignmentMechanic::Setup(UInteractiveTool* ParentToolIn)
+{	
+	UInteractionMechanic::Setup(ParentToolIn);
+
+	// Register modifier listeners.
+	UKeyAsModifierInputBehavior* AlignmentToggleBehavior = NewObject<UKeyAsModifierInputBehavior>();
+	AlignmentToggleBehavior->Initialize(this, AlignmentModifierID, FInputDeviceState::IsCtrlKeyDown);
+	ParentTool->AddInputBehavior(AlignmentToggleBehavior);
+
+	SetupInternal();
+}
+
+void UDragAlignmentMechanic::Shutdown()
+{
+}
+
+
+void UDragAlignmentMechanic::Render(IToolsContextRenderAPI* RenderAPI)
+{
+	RenderInternal(RenderAPI);
+}
+
+USceneSnappingManager* UDragAlignmentMechanic::GetSnappingManager()
+{
+	return USceneSnappingManager::Find(GetParentTool()->GetToolManager());
+}
+
+FViewCameraState UDragAlignmentMechanic::GetCameraState()
+{
+	FViewCameraState CurState;
+	GetParentTool()->GetToolManager()->GetContextQueriesAPI()->GetCurrentViewState(CurState);
+	return CurState;
+}
+
 void UDragAlignmentMechanic::OnUpdateModifierState(int ModifierID, bool bIsOn)
 {
 	if (ModifierID == AlignmentModifierID)
@@ -375,3 +410,34 @@
 	}
 }
 
+
+
+
+void UDragAlignmentInteraction::Setup(USceneSnappingManager* SnappingManager)
+{
+	this->SceneSnappingManager = SnappingManager;
+	SetupInternal();
+}
+
+void UDragAlignmentInteraction::RegisterAsBehaviorTarget(UKeyAsModifierInputBehavior* KeyModifierBehavior)
+{
+	KeyModifierBehavior->Initialize(this, AlignmentModifierID, FInputDeviceState::IsCtrlKeyDown);
+}
+
+
+void UDragAlignmentInteraction::OnUpdateModifierState(int ModifierID, bool bIsOn)
+{
+	if (ModifierID == AlignmentModifierID)
+	{
+		bAlignmentToggle = bIsOn;
+	}
+	if (!bAlignmentToggle)
+	{
+		bPreviewEndpointsValid = false;
+	}
+}
+void UDragAlignmentInteraction::Render(IToolsContextRenderAPI* RenderAPI)
+{
+	LastRenderCameraState = RenderAPI->GetCameraState();
+	RenderInternal(RenderAPI);
+}
diff --git a/Engine/Plugins/Runtime/MeshModelingToolset/Source/ModelingComponents/Private/ToolSceneQueriesUtil.cpp b/Engine/Plugins/Runtime/MeshModelingToolset/Source/ModelingComponents/Private/ToolSceneQueriesUtil.cpp
--- a/Engine/Plugins/Runtime/MeshModelingToolset/Source/ModelingComponents/Private/ToolSceneQueriesUtil.cpp
+++ b/Engine/Plugins/Runtime/MeshModelingToolset/Source/ModelingComponents/Private/ToolSceneQueriesUtil.cpp
@@ -143,7 +143,8 @@ bool ToolSceneQueriesUtil::FindSceneSnapPoint(const UInteractiveTool* Tool, cons
 {
 	FFindSceneSnapPointParams Params;
 
-	Params.Tool = Tool;
+	Params.SnappingManager = USceneSnappingManager::Find(Tool->GetToolManager());
+	Tool->GetToolManager()->GetContextQueriesAPI()->GetCurrentViewState(Params.CameraState);
 	Params.Point = &Point;
 	Params.SnapPointOut = &SnapPointOut;
 	Params.bVertices = bVertices;
@@ -159,15 +160,12 @@ bool ToolSceneQueriesUtil::FindSceneSnapPoint(FFindSceneSnapPointParams& Params)
 {
 	double UseThreshold = (Params.VisualAngleThreshold <= 0) ? GetDefaultVisualAngleSnapThreshD() : Params.VisualAngleThreshold;
 
-	USceneSnappingManager* SnapManager = USceneSnappingManager::Find(Params.Tool->GetToolManager());
-	if (!SnapManager)
+	if (!Params.SnappingManager)
 	{
 		return false;
 	}
 
-	FViewCameraState CameraState;
-	Params.Tool->GetToolManager()->GetContextQueriesAPI()->GetCurrentViewState(CameraState);
-	UseThreshold *= CameraState.GetFOVAngleNormalizationFactor();
+	UseThreshold *= Params.CameraState.GetFOVAngleNormalizationFactor();
 
 	FSceneSnapQueryRequest Request;
 	Request.RequestType = ESceneSnapQueryType::Position;
@@ -186,7 +184,7 @@ bool ToolSceneQueriesUtil::FindSceneSnapPoint(FFindSceneSnapPointParams& Params)
 	Request.InvisibleComponentsToInclude = Params.InvisibleComponentsToInclude;
 
 	TArray<FSceneSnapQueryResult> Results;
-	if (SnapManager->ExecuteSceneSnapQuery(Request, Results))
+	if (Params.SnappingManager->ExecuteSceneSnapQuery(Request, Results))
 	{
 		*Params.SnapPointOut = (FVector3d)Results[0].Position;
 
diff --git a/Engine/Plugins/Runtime/MeshModelingToolset/Source/ModelingComponents/Public/Mechanics/DragAlignmentMechanic.h b/Engine/Plugins/Runtime/MeshModelingToolset/Source/ModelingComponents/Public/Mechanics/DragAlignmentMechanic.h
--- a/Engine/Plugins/Runtime/MeshModelingToolset/Source/ModelingComponents/Public/Mechanics/DragAlignmentMechanic.h
+++ b/Engine/Plugins/Runtime/MeshModelingToolset/Source/ModelingComponents/Public/Mechanics/DragAlignmentMechanic.h
@@ -17,22 +17,35 @@
 
 #include "DragAlignmentMechanic.generated.h"
 
-PREDECLARE_USE_GEOMETRY_CLASS(FGroupTopologyDeformer);
+namespace UE::Geometry { class FGroupTopologyDeformer; }
 class UPrimitiveComponent;
 class UIntervalGizmo;
 class UCombinedTransformGizmo;
+class USceneSnappingManager;
+class UKeyAsModifierInputBehavior;
+
 
 /**
- * Mechanic that can be added to (potentially multiple) UCombinedTransformGizmo object to allow them to snap to
- * scene geometry in rotation and translation when the Ctrl key is pressed.
+ * FDragAlignmentBase is a base class for shared functionality of UDragAlignmentMechanic and UDragAlignmentInteraction.
+ * This interaction can be attached to (potentially multiple) UCombinedTransformGizmo object to allow them to snap to
+ * scene geometry in rotation and translation when a modifier key is pressed. 
  */
-UCLASS()
-class MODELINGCOMPONENTS_API UDragAlignmentMechanic : public UInteractionMechanic, public IModifierToggleBehaviorTarget
+class MODELINGCOMPONENTS_API FDragAlignmentBase
 {
-	GENERATED_BODY()
+public:
+	virtual ~FDragAlignmentBase() {}
+
+	//
+	// Subclasses must implement these various API functions for the interaction to function
+	//
+	virtual USceneSnappingManager* GetSnappingManager() = 0;
+	virtual FViewCameraState GetCameraState() = 0;
+	virtual bool GetAlignmentModeEnabled() const = 0;
+	virtual UObject* GetUObjectContainer() = 0;
 
 public:
 
+
 	/**
 	 * Adds this mechanic to the given gizmo. Can be called on multiple gizmos, and works both for UCombinedTransformGizmo
 	 * and URepositionableTransformGizmo.
@@ -57,19 +70,21 @@ public:
 	 */
 	void InitializeDeformedMeshRayCast(
 		TFunction<UE::Geometry::FDynamicMeshAABBTree3* ()> GetSpatialIn,
-		const FTransform3d& TargetTransform, const FGroupTopologyDeformer* LinearDeformer);
+		const FTransform3d& TargetTransform, const UE::Geometry::FGroupTopologyDeformer* LinearDeformer);
 
-	virtual void Setup(UInteractiveTool* ParentToolIn) override;
-	virtual void Shutdown() override;
+
+protected:
+
+	/**
+	 * Subclasses must call this to initialize the internal functions/states/etc
+	 */
+	virtual void SetupInternal();
 
 	/**
 	 * This will draw a line from the last input position to the last output position. This
 	 * gets cleared when the transform is ended. 
 	 */
-	virtual void Render(IToolsContextRenderAPI* RenderAPI) override;
-
-	// IModifierToggleBehaviorTarget implementation
-	virtual void OnUpdateModifierState(int ModifierID, bool bIsOn) override;
+	virtual void RenderInternal(IToolsContextRenderAPI* RenderAPI);
 
 protected:
 	// These are modifiable in case we someday want to make it easier to customize them further.
@@ -95,12 +110,86 @@ protected:
 	FHitResult LastHitResult;
 	FVector3d LastProjectedResult;
 
-	bool bAlignmentToggle = false;
-	int32 AlignmentModifierID = 1;
-
 	virtual bool CastRay(const FRay& WorldRay, FVector& OutputPoint,
 		const TArray<const UPrimitiveComponent*>* ComponentsToIgnore, 
 		const TArray<const UPrimitiveComponent*>* InvisibleComponentsToInclude,
 		bool bUseFilter);
+
 	void OnGizmoTransformChanged(FTransform NewTransform);
+};
+
+
+/**
+ * Mechanic that can be added to (potentially multiple) UCombinedTransformGizmo object to allow them to snap to
+ * scene geometry in rotation and translation when the Ctrl key is pressed.
+ */
+UCLASS()
+class MODELINGCOMPONENTS_API UDragAlignmentMechanic : public UInteractionMechanic, public FDragAlignmentBase, public IModifierToggleBehaviorTarget
+{
+	GENERATED_BODY()
+
+public:
+
+	virtual USceneSnappingManager* GetSnappingManager();
+	virtual FViewCameraState GetCameraState();
+	virtual bool GetAlignmentModeEnabled() const { return bAlignmentToggle;}
+	virtual UObject* GetUObjectContainer() { return this; }
+
+public:
+
+	virtual void Setup(UInteractiveTool* ParentToolIn) override;
+	virtual void Shutdown() override;
+
+
+	virtual void Render(IToolsContextRenderAPI* RenderAPI) override;
+
+	// IModifierToggleBehaviorTarget implementation
+	virtual void OnUpdateModifierState(int ModifierID, bool bIsOn) override;
+
+protected:
+	bool bAlignmentToggle = false;
+	int32 AlignmentModifierID = 1;
+
+};
+
+
+/**
+ * Interaction that can be added to (potentially multiple) UCombinedTransformGizmo object to allow them to snap to
+ * scene geometry in rotation and translation. Generally driven by an externally-provided UKeyAsModifierInputBehavior,
+ * or alternately can be directly updated by calling ::OnUpdateModifierState()
+ */
+UCLASS()
+class MODELINGCOMPONENTS_API UDragAlignmentInteraction : public UObject, public FDragAlignmentBase, public IModifierToggleBehaviorTarget
+{
+	GENERATED_BODY()
+
+public:
+	virtual USceneSnappingManager* GetSnappingManager() { return SceneSnappingManager.Get(); }
+	virtual FViewCameraState GetCameraState() { return LastRenderCameraState; }
+	virtual bool GetAlignmentModeEnabled() const { return bAlignmentToggle;}
+	virtual UObject* GetUObjectContainer() { return this; }
+
+
+public:
+	virtual void Setup(USceneSnappingManager* SnappingManager);
+
+	// Client must call this every frame to do alignment visualization rendering
+	virtual void Render(IToolsContextRenderAPI* RenderAPI);
+
+
+public:
+	// Client can call this on setup to connect the interaction to an external Key behavior.
+	// This function is hardcoded to register the the Ctrl key as the modifier that drives the snapping toggle
+	virtual void RegisterAsBehaviorTarget(UKeyAsModifierInputBehavior* KeyModifierBehavior);
+
+	// IModifierToggleBehaviorTarget implementation
+	virtual void OnUpdateModifierState(int ModifierID, bool bIsOn) override;
+
+protected:
+	TWeakObjectPtr<USceneSnappingManager> SceneSnappingManager;
+
+	FViewCameraState LastRenderCameraState;		// could maybe move this to FDragAlignmentBase and get rid of GetCameraState() API function...
+
+	bool bAlignmentToggle = false;
+	int32 AlignmentModifierID = 1;
 };
\ No newline at end of file
diff --git a/Engine/Plugins/Runtime/MeshModelingToolset/Source/ModelingComponents/Public/ToolSceneQueriesUtil.h b/Engine/Plugins/Runtime/MeshModelingToolset/Source/ModelingComponents/Public/ToolSceneQueriesUtil.h
--- a/Engine/Plugins/Runtime/MeshModelingToolset/Source/ModelingComponents/Public/ToolSceneQueriesUtil.h
+++ b/Engine/Plugins/Runtime/MeshModelingToolset/Source/ModelingComponents/Public/ToolSceneQueriesUtil.h
@@ -97,7 +97,8 @@ namespace ToolSceneQueriesUtil
 	struct FFindSceneSnapPointParams
 	{
 		// Required inputs/outputs
-		const UInteractiveTool* Tool = nullptr;
+		const USceneSnappingManager* SnappingManager = nullptr;
+		FViewCameraState CameraState;
 		const FVector3d* Point = nullptr;
 		FVector3d* SnapPointOut = nullptr;
 
