Geometry : ModelingComponents and GeometryScript - add ability to deform a dynamic mesh according to a spline mesh component add conversion from spline mesh component to dynamic mesh

diff --git a/Engine/Plugins/Experimental/GeometryScripting/Source/GeometryScriptingCore/Private/SceneUtilityFunctions.cpp b/Engine/Plugins/Experimental/GeometryScripting/Source/GeometryScriptingCore/Private/SceneUtilityFunctions.cpp
--- a/Engine/Plugins/Experimental/GeometryScripting/Source/GeometryScriptingCore/Private/SceneUtilityFunctions.cpp
+++ b/Engine/Plugins/Experimental/GeometryScripting/Source/GeometryScriptingCore/Private/SceneUtilityFunctions.cpp
@@ -15,8 +15,10 @@
 #include "Components/SkinnedMeshComponent.h"
 #include "Components/BrushComponent.h"
 #include "Components/DynamicMeshComponent.h"
+#include "Components/SplineMeshComponent.h"
 #include "ConversionUtils/VolumeToDynamicMesh.h"
 #include "ConversionUtils/SkinnedMeshToDynamicMesh.h"
+#include "ConversionUtils/SplineComponentDeformDynamicMesh.h"
 #include "Physics/ComponentCollisionUtil.h"
 
 #include UE_INLINE_GENERATED_CPP_BY_NAME(SceneUtilityFunctions)
@@ -70,6 +72,28 @@ UDynamicMesh* UGeometryScriptLibrary_SceneUtilityFunctions::CopyMeshFromComponen
 		}
 
 	}
+	else if (USplineMeshComponent* SplineMeshComponent = Cast<USplineMeshComponent>(Component))
+	{
+		LocalToWorld = SplineMeshComponent->GetComponentTransform();
+		UStaticMesh* StaticMesh = SplineMeshComponent->GetStaticMesh();
+		if (StaticMesh)
+		{
+			FGeometryScriptCopyMeshFromAssetOptions AssetOptions;
+			AssetOptions.bApplyBuildSettings = (Options.bWantNormals || Options.bWantTangents);
+			AssetOptions.bRequestTangents = Options.bWantTangents;
+			UGeometryScriptLibrary_StaticMeshFunctions::CopyMeshFromStaticMesh(
+				StaticMesh, ToDynamicMesh, AssetOptions, Options.RequestedLOD, Outcome, Debug);	// will set Outcome pin
+
+			// deform the dynamic mesh and its tangent space with the spline
+			constexpr bool bUpdateTangentSpace = true;
+			UE::Geometry::SplineDeformDynamicMesh(*SplineMeshComponent, ToDynamicMesh->GetMeshRef(), bUpdateTangentSpace);
+
+		}
+		else
+		{
+			UE::Geometry::AppendError(Debug, EGeometryScriptErrorType::InvalidInputs, LOCTEXT("CopyMeshFromSplineMeshComponent_MissingStaticMesh", "CopyMeshFromComponent: SplineMeshComponent has a null StaticMesh"));
+		}
+	}
 	else if (UStaticMeshComponent* StaticMeshComponent = Cast<UStaticMeshComponent>(Component))
 	{
 		LocalToWorld = StaticMeshComponent->GetComponentTransform();
diff --git a/Engine/Plugins/Runtime/MeshModelingToolset/Source/ModelingComponents/Private/ConversionUtils/SplineComponentDeformDynamicMesh.cpp b/Engine/Plugins/Runtime/MeshModelingToolset/Source/ModelingComponents/Private/ConversionUtils/SplineComponentDeformDynamicMesh.cpp
new file mode 100644
--- /dev/null
+++ b/Engine/Plugins/Runtime/MeshModelingToolset/Source/ModelingComponents/Private/ConversionUtils/SplineComponentDeformDynamicMesh.cpp
@@ -0,0 +1,121 @@
+// Copyright Epic Games, Inc. All Rights Reserved.
+
+#include "ConversionUtils/SplineComponentDeformDynamicMesh.h"
+
+
+#include "Components/SplineMeshComponent.h"
+#include "CoreMinimal.h"
+#include "DynamicMesh/DynamicMesh3.h"
+#include "DynamicMesh/DynamicMeshAttributeSet.h"
+
+#include "Async/ParallelFor.h"
+
+namespace UE {
+namespace Geometry {
+
+void SplineDeformDynamicMesh(USplineMeshComponent& SplineMeshComponent, FDynamicMesh3& Mesh, bool bUpdateTangentSpace)
+{
+	
+
+	const ESplineMeshAxis::Type ForwardAxis = SplineMeshComponent.ForwardAxis;
+	const FVector Mask = SplineMeshComponent.GetAxisMask(ForwardAxis);
+	const int32 ForwardComponent = [&] { switch (SplineMeshComponent.ForwardAxis)
+											{
+											case ESplineMeshAxis::X: return 0; break;
+											case ESplineMeshAxis::Y: return 1; break;
+											case ESplineMeshAxis::Z: return 2; break;
+											default: check(0); return -1; break;
+											}}();
+
+		// the transform is a function of position
+		auto SliceTransform = [&SplineMeshComponent, ForwardComponent](const FVector& Position)-> FTransform
+								{
+									const float ForwardDistance = static_cast<float>(Position[ForwardComponent]);
+									return SplineMeshComponent.CalcSliceTransform(ForwardDistance);
+								};
+		// position transform. note specialized because SliceTranform has the forward axis component of position baked in.
+		auto TransformPosition = [&SliceTransform, Mask](const FVector& Position)->FVector3d
+									{
+										const FVector MaskedVert = Position * Mask;
+										FVector3d OutPos = SliceTransform(Position).TransformPosition(MaskedVert);
+										return OutPos;
+									};
+
+
+		// Since the transform is a function of position, we need to update the tangent space first.
+		// For consistency with the way a static mesh tangent space is deformed by the spline on the GPU (see LocalVertexFactory.ush ) the vectors are simply rotated by 
+		// the rotational component of the transform.  Strictly speaking this is wrong... 
+		if (bUpdateTangentSpace)
+		{
+			// update tangent space attributes layers if they exist
+			if (FDynamicMeshAttributeSet* Attributes = Mesh.Attributes())
+			{
+				auto TransformVectorOverlay = [&SliceTransform, &Mesh](auto* Overlay)
+												{
+													const int32 NumElements = Overlay->MaxElementID();
+													ParallelFor(NumElements, [&](int32 elid)
+														{
+															if (Overlay->IsElement(elid))
+															{
+																// transform at position where this element lives
+																const  FTransform Xform = [&] {  const int32 vid = Overlay->GetParentVertex(elid);
+																								 const FVector3d Pos = Mesh.GetVertex(vid);
+																								 return SliceTransform(Pos); } ();
+
+																const FVector3f Vector3f = Overlay->GetElement(elid);
+																const FVector3d TransformedVector3d = Xform.GetRotation() * static_cast<FVector3d>(Vector3f);
+																const FVector3f TransformedVector3f = static_cast<FVector3f>(TransformedVector3d);
+
+																Overlay->SetElement(elid, TransformedVector3f);
+															}
+														});
+												};
+
+				// update tangents
+				if (FDynamicMeshNormalOverlay* TangentOverlay = Attributes->PrimaryTangents())
+				{
+					TransformVectorOverlay(TangentOverlay);
+				}
+				// update bi-tangents
+				if (FDynamicMeshNormalOverlay* BiTangentOverlay = Attributes->PrimaryBiTangents())
+				{
+					TransformVectorOverlay(BiTangentOverlay);
+				}
+				// update normals
+				if (FDynamicMeshNormalOverlay* NormalOverlay = Attributes->PrimaryNormals())
+				{
+					TransformVectorOverlay(NormalOverlay);
+				}
+			}
+			else if (Mesh.HasVertexNormals())
+			{
+				const int32 NumVIDs = Mesh.MaxVertexID();
+				ParallelFor(NumVIDs, [&](int32 vid)
+					{
+						if (Mesh.IsVertex(vid))
+						{
+							const  FTransform Xform = [&] {  const FVector3d Pos = Mesh.GetVertex(vid);
+							return SliceTransform(Pos); } ();
+							const FVector3f VertexNormal3f = Mesh.GetVertexNormal(vid);
+							const FVector3f TransformedVertexNormal3f = static_cast<FVector3f>(Xform.GetRotation() * static_cast<FVector3d>(VertexNormal3f));
+							Mesh.SetVertexNormal(vid, TransformedVertexNormal3f);
+						}
+					});
+			}
+		}
+
+		// update vertex locations.  Note, tangents should be updated before position
+		const int32 NumVertices = Mesh.MaxVertexID();
+		ParallelFor(NumVertices, [&](int32 vid)
+			{
+				if (Mesh.IsVertex(vid))
+				{
+					FVector3d Position = Mesh.GetVertex(vid);
+					Position = TransformPosition(Position);
+					Mesh.SetVertex(vid, Position);
+				}
+			});
+	}
+
+};
+};
\ No newline at end of file
diff --git a/Engine/Plugins/Runtime/MeshModelingToolset/Source/ModelingComponents/Public/ConversionUtils/SplineComponentDeformDynamicMesh.h b/Engine/Plugins/Runtime/MeshModelingToolset/Source/ModelingComponents/Public/ConversionUtils/SplineComponentDeformDynamicMesh.h
new file mode 100644
--- /dev/null
+++ b/Engine/Plugins/Runtime/MeshModelingToolset/Source/ModelingComponents/Public/ConversionUtils/SplineComponentDeformDynamicMesh.h
@@ -0,0 +1,24 @@
+// Copyright Epic Games, Inc. All Rights Reserved.
+
+#pragma once
+
+#include "CoreMinimal.h"
+
+class USplineMeshComponent;
+
+namespace UE {
+namespace Geometry {
+
+
+class FDynamicMesh3;
+
+/**
+*  Applies a deformation as defined by the SplineMeshComponent to the dynamic mesh
+*  The normal, tangent, and bi-tangent will be updated if bUpdateTangentSpace is selected.  
+*  NB: Superior tangents will be obtained by recomputing them after this deformation; 
+*      while the tangent space update this function does is not ideal, it should be consistent with the SplineMeshComponent GPU computation
+*/ 
+void MODELINGCOMPONENTS_API SplineDeformDynamicMesh(USplineMeshComponent& SplineMeshComponent, FDynamicMesh3& Mesh, bool bUpdateTangentSpace = true);
+
+}; // end namespace Geometry
+}; // end namespace UE
