ModelingMode: Selection system improvements. Add Disconnect command and various new Selection-Edit Commands

Add UInteractiveCommandResult, UInteractiveCommand::ExecuteCommand() now optionally can return a result subclass
Add UGeometrySelectionEditCommandResult, UGeometrySelectionEditCommand now can optionally return an "output" selection via this type
UGeometrySelectionManager::ExecuteSelectionCommand optionally can restore a selection after a command, via UGeometrySelectionEditCommandResult

Add UDisconnectGeometrySelectionCommand, implements disconnection of selected triangles (ie separates but not into a new mesh)
Add UModifyGeometrySelectionCommand, implements various selection edits (select all, expand to connected, invert, invert connected, expand, contract)
Add IGeometrySelector::InitializeSelectionFromPredicate() and ::UpdateSelectionFromSelection(), implement in UDynamicMeshSelector, used to implement selection edit commands
Add UI to enable new commands in Modeling Mode

diff --git a/Engine/Plugins/Editor/ModelingToolsEditorMode/Source/ModelingToolsEditorMode/Private/ModelingModeToolkit_Toolbars.cpp b/Engine/Plugins/Editor/ModelingToolsEditorMode/Source/ModelingToolsEditorMode/Private/ModelingModeToolkit_Toolbars.cpp
--- a/Engine/Plugins/Editor/ModelingToolsEditorMode/Source/ModelingToolsEditorMode/Private/ModelingModeToolkit_Toolbars.cpp
+++ b/Engine/Plugins/Editor/ModelingToolsEditorMode/Source/ModelingToolsEditorMode/Private/ModelingModeToolkit_Toolbars.cpp
@@ -382,6 +382,29 @@ TSharedRef<SWidget> MakeMenu_SelectionConfigSettings(FModelingToolsEditorModeToo
 	return MenuWidget;
 }
 
+
+
+
+TSharedRef<SWidget> MakeMenu_SelectionEdits(FModelingToolsEditorModeToolkit* Toolkit)
+{
+	FMenuBuilder MenuBuilder(true, Toolkit->GetToolkitCommands() );
+
+	MenuBuilder.BeginSection("Section_SelectionEdits", LOCTEXT("Section_SelectionEdits", "Selection Edits"));
+
+	const FModelingToolsManagerCommands& Commands = FModelingToolsManagerCommands::Get();
+
+	MenuBuilder.AddMenuEntry(Commands.BeginSelectionAction_SelectAll);
+	MenuBuilder.AddMenuEntry(Commands.BeginSelectionAction_ExpandToConnected);
+	MenuBuilder.AddMenuEntry(Commands.BeginSelectionAction_Invert);
+	MenuBuilder.AddMenuEntry(Commands.BeginSelectionAction_InvertConnected);
+	MenuBuilder.AddMenuEntry(Commands.BeginSelectionAction_Expand);
+	MenuBuilder.AddMenuEntry(Commands.BeginSelectionAction_Contract);
+
+	TSharedRef<SWidget> MenuWidget = MenuBuilder.MakeWidget();
+	return MenuWidget;
+}
+
+
 } // end namespace UELocal
 
 
@@ -406,6 +429,22 @@ void FModelingToolsEditorModeToolkit::MakeSelectionPaletteOverlayWidget()
 	ToolbarBuilder.AddToolBarButton(Commands.MeshSelectionModeAction_GroupEdges);
 	ToolbarBuilder.AddToolBarButton(Commands.MeshSelectionModeAction_GroupCorners);
 
+	ToolbarBuilder.AddWidget(
+		SNew(SComboButton)
+		.HasDownArrow(false)
+		.MenuPlacement(EMenuPlacement::MenuPlacement_MenuRight)
+		.ComboButtonStyle(FAppStyle::Get(), "SimpleComboButton")
+		.OnGetMenuContent(FOnGetContent::CreateLambda([this]()
+		{
+			return UELocal::MakeMenu_SelectionEdits(this);
+		}))
+		.ContentPadding(FMargin(1.0f, 1.0f))
+		.ButtonContent()
+		[
+			SNew(SImage)
+			.Image(FModelingToolsEditorModeStyle::Get()->GetBrush("ModelingModeSelection.Edits_Right"))
+		] );
+
 	ToolbarBuilder.AddWidget(
 		SNew(SComboButton)
 		.HasDownArrow(false)
@@ -422,6 +461,7 @@ void FModelingToolsEditorModeToolkit::MakeSelectionPaletteOverlayWidget()
 			.Image(FModelingToolsEditorModeStyle::Get()->GetBrush("ModelingModeSelection.More_Right"))
 		] );
 
+
 	ToolbarBuilder.AddWidget( SNew(SSpacer).Size(FVector2D(1,10)) );
 
 	ToolbarBuilder.AddWidget(
diff --git a/Engine/Plugins/Editor/ModelingToolsEditorMode/Source/ModelingToolsEditorMode/Private/ModelingToolsEditorMode.cpp b/Engine/Plugins/Editor/ModelingToolsEditorMode/Source/ModelingToolsEditorMode/Private/ModelingToolsEditorMode.cpp
--- a/Engine/Plugins/Editor/ModelingToolsEditorMode/Source/ModelingToolsEditorMode/Private/ModelingToolsEditorMode.cpp
+++ b/Engine/Plugins/Editor/ModelingToolsEditorMode/Source/ModelingToolsEditorMode/Private/ModelingToolsEditorMode.cpp
@@ -105,6 +105,8 @@
 
 // commands
 #include "Commands/DeleteGeometrySelectionCommand.h"
+#include "Commands/DisconnectGeometrySelectionCommand.h"
+#include "Commands/ModifyGeometrySelectionCommand.h"
 
 
 #include "EditorModeManager.h"
@@ -811,7 +813,13 @@ void UModelingToolsEditorMode::Enter()
 
 	// create and register InteractiveCommands for mesh selections
 	RegisterSelectionCommand(NewObject<UDeleteGeometrySelectionCommand>(), ToolManagerCommands.BeginSelectionAction_Delete);
-
+	RegisterSelectionCommand(NewObject<UDisconnectGeometrySelectionCommand>(), ToolManagerCommands.BeginSelectionAction_Disconnect);
+	RegisterSelectionCommand(NewObject<UModifyGeometrySelectionCommand>(), ToolManagerCommands.BeginSelectionAction_SelectAll);
+	RegisterSelectionCommand(NewObject<UModifyGeometrySelectionCommand_Invert>(), ToolManagerCommands.BeginSelectionAction_Invert);
+	RegisterSelectionCommand(NewObject<UModifyGeometrySelectionCommand_ExpandToConnected>(), ToolManagerCommands.BeginSelectionAction_ExpandToConnected);
+	RegisterSelectionCommand(NewObject<UModifyGeometrySelectionCommand_InvertConnected>(), ToolManagerCommands.BeginSelectionAction_InvertConnected);
+	RegisterSelectionCommand(NewObject<UModifyGeometrySelectionCommand_Expand>(), ToolManagerCommands.BeginSelectionAction_Expand);
+	RegisterSelectionCommand(NewObject<UModifyGeometrySelectionCommand_Contract>(), ToolManagerCommands.BeginSelectionAction_Contract);
 
 	// register extensions
 	TArray<IModelingModeToolExtension*> Extensions = IModularFeatures::Get().GetModularFeatureImplementations<IModelingModeToolExtension>(
diff --git a/Engine/Plugins/Editor/ModelingToolsEditorMode/Source/ModelingToolsEditorMode/Private/ModelingToolsEditorModeStyle.cpp b/Engine/Plugins/Editor/ModelingToolsEditorMode/Source/ModelingToolsEditorMode/Private/ModelingToolsEditorModeStyle.cpp
--- a/Engine/Plugins/Editor/ModelingToolsEditorMode/Source/ModelingToolsEditorMode/Private/ModelingToolsEditorModeStyle.cpp
+++ b/Engine/Plugins/Editor/ModelingToolsEditorMode/Source/ModelingToolsEditorMode/Private/ModelingToolsEditorModeStyle.cpp
@@ -332,6 +332,7 @@ void FModelingToolsEditorModeStyle::Initialize()
 		StyleSet->Set("ModelingToolsManagerCommands.BeginSelectionAction_Delete", new FSlateImageBrush(StyleSet->RootToCoreContentDir(TEXT("../Editor/Slate/Icons/GeneralTools/Delete_40x.png")), Icon20x20));
 
 		StyleSet->Set("ModelingModeSelection.More_Right",  new IMAGE_BRUSH_SVG("Icons/SelectionToolbar_More", Icon20x20));
+		StyleSet->Set("ModelingModeSelection.Edits_Right",  new IMAGE_BRUSH_SVG("Icons/SelectionToolbar_Edits", Icon20x20));
 
 
 
diff --git a/Engine/Plugins/Editor/ModelingToolsEditorMode/Source/ModelingToolsEditorMode/Private/ModelingToolsEditorModeToolkit.cpp b/Engine/Plugins/Editor/ModelingToolsEditorMode/Source/ModelingToolsEditorMode/Private/ModelingToolsEditorModeToolkit.cpp
--- a/Engine/Plugins/Editor/ModelingToolsEditorMode/Source/ModelingToolsEditorMode/Private/ModelingToolsEditorModeToolkit.cpp
+++ b/Engine/Plugins/Editor/ModelingToolsEditorMode/Source/ModelingToolsEditorMode/Private/ModelingToolsEditorModeToolkit.cpp
@@ -674,6 +674,7 @@ void FModelingToolsEditorModeToolkit::BuildToolPalette(FName PaletteIndex, class
 	else if (PaletteIndex == SelectionActionsTabName)
 	{
 		ToolbarBuilder.AddToolBarButton(Commands.BeginSelectionAction_Delete);
+		ToolbarBuilder.AddToolBarButton(Commands.BeginSelectionAction_Disconnect);
 		ToolbarBuilder.AddToolBarButton(Commands.BeginSelectionAction_Extrude);
 	}
 	else if (PaletteIndex == PrimitiveTabName)
diff --git a/Engine/Plugins/Editor/ModelingToolsEditorMode/Source/ModelingToolsEditorMode/Private/ModelingToolsManagerActions.cpp b/Engine/Plugins/Editor/ModelingToolsEditorMode/Source/ModelingToolsEditorMode/Private/ModelingToolsManagerActions.cpp
--- a/Engine/Plugins/Editor/ModelingToolsEditorMode/Source/ModelingToolsEditorMode/Private/ModelingToolsManagerActions.cpp
+++ b/Engine/Plugins/Editor/ModelingToolsEditorMode/Source/ModelingToolsEditorMode/Private/ModelingToolsManagerActions.cpp
@@ -46,6 +46,11 @@ void FModelingToolsManagerCommands::RegisterCommands()
 		UI_COMMAND(ToolCommandInfo, ToolName, ToolTip, EUserInterfaceActionType::ToggleButton, FInputChord()); \
 		RegisteredTools.Add(FStartToolCommand{ ToolName, ToolCommandInfo });
 
+	// this has to be done with a compile-time macro because UI_COMMAND expands to LOCTEXT macros
+#define REGISTER_MODELING_ACTION_COMMAND(ToolCommandInfo, ToolName, ToolTip) \
+		UI_COMMAND(ToolCommandInfo, ToolName, ToolTip, EUserInterfaceActionType::Button, FInputChord()); \
+		RegisteredTools.Add(FStartToolCommand{ ToolName, ToolCommandInfo });
+
 	// Shapes
 	REGISTER_MODELING_TOOL_COMMAND(BeginAddBoxPrimitiveTool, "Box", "Create new box objects");
 	REGISTER_MODELING_TOOL_COMMAND(BeginAddSpherePrimitiveTool, "Sphere", "Create new sphere objects");
@@ -198,9 +203,20 @@ void FModelingToolsManagerCommands::RegisterCommands()
 	REGISTER_MODELING_TOOL_COMMAND(MeshSelectionModeAction_GroupCorners, "Corners", "Select Mesh Polygroup Corners/Vertices");
 	REGISTER_MODELING_TOOL_COMMAND(MeshSelectionModeAction_GroupEdges, "Borders", "Select Mesh Polygroup Borders/Edges");
 
-	REGISTER_MODELING_TOOL_COMMAND(BeginSelectionAction_Delete, "Delete", "Delete Selection");
+	REGISTER_MODELING_ACTION_COMMAND(BeginSelectionAction_Delete, "Delete", "Delete Selection");
+	REGISTER_MODELING_ACTION_COMMAND(BeginSelectionAction_Disconnect, "Discon", "Disconnect Selection");
+
 	REGISTER_MODELING_TOOL_COMMAND(BeginSelectionAction_Extrude, "Extrude", "Extrude Selection");
 
+	REGISTER_MODELING_ACTION_COMMAND(BeginSelectionAction_SelectAll, "Select All", "Select All Elements");
+	REGISTER_MODELING_ACTION_COMMAND(BeginSelectionAction_ExpandToConnected, "Expand To Connected", "Expand Selection to Geometrically-Connected Elements");
+	REGISTER_MODELING_ACTION_COMMAND(BeginSelectionAction_Invert, "Invert Selection", "Invert the current Selection");
+	REGISTER_MODELING_ACTION_COMMAND(BeginSelectionAction_InvertConnected, "Invert Connected", "Invert the current Selection to Geometrically-Connected Elements");
+	REGISTER_MODELING_ACTION_COMMAND(BeginSelectionAction_Expand, "Expand Selection", "Expand Selection");
+	REGISTER_MODELING_ACTION_COMMAND(BeginSelectionAction_Contract, "Contract Selection", "Contract Selection");
+
+
+
 	UI_COMMAND(AcceptActiveTool, "Accept", "Accept the active tool", EUserInterfaceActionType::Button, FInputChord());
 	UI_COMMAND(CancelActiveTool, "Cancel", "Cancel the active tool", EUserInterfaceActionType::Button, FInputChord());
 	UI_COMMAND(CompleteActiveTool, "Done", "Complete the active tool", EUserInterfaceActionType::Button, FInputChord());
diff --git a/Engine/Plugins/Editor/ModelingToolsEditorMode/Source/ModelingToolsEditorMode/Public/ModelingToolsManagerActions.h b/Engine/Plugins/Editor/ModelingToolsEditorMode/Source/ModelingToolsEditorMode/Public/ModelingToolsManagerActions.h
--- a/Engine/Plugins/Editor/ModelingToolsEditorMode/Source/ModelingToolsEditorMode/Public/ModelingToolsManagerActions.h
+++ b/Engine/Plugins/Editor/ModelingToolsEditorMode/Source/ModelingToolsEditorMode/Public/ModelingToolsManagerActions.h
@@ -163,6 +163,15 @@ public:
 
 
 	TSharedPtr<FUICommandInfo> BeginSelectionAction_Delete;
+	TSharedPtr<FUICommandInfo> BeginSelectionAction_Disconnect;
+
+	TSharedPtr<FUICommandInfo> BeginSelectionAction_SelectAll;
+	TSharedPtr<FUICommandInfo> BeginSelectionAction_ExpandToConnected;
+	TSharedPtr<FUICommandInfo> BeginSelectionAction_Invert;
+	TSharedPtr<FUICommandInfo> BeginSelectionAction_InvertConnected;
+	TSharedPtr<FUICommandInfo> BeginSelectionAction_Expand;
+	TSharedPtr<FUICommandInfo> BeginSelectionAction_Contract;
+
 	TSharedPtr<FUICommandInfo> BeginSelectionAction_Extrude;
 
 	// For connecting ModelingMode with UVEditor plugin
diff --git a/Engine/Plugins/Runtime/MeshModelingToolset/Source/MeshModelingTools/Private/Commands/DeleteGeometrySelectionCommand.cpp b/Engine/Plugins/Runtime/MeshModelingToolset/Source/MeshModelingTools/Private/Commands/DeleteGeometrySelectionCommand.cpp
--- a/Engine/Plugins/Runtime/MeshModelingToolset/Source/MeshModelingTools/Private/Commands/DeleteGeometrySelectionCommand.cpp
+++ b/Engine/Plugins/Runtime/MeshModelingToolset/Source/MeshModelingTools/Private/Commands/DeleteGeometrySelectionCommand.cpp
@@ -29,8 +29,14 @@ bool UDeleteGeometrySelectionCommand::CanExecuteCommandForSelection(UGeometrySel
 }
 
 
-void UDeleteGeometrySelectionCommand::ExecuteCommandForSelection(UGeometrySelectionEditCommandArguments* SelectionArgs)
+void UDeleteGeometrySelectionCommand::ExecuteCommandForSelection(UGeometrySelectionEditCommandArguments* SelectionArgs, UInteractiveCommandResult** Result)
 {
+	// delete never returns a new selection
+	if (Result != nullptr)
+	{
+		*Result = nullptr;
+	}
+
 	// should have been verified by CanExecute
 	check(SelectionArgs->IsMatchingType(FGeometryIdentifier::ETargetType::MeshContainer, FGeometryIdentifier::EObjectType::DynamicMesh));
 	
diff --git a/Engine/Plugins/Runtime/MeshModelingToolset/Source/MeshModelingTools/Private/Commands/DisconnectGeometrySelectionCommand.cpp b/Engine/Plugins/Runtime/MeshModelingToolset/Source/MeshModelingTools/Private/Commands/DisconnectGeometrySelectionCommand.cpp
new file mode 100644
--- /dev/null
+++ b/Engine/Plugins/Runtime/MeshModelingToolset/Source/MeshModelingTools/Private/Commands/DisconnectGeometrySelectionCommand.cpp
@@ -0,0 +1,104 @@
+// Copyright Epic Games, Inc. All Rights Reserved.
+
+#include "Commands/DisconnectGeometrySelectionCommand.h"
+#include "ToolContextInterfaces.h"
+#include "UDynamicMesh.h"
+#include "DynamicMesh/DynamicMesh3.h"
+#include "DynamicMesh/DynamicMeshChangeTracker.h"
+#include "DynamicMeshEditor.h"
+#include "Changes/MeshChange.h"
+#include "Selections/GeometrySelectionUtil.h"
+
+#include UE_INLINE_GENERATED_CPP_BY_NAME(DisconnectGeometrySelectionCommand)
+
+using namespace UE::Geometry;
+
+
+#define LOCTEXT_NAMESPACE "UDisconnectGeometrySelectionCommand"
+
+
+FText UDisconnectGeometrySelectionCommand::GetCommandShortString() const
+{
+	return LOCTEXT("ShortString", "Disconnect Selection");
+}
+
+
+bool UDisconnectGeometrySelectionCommand::CanExecuteCommandForSelection(UGeometrySelectionEditCommandArguments* SelectionArgs)
+{
+	return SelectionArgs->IsMatchingType(FGeometryIdentifier::ETargetType::MeshContainer, FGeometryIdentifier::EObjectType::DynamicMesh);
+}
+
+
+void UDisconnectGeometrySelectionCommand::ExecuteCommandForSelection(UGeometrySelectionEditCommandArguments* SelectionArgs, UInteractiveCommandResult** Result)
+{
+	if (Result != nullptr)
+	{
+		*Result = nullptr;
+	}
+
+	// should have been verified by CanExecute
+	check(SelectionArgs->IsMatchingType(FGeometryIdentifier::ETargetType::MeshContainer, FGeometryIdentifier::EObjectType::DynamicMesh));
+	
+	// collect up all our inputs
+	UDynamicMesh* MeshObject = SelectionArgs->SelectionHandle.Identifier.GetAsObjectType<UDynamicMesh>();
+	check(MeshObject != nullptr);
+	const FGeometrySelection* Selection = SelectionArgs->SelectionHandle.Selection;
+	if (Selection->Selection.Num() == 0)
+	{
+		return;
+	}
+
+	bool bTrackChanges = SelectionArgs->HasTransactionsAPI();
+	TUniquePtr<FDynamicMeshChange> DynamicMeshChange;	// only initialized if bTrackChanges == true
+
+	// apply the Disconnect operation
+	MeshObject->EditMesh([&](FDynamicMesh3& EditMesh)
+	{
+		// build list of triangles from whatever the selection contains
+		TSet<int32> TriangleList;
+		//UE::Geometry::FPolygroupSet UsePolygroupSet = ...;		// need to support this eventually
+		UE::Geometry::EnumerateSelectionTriangles(*Selection, EditMesh,
+			[&](int32 TriangleID) { TriangleList.Add(TriangleID); });
+
+		// mark triangles for change
+		FDynamicMeshChangeTracker ChangeTracker(&EditMesh);
+		if (bTrackChanges)
+		{
+			ChangeTracker.BeginChange();
+			ChangeTracker.SaveTriangles(TriangleList, true);
+		}
+
+		// actually Disconnect them
+		FDynamicMeshEditor Editor(&EditMesh);
+		Editor.DisconnectTriangles(TriangleList.Array(), false);
+
+		// extract the change record
+		if (bTrackChanges)
+		{
+			DynamicMeshChange = ChangeTracker.EndChange();
+		}
+
+	}, EDynamicMeshChangeType::GeneralEdit, EDynamicMeshAttributeChangeFlags::Unknown, false);
+
+
+	// emit change 
+	if ( bTrackChanges && DynamicMeshChange.IsValid() )
+	{
+		SelectionArgs->GetTransactionsAPI()->BeginUndoTransaction(GetCommandShortString());
+		SelectionArgs->GetTransactionsAPI()->AppendChange(MeshObject, 
+			MakeUnique<FMeshChange>(MoveTemp(DynamicMeshChange)), GetCommandShortString());
+		SelectionArgs->GetTransactionsAPI()->EndUndoTransaction();
+	}
+
+	if (Result != nullptr)
+	{
+		UGeometrySelectionEditCommandResult* NewResult = NewObject<UGeometrySelectionEditCommandResult>();
+		NewResult->SourceHandle = SelectionArgs->SelectionHandle;
+		NewResult->OutputSelection = *NewResult->SourceHandle.Selection;
+		*Result = NewResult;
+	}
+
+}
+
+
+#undef LOCTEXT_NAMESPACE
diff --git a/Engine/Plugins/Runtime/MeshModelingToolset/Source/MeshModelingTools/Private/Commands/ModifyGeometrySelectionCommand.cpp b/Engine/Plugins/Runtime/MeshModelingToolset/Source/MeshModelingTools/Private/Commands/ModifyGeometrySelectionCommand.cpp
new file mode 100644
--- /dev/null
+++ b/Engine/Plugins/Runtime/MeshModelingToolset/Source/MeshModelingTools/Private/Commands/ModifyGeometrySelectionCommand.cpp
@@ -0,0 +1,134 @@
+// Copyright Epic Games, Inc. All Rights Reserved.
+
+#include "Commands/ModifyGeometrySelectionCommand.h"
+#include "ToolContextInterfaces.h"
+#include "UDynamicMesh.h"
+#include "DynamicMesh/DynamicMesh3.h"
+#include "DynamicMesh/DynamicMeshChangeTracker.h"
+#include "DynamicMeshEditor.h"
+#include "Changes/MeshChange.h"
+#include "Selections/GeometrySelectionUtil.h"
+
+#include UE_INLINE_GENERATED_CPP_BY_NAME(ModifyGeometrySelectionCommand)
+
+using namespace UE::Geometry;
+
+
+#define LOCTEXT_NAMESPACE "UModifyGeometrySelectionCommand"
+
+
+FText UModifyGeometrySelectionCommand::GetCommandShortString() const
+{
+	switch (GetModificationType())
+	{
+	default:
+	case EModificationType::SelectAll:
+		return LOCTEXT("ShortString_SelectAll", "Select All");
+	case EModificationType::ExpandToConnected:
+		return LOCTEXT("ShortString_ExpandToConnected", "Expand To Connected");
+		
+	case EModificationType::Invert:
+		return LOCTEXT("ShortString_Invert", "Invert Selection");
+	case EModificationType::InvertConnected:
+		return LOCTEXT("ShortString_InvertConnected", "Invert Selection (Connected)");
+
+	case EModificationType::Expand:
+		return LOCTEXT("ShortString_Expand", "Expand/Grow Selection");
+	case EModificationType::Contract:
+		return LOCTEXT("ShortString_Contract", "Contract/Shrink Selection");
+
+	}
+
+	return LOCTEXT("ShortString_SelectAll", "Select All");
+}
+
+
+bool UModifyGeometrySelectionCommand::CanExecuteCommandForSelection(UGeometrySelectionEditCommandArguments* SelectionArgs)
+{
+	return SelectionArgs->IsMatchingType(FGeometryIdentifier::ETargetType::MeshContainer, FGeometryIdentifier::EObjectType::DynamicMesh);
+}
+
+
+void UModifyGeometrySelectionCommand::ExecuteCommandForSelection(UGeometrySelectionEditCommandArguments* SelectionArgs, UInteractiveCommandResult** Result)
+{
+	if (!ensure(Result != nullptr))
+	{
+		UE_LOG(LogGeometry, Warning, TEXT("UModifyGeometrySelectionCommand: Selection-Modify commands are no-ops unless Result is provided and used"));
+		return;
+	}
+	if (!ensure(SelectionArgs->SelectionHandle.Selector != nullptr))
+	{
+		UE_LOG(LogGeometry, Warning, TEXT("UModifyGeometrySelectionCommand: Selection-Modify commands require Selector be provided in Selection Arguments"));
+		return;
+	}
+
+	// should have been verified by CanExecute
+	check(SelectionArgs->IsMatchingType(FGeometryIdentifier::ETargetType::MeshContainer, FGeometryIdentifier::EObjectType::DynamicMesh));
+	
+	// collect up all our inputs
+	UDynamicMesh* MeshObject = SelectionArgs->SelectionHandle.Identifier.GetAsObjectType<UDynamicMesh>();
+	check(MeshObject != nullptr);
+	const FGeometrySelection* InitialSelection = SelectionArgs->SelectionHandle.Selection;
+
+	// some might be no-ops?
+	//if (Selection->Selection.Num() == 0)
+	//{
+	//	return;
+	//}
+
+	UGeometrySelectionEditCommandResult* NewResult = NewObject<UGeometrySelectionEditCommandResult>();
+	NewResult->SourceHandle = SelectionArgs->SelectionHandle;
+	NewResult->OutputSelection.InitializeTypes(*InitialSelection);
+
+	if (GetModificationType() == EModificationType::SelectAll)
+	{
+		SelectionArgs->SelectionHandle.Selector->InitializeSelectionFromPredicate(
+			NewResult->OutputSelection, [&](FGeoSelectionID) { return true;}, IGeometrySelector::EInitializeSelectionMode::All, nullptr );
+	}
+	else if (GetModificationType() == EModificationType::ExpandToConnected)
+	{
+		SelectionArgs->SelectionHandle.Selector->InitializeSelectionFromPredicate(
+			NewResult->OutputSelection, [&](FGeoSelectionID) { return true;}, 
+			IGeometrySelector::EInitializeSelectionMode::Connected, InitialSelection );
+	}
+	else if (GetModificationType() == EModificationType::Invert)
+	{
+		FGeometrySelectionEditor TmpReadOnlyEditor;
+		TmpReadOnlyEditor.Initialize(const_cast<FGeometrySelection*>(InitialSelection), InitialSelection->TopologyType == EGeometryTopologyType::Polygroup);		// not actually going to edit, just use selected test
+		SelectionArgs->SelectionHandle.Selector->InitializeSelectionFromPredicate(
+			NewResult->OutputSelection, 
+			[&](FGeoSelectionID SelectionID) { return TmpReadOnlyEditor.IsSelected(SelectionID.Encoded()) == false; },
+			IGeometrySelector::EInitializeSelectionMode::All, nullptr);
+	}
+	else if (GetModificationType() == EModificationType::InvertConnected)
+	{
+		SelectionArgs->SelectionHandle.Selector->InitializeSelectionFromPredicate(
+			NewResult->OutputSelection, [&](FGeoSelectionID) { return true;}, 
+			IGeometrySelector::EInitializeSelectionMode::Connected, InitialSelection );
+		UE::Geometry::CombineSelectionInPlace(NewResult->OutputSelection, *InitialSelection, UE::Geometry::EGeometrySelectionCombineModes::Subtract);
+	}
+	else if (GetModificationType() == EModificationType::Expand)
+	{
+		// create border-ring selection
+		SelectionArgs->SelectionHandle.Selector->InitializeSelectionFromPredicate(
+			NewResult->OutputSelection, [&](FGeoSelectionID) { return true;}, 
+			IGeometrySelector::EInitializeSelectionMode::AdjacentToBorder, InitialSelection );
+		UE::Geometry::CombineSelectionInPlace(NewResult->OutputSelection, *InitialSelection, UE::Geometry::EGeometrySelectionCombineModes::Add);
+	}
+	else if (GetModificationType() == EModificationType::Contract)
+	{
+		// create border-ring selection
+		FGeometrySelection TmpSelection;
+		TmpSelection.InitializeTypes(NewResult->OutputSelection);
+		SelectionArgs->SelectionHandle.Selector->InitializeSelectionFromPredicate(
+			TmpSelection, [&](FGeoSelectionID) { return true;}, 
+			IGeometrySelector::EInitializeSelectionMode::AdjacentToBorder, InitialSelection );
+		NewResult->OutputSelection = *InitialSelection;
+		UE::Geometry::CombineSelectionInPlace(NewResult->OutputSelection, TmpSelection, UE::Geometry::EGeometrySelectionCombineModes::Subtract);
+	}
+
+	*Result = NewResult;
+}
+
+
+#undef LOCTEXT_NAMESPACE
diff --git a/Engine/Plugins/Runtime/MeshModelingToolset/Source/MeshModelingTools/Public/Commands/DeleteGeometrySelectionCommand.h b/Engine/Plugins/Runtime/MeshModelingToolset/Source/MeshModelingTools/Public/Commands/DeleteGeometrySelectionCommand.h
--- a/Engine/Plugins/Runtime/MeshModelingToolset/Source/MeshModelingTools/Public/Commands/DeleteGeometrySelectionCommand.h
+++ b/Engine/Plugins/Runtime/MeshModelingToolset/Source/MeshModelingTools/Public/Commands/DeleteGeometrySelectionCommand.h
@@ -20,5 +20,5 @@ public:
 	virtual FText GetCommandShortString() const override;
 
 	virtual bool CanExecuteCommandForSelection(UGeometrySelectionEditCommandArguments* Arguments) override;
-	virtual void ExecuteCommandForSelection(UGeometrySelectionEditCommandArguments* Arguments) override;
+	virtual void ExecuteCommandForSelection(UGeometrySelectionEditCommandArguments* Arguments, UInteractiveCommandResult** Result) override;
 };
\ No newline at end of file
diff --git a/Engine/Plugins/Runtime/MeshModelingToolset/Source/MeshModelingTools/Public/Commands/DisconnectGeometrySelectionCommand.h b/Engine/Plugins/Runtime/MeshModelingToolset/Source/MeshModelingTools/Public/Commands/DisconnectGeometrySelectionCommand.h
new file mode 100644
--- /dev/null
+++ b/Engine/Plugins/Runtime/MeshModelingToolset/Source/MeshModelingTools/Public/Commands/DisconnectGeometrySelectionCommand.h
@@ -0,0 +1,24 @@
+// Copyright Epic Games, Inc. All Rights Reserved.
+
+#pragma once
+
+#include "Selection/SelectionEditInteractiveCommand.h"
+#include "DisconnectGeometrySelectionCommand.generated.h"
+
+
+/**
+ * UDisconnectGeometrySelectionCommand disconnects the geometric elements identified by the Selection.
+ * Currently only supports mesh selections (Triangle and Polygroup types)
+ * Disconnects selected faces, or faces connected to selected edges, or faces connected to selected vertices.
+ */
+UCLASS()
+class MESHMODELINGTOOLS_API UDisconnectGeometrySelectionCommand : public UGeometrySelectionEditCommand
+{
+	GENERATED_BODY()
+public:
+
+	virtual FText GetCommandShortString() const override;
+
+	virtual bool CanExecuteCommandForSelection(UGeometrySelectionEditCommandArguments* Arguments) override;
+	virtual void ExecuteCommandForSelection(UGeometrySelectionEditCommandArguments* Arguments, UInteractiveCommandResult** Result) override;
+};
\ No newline at end of file
diff --git a/Engine/Plugins/Runtime/MeshModelingToolset/Source/MeshModelingTools/Public/Commands/ModifyGeometrySelectionCommand.h b/Engine/Plugins/Runtime/MeshModelingToolset/Source/MeshModelingTools/Public/Commands/ModifyGeometrySelectionCommand.h
new file mode 100644
--- /dev/null
+++ b/Engine/Plugins/Runtime/MeshModelingToolset/Source/MeshModelingTools/Public/Commands/ModifyGeometrySelectionCommand.h
@@ -0,0 +1,96 @@
+// Copyright Epic Games, Inc. All Rights Reserved.
+
+#pragma once
+
+#include "Selection/SelectionEditInteractiveCommand.h"
+#include "ModifyGeometrySelectionCommand.generated.h"
+
+
+/**
+ * UModifyGeometrySelectionCommand updates/edits the current selection in various ways.
+ * Default operation is to Select All.
+ * Subclasses below can be used in situations where specific per-modification types are needed.
+ */
+UCLASS()
+class MESHMODELINGTOOLS_API UModifyGeometrySelectionCommand : public UGeometrySelectionEditCommand
+{
+	GENERATED_BODY()
+public:
+
+	enum class EModificationType
+	{
+		SelectAll = 0,
+		ExpandToConnected = 1,
+
+		Invert = 10,
+		InvertConnected = 11,
+
+		Expand = 20,
+		Contract = 21
+	};
+	virtual EModificationType GetModificationType() const { return EModificationType::SelectAll; }
+
+	virtual bool IsModifySelectionCommand() const { return true; }
+	virtual FText GetCommandShortString() const override;
+
+	virtual bool CanExecuteCommandForSelection(UGeometrySelectionEditCommandArguments* Arguments) override;
+	virtual void ExecuteCommandForSelection(UGeometrySelectionEditCommandArguments* Arguments, UInteractiveCommandResult** Result) override;
+};
+
+
+/**
+ * Command to Invert the current Selection
+ */
+UCLASS()
+class MESHMODELINGTOOLS_API UModifyGeometrySelectionCommand_Invert : public UModifyGeometrySelectionCommand
+{
+	GENERATED_BODY()
+public:
+	virtual EModificationType GetModificationType() const override { return EModificationType::Invert; }
+
+};
+
+/**
+ * Command to Expand the current Selection to all connected geometry
+ */
+UCLASS()
+class MESHMODELINGTOOLS_API UModifyGeometrySelectionCommand_ExpandToConnected : public UModifyGeometrySelectionCommand
+{
+	GENERATED_BODY()
+public:
+	virtual EModificationType GetModificationType() const override { return EModificationType::ExpandToConnected; }
+
+};
+
+/**
+ * Command to Invert the current Selection, only considering connected geometry
+ */
+UCLASS()
+class MESHMODELINGTOOLS_API UModifyGeometrySelectionCommand_InvertConnected : public UModifyGeometrySelectionCommand
+{
+	GENERATED_BODY()
+public:
+	virtual EModificationType GetModificationType() const override { return EModificationType::InvertConnected; }
+};
+
+/**
+ * Command to Expand the current Selection by a one-ring
+ */
+UCLASS()
+class MESHMODELINGTOOLS_API UModifyGeometrySelectionCommand_Expand : public UModifyGeometrySelectionCommand
+{
+	GENERATED_BODY()
+public:
+	virtual EModificationType GetModificationType() const override { return EModificationType::Expand; }
+};
+
+/**
+ * Command to Contract the current Selection by a one-ring
+ */
+UCLASS()
+class MESHMODELINGTOOLS_API UModifyGeometrySelectionCommand_Contract : public UModifyGeometrySelectionCommand
+{
+	GENERATED_BODY()
+public:
+	virtual EModificationType GetModificationType() const override { return EModificationType::Contract; }
+};
\ No newline at end of file
diff --git a/Engine/Plugins/Runtime/MeshModelingToolset/Source/ModelingComponents/Private/Selection/DynamicMeshSelector.cpp b/Engine/Plugins/Runtime/MeshModelingToolset/Source/ModelingComponents/Private/Selection/DynamicMeshSelector.cpp
--- a/Engine/Plugins/Runtime/MeshModelingToolset/Source/ModelingComponents/Private/Selection/DynamicMeshSelector.cpp
+++ b/Engine/Plugins/Runtime/MeshModelingToolset/Source/ModelingComponents/Private/Selection/DynamicMeshSelector.cpp
@@ -109,9 +109,69 @@ bool FBaseDynamicMeshSelector::Restore()
 }
 
 
+void FBaseDynamicMeshSelector::InitializeSelectionFromPredicate(
+	FGeometrySelection& SelectionOut,
+	TFunctionRef<bool(UE::Geometry::FGeoSelectionID)> SelectionIDPredicate,
+	EInitializeSelectionMode InitializeMode,
+	const FGeometrySelection* ReferenceSelection)
+{
+	const FGroupTopology* UseGroupTopology = ( SelectionOut.TopologyType == EGeometryTopologyType::Polygroup ) ? 
+		GetGroupTopology() : nullptr;
+
+	if (InitializeMode != EInitializeSelectionMode::All && ! ensure(ReferenceSelection != nullptr) )
+	{
+		return;
+	}
+
+	GetDynamicMesh()->ProcessMesh([&](const FDynamicMesh3& Mesh) 
+	{ 
+		if (InitializeMode == EInitializeSelectionMode::Connected)
+		{
+			UE::Geometry::MakeSelectAllConnectedSelection(Mesh, UseGroupTopology, *ReferenceSelection,
+				SelectionIDPredicate, 
+				[](FGeoSelectionID, FGeoSelectionID) { return true; },
+				SelectionOut);
+		}
+		else if (InitializeMode == EInitializeSelectionMode::AdjacentToBorder)
+		{
+			UE::Geometry::FGeometrySelectionEditor TmpReadOnlyEditor;
+			TmpReadOnlyEditor.Initialize( const_cast<FGeometrySelection*>(ReferenceSelection), 
+				(ReferenceSelection->TopologyType == EGeometryTopologyType::Polygroup) );
+			UE::Geometry::MakeBoundaryConnectedSelection(Mesh, UseGroupTopology, *ReferenceSelection,
+				SelectionIDPredicate, SelectionOut);			
+		}
+		else   // All
+		{
+			UE::Geometry::MakeSelectAllSelection(Mesh, UseGroupTopology, SelectionIDPredicate, SelectionOut);
+		}
+	});
+}
 
 
 
+void FBaseDynamicMeshSelector::UpdateSelectionFromSelection(
+	const FGeometrySelection& FromSelection,
+	bool bAllowConversion,
+	FGeometrySelectionEditor& SelectionEditor,
+	const FGeometrySelectionUpdateConfig& UpdateConfig,
+	FGeometrySelectionDelta* SelectionDelta )
+{
+	if ( FromSelection.IsSameType(SelectionEditor.GetSelection()) )
+	{
+		UE::Geometry::UpdateSelectionWithNewElements(&SelectionEditor, UpdateConfig.ChangeType,
+			FromSelection.Selection.Array(), SelectionDelta);
+		return;
+	}
+
+	if (bAllowConversion == false)
+	{
+		return;
+	}
+
+	check(false);		// todo: convert selection
+}
+
+
 
 bool FBaseDynamicMeshSelector::RayHitTest(
 	const FWorldRayQueryInfo& RayInfo,
@@ -230,7 +290,8 @@ void FBaseDynamicMeshSelector::UpdateSelectionViaRaycast_GroupEdges(
 					{
 						FMeshTriEdgeID TriEdgeID;
 						GetDynamicMesh()->ProcessMesh([&](const FDynamicMesh3& Mesh) { TriEdgeID = Mesh.GetTriEdgeIDFromEdgeID(Segment.ID); });
-						ResultOut.bSelectionModified = UpdateSelectionWithNewElements(&SelectionEditor, UpdateConfig.ChangeType, TArray<uint64>{(uint64)TriEdgeID.Encoded()}, & ResultOut.SelectionDelta);
+						FGeoSelectionID SelectionID(TriEdgeID.Encoded(), GroupEdgeID);
+						ResultOut.bSelectionModified = UpdateSelectionWithNewElements(&SelectionEditor, UpdateConfig.ChangeType, TArray<uint64>{SelectionID.Encoded()}, & ResultOut.SelectionDelta);
 						ResultOut.bSelectionMissed = false;
 					}
 				}
diff --git a/Engine/Plugins/Runtime/MeshModelingToolset/Source/ModelingComponents/Private/Selection/GeometrySelectionManager.cpp b/Engine/Plugins/Runtime/MeshModelingToolset/Source/ModelingComponents/Private/Selection/GeometrySelectionManager.cpp
--- a/Engine/Plugins/Runtime/MeshModelingToolset/Source/ModelingComponents/Private/Selection/GeometrySelectionManager.cpp
+++ b/Engine/Plugins/Runtime/MeshModelingToolset/Source/ModelingComponents/Private/Selection/GeometrySelectionManager.cpp
@@ -625,7 +625,7 @@ void UGeometrySelectionManager::AccumulateSelectionUpdate_Raycast(
 	QueryAPI->GetCurrentViewState(RayQueryInfo.CameraState);
 
 	Target->Selector->UpdateSelectionViaRaycast(
-		RayQueryInfo, *Target->SelectionEditor, ActiveTrackedUpdateConfig, ResultOut	);
+		RayQueryInfo, *Target->SelectionEditor, ActiveTrackedUpdateConfig, ResultOut );
 
 	if (ResultOut.bSelectionModified)
 	{
@@ -885,6 +885,9 @@ void UGeometrySelectionManager::ExecuteSelectionCommand(UGeometrySelectionEditCo
 	{
 		if (Target->Selection.IsEmpty()) continue;
 
+		// TODO: can use Command->IsModifySelectionCommand() to check if this is a command that only affects selection
+		// and not geometry. In that case we can skip the intermediate clear-selection and emit a single change.
+
 		// When initially executing the command, we do not clear the selection, because we pass it to the command.
 		// However, when we later *undo* any changes emitted by the command, we need to restore the selection aftewards.
 		// So we emit a clearing change here, so that undo un-clears. 
@@ -899,17 +902,34 @@ void UGeometrySelectionManager::ExecuteSelectionCommand(UGeometrySelectionEditCo
 		GetTransactionsAPI()->AppendChange(this, MoveTemp(ClearChange), LOCTEXT("ClearSelection", "Clear Selection"));
 
 		// q: we could clear the selection here, and pass the Handle a copy. Perhaps safer?
-
-		SelectionArguments->SelectionHandle = FGeometrySelectionHandle{ Target->Selector->GetIdentifier(), &Target->Selection };
+		UInteractiveCommandResult* ResultPtr = nullptr;
+		SelectionArguments->SelectionHandle = FGeometrySelectionHandle{ Target->Selector->GetIdentifier(), &Target->Selection, Target->Selector.Get() };
 		SelectionArguments->SetTransactionsAPI(TransactionsAPI);
-		Command->ExecuteCommand(SelectionArguments);
+		Command->ExecuteCommand(SelectionArguments, &ResultPtr);
 
-		// clear selection after executing command
+		// actually clear selection after executing command. 
 		FGeometrySelectionDelta ClearDelta;
 		Target->SelectionEditor->ClearSelection(ClearDelta);
 
-		// 
-		// if commands could emit new selections, this is where we would set them up and emit changes!
+		// if selection returned a result, and it was a non-empty selection, select it
+		if ( UGeometrySelectionEditCommandResult* SelectionResult = Cast<UGeometrySelectionEditCommandResult>(ResultPtr) )
+		{
+			if (SelectionResult->OutputSelection.IsEmpty() == false )
+			{
+				FGeometrySelectionDelta AfterDelta;
+				Target->Selector->UpdateSelectionFromSelection(
+					SelectionResult->OutputSelection, true, *Target->SelectionEditor, 
+					FGeometrySelectionUpdateConfig{EGeometrySelectionChangeType::Add}, &AfterDelta);
+				if (Target->Selection.IsEmpty() == false)
+				{
+					TUniquePtr<FGeometrySelectionReplaceChange> NewSelectionChange = MakeUnique<FGeometrySelectionReplaceChange>();
+					NewSelectionChange->Identifier = Target->Selector->GetIdentifier();
+					NewSelectionChange->After = Target->Selection;
+					NewSelectionChange->Before.InitializeTypes(Target->Selection);
+					GetTransactionsAPI()->AppendChange(this, MoveTemp(NewSelectionChange), LOCTEXT("NewSelection", "New Selection"));
+				}
+			}
+		}
 	}
 
 	GetTransactionsAPI()->EndUndoTransaction();
@@ -929,6 +949,7 @@ void UGeometrySelectionManager::ProcessActiveSelections(TFunctionRef<void(FGeome
 			FGeometrySelectionHandle Handle;
 			Handle.Selection = & Target->Selection;
 			Handle.Identifier = Target->Selector->GetIdentifier();
+			Handle.Selector = Target->Selector.Get();
 			ProcessFunc(Handle);
 		}
 	}
diff --git a/Engine/Plugins/Runtime/MeshModelingToolset/Source/ModelingComponents/Public/Selection/DynamicMeshSelector.h b/Engine/Plugins/Runtime/MeshModelingToolset/Source/ModelingComponents/Public/Selection/DynamicMeshSelector.h
--- a/Engine/Plugins/Runtime/MeshModelingToolset/Source/ModelingComponents/Public/Selection/DynamicMeshSelector.h
+++ b/Engine/Plugins/Runtime/MeshModelingToolset/Source/ModelingComponents/Public/Selection/DynamicMeshSelector.h
@@ -67,6 +67,22 @@ public:
 		return Identifier;
 	}
 
+	virtual void InitializeSelectionFromPredicate(	
+		FGeometrySelection& SelectionInOut,
+		TFunctionRef<bool(UE::Geometry::FGeoSelectionID)> SelectionIDPredicate,
+		EInitializeSelectionMode InitializeMode = EInitializeSelectionMode::All,
+		const FGeometrySelection* ReferenceSelection = nullptr) override;
+
+
+
+	virtual void UpdateSelectionFromSelection(	
+		const FGeometrySelection& FromSelection,
+		bool bAllowConversion,
+		FGeometrySelectionEditor& SelectionEditor,
+		const FGeometrySelectionUpdateConfig& UpdateConfig,
+		UE::Geometry::FGeometrySelectionDelta* SelectionDelta = nullptr ) override;
+
+
 	virtual bool RayHitTest(
 		const FWorldRayQueryInfo& RayInfo,
 		UE::Geometry::FGeometrySelectionHitQueryConfig QueryConfig,
diff --git a/Engine/Plugins/Runtime/MeshModelingToolset/Source/ModelingComponents/Public/Selection/GeometrySelector.h b/Engine/Plugins/Runtime/MeshModelingToolset/Source/ModelingComponents/Public/Selection/GeometrySelector.h
--- a/Engine/Plugins/Runtime/MeshModelingToolset/Source/ModelingComponents/Public/Selection/GeometrySelector.h
+++ b/Engine/Plugins/Runtime/MeshModelingToolset/Source/ModelingComponents/Public/Selection/GeometrySelector.h
@@ -125,6 +125,9 @@ struct MODELINGCOMPONENTS_API FGeometryIdentifier
 
 
 
+class IGeometrySelector;
+
+
 /**
  * FGeometrySelectionHandle stores a Selection and an Identifier for
  * the geometry object that the Selection is defined relative to
@@ -133,10 +136,11 @@ struct MODELINGCOMPONENTS_API FGeometrySelectionHandle
 {
 	FGeometryIdentifier Identifier;
 	const UE::Geometry::FGeometrySelection* Selection;
-};
 
+	/** optional Selector for SelectionHandle */
+ 	IGeometrySelector* Selector = nullptr;
+};
 
-class IGeometrySelector;
 
 
 /**
@@ -257,6 +261,51 @@ public:
 	 */
 	virtual FGeometryIdentifier GetIdentifier() const = 0;
 
+
+	enum class EInitializeSelectionMode
+	{
+		// create a selection of all geometry elements
+		All,
+		// create a selection of all geometry elements connected to a ReferenceSelection
+		Connected,
+		// create a selection of all geometry elements directly connected to a ReferenceSelection (ie adjacent/one-ring/etc)
+		AdjacentToBorder
+	};
+
+	/**
+	 * Populate a Selection using the provided ReferenceMode, with a Predicate and optional ReferenceSelection.
+	 * Note that this method does not modify/access any selection that might be "owned" by the Selector.
+	 * It is part of the IGeometrySelector API because often a Selector has access to internal data structures
+	 * (like a Mesh, GroupTopology, etc) that are required to initialize a Selection.
+	 * @param SelectionInOut the Selection to initialize, the input Geometry/Topology types define what type of element will be selected
+	 * @param SelectionIDPredicate only Elements that pass this predicate will be included in SelectionInOut
+	 * @param InitializeMode method to use to initialize SelectionInOut
+	 * @param ReferenceSelection optional Selection to select "relative to", based on the InitializeMode
+	 */
+	virtual void InitializeSelectionFromPredicate(	
+		FGeometrySelection& SelectionInOut,
+		TFunctionRef<bool(UE::Geometry::FGeoSelectionID)> SelectionIDPredicate,
+		EInitializeSelectionMode InitializeMode = EInitializeSelectionMode::All,
+		const FGeometrySelection* ReferenceSelection = nullptr 
+	) = 0;
+
+
+	/**
+	 * Combine FromSelection with the current Selection owned/referenced by the Selector, 
+	 * using the provided SelectionEditor and UpdateConfig. This can be used to explicitly
+	 * initialize the Selection, restore it to a previous state for undo/redo, etc.
+	 * @param bAllowConversion if true, FromSelection's Geometry/Topology type will be converted to the internal types of the Selector/SelectionEditor
+	 * @param SelectionDelta if passed as non-null, this Delta will be initialized based on the change in the internal Selection
+	 */
+	virtual void UpdateSelectionFromSelection(	
+		const FGeometrySelection& FromSelection,
+		bool bAllowConversion,
+		FGeometrySelectionEditor& SelectionEditor,
+		const FGeometrySelectionUpdateConfig& UpdateConfig,
+		UE::Geometry::FGeometrySelectionDelta* SelectionDelta = nullptr
+	) = 0;
+
+
 	struct FWorldRayQueryInfo
 	{
 		FRay3d WorldRay;
diff --git a/Engine/Plugins/Runtime/MeshModelingToolset/Source/ModelingComponents/Public/Selection/SelectionEditInteractiveCommand.h b/Engine/Plugins/Runtime/MeshModelingToolset/Source/ModelingComponents/Public/Selection/SelectionEditInteractiveCommand.h
--- a/Engine/Plugins/Runtime/MeshModelingToolset/Source/ModelingComponents/Public/Selection/SelectionEditInteractiveCommand.h
+++ b/Engine/Plugins/Runtime/MeshModelingToolset/Source/ModelingComponents/Public/Selection/SelectionEditInteractiveCommand.h
@@ -29,6 +29,16 @@ public:
 };
 
 
+UCLASS()
+class MODELINGCOMPONENTS_API UGeometrySelectionEditCommandResult : public UInteractiveCommandResult
+{
+	GENERATED_BODY()
+public:
+	FGeometrySelectionHandle SourceHandle;
+
+	UE::Geometry::FGeometrySelection OutputSelection;
+};
+
 
 /**
  * UGeometrySelectionEditCommand is a command that edits geometry based on a selection.
@@ -40,13 +50,14 @@ class MODELINGCOMPONENTS_API UGeometrySelectionEditCommand : public UInteractive
 	GENERATED_BODY()
 public:
 	virtual bool AllowEmptySelection() const { return false; }
+	virtual bool IsModifySelectionCommand() const { return false; }
 
 	virtual bool CanExecuteCommandForSelection(UGeometrySelectionEditCommandArguments* SelectionArgs)
 	{
 		return false;
 	}
 	
-	virtual void ExecuteCommandForSelection(UGeometrySelectionEditCommandArguments* Arguments)
+	virtual void ExecuteCommandForSelection(UGeometrySelectionEditCommandArguments* Arguments, UInteractiveCommandResult** Result = nullptr)
 	{
 	}
 
@@ -65,12 +76,12 @@ public:
 		return false;
 	}
 
-	virtual void ExecuteCommand(UInteractiveCommandArguments* Arguments) override
+	virtual void ExecuteCommand(UInteractiveCommandArguments* Arguments, UInteractiveCommandResult** Result = nullptr) override
 	{
 		UGeometrySelectionEditCommandArguments* SelectionArgs = Cast<UGeometrySelectionEditCommandArguments>(Arguments);
 		if ( SelectionArgs && (SelectionArgs->IsSelectionEmpty() == false || AllowEmptySelection()) )
 		{
-			ExecuteCommandForSelection(SelectionArgs);
+			ExecuteCommandForSelection(SelectionArgs, Result);
 		}
 	}
 
diff --git a/Engine/Source/Runtime/InteractiveToolsFramework/Public/InteractiveCommand.h b/Engine/Source/Runtime/InteractiveToolsFramework/Public/InteractiveCommand.h
--- a/Engine/Source/Runtime/InteractiveToolsFramework/Public/InteractiveCommand.h
+++ b/Engine/Source/Runtime/InteractiveToolsFramework/Public/InteractiveCommand.h
@@ -30,6 +30,18 @@ protected:
 };
 
 
+/**
+ * UInteractiveCommandResult subclasses are returned from UInteractiveCommands, to allow
+ * commands to return custom information.
+ */
+UCLASS()
+class INTERACTIVETOOLSFRAMEWORK_API UInteractiveCommandResult : public UObject
+{
+	GENERATED_BODY()
+
+};
+
+
 /**
  * A UInteractiveCommand is an atomic action that can be executed via some user interaction.
  * For example clicking a button that deletes an active selection can be considered an Interactive Command.
@@ -61,8 +73,9 @@ public:
 
 	/**
 	 * Execute the command with the given Arguments
+	 * @param Result optional command result. Command would have to allocate the UInteractiveCommandResult object to return it, but is not required to. Caller should not assume that a non-null Result is necessarily returned.
 	 */
-	virtual void ExecuteCommand(UInteractiveCommandArguments* Arguments)
+	virtual void ExecuteCommand(UInteractiveCommandArguments* Arguments, UInteractiveCommandResult** Result = nullptr)
 	{
 	}
 };
\ No newline at end of file
