ModelingMode: improve interop with new Selection System and existing Tools

Add support for Tools to provide an "output" selection. Add UGeometrySelectionManager::SetSelectionForComponent() which can set an explicit externally-provided selection. FBaseDynamicMeshSelector::UpdateSelectionFromSelection() now supports selection conversion when available and requested (is used to implement SetSelectionForComponent). New GeometrySelectionUtil functions InitializeSelectionFromTriangles() and ConvertSelection() are used to implement this (note: only Triangles->other conversion is currently supported). Add HaveAvailableGeometrySelection() and SetToolOutputGeometrySelectionForTarget() in StoredMeshSelectionUtil.h, this is the top-level function that Tools can use to set an Output selection.

ExtrudeMeshSelectionTool now emits output selection.

Update EditMeshPolygonsTool to use new Selection system and allow individual operations to be utilized as standalone Tools. Convert EditMeshPolygonsTool to be a USingleTargetWithSelectionTool, use FGeometrySelection to initialize selection. Add bTerminateOnPendingActionComplete flag, which is set when Tool is directly initialized to a specific operation, and forces tool to shut down when operation completes. This allows it to be used to more cleanly implement multiple action buttons in Modeling UI. When in this mode, selection panels are not shown. On Shutdown, now emits an "output" selection which GeometrySelectionManager can use to provide new selection to user. Update UPolygonSelectionMechanic Set/Get selection APIs to use FGeometrySelection instead of UPersistentMeshSelection.

port UVProjectionTool to derive from USingleTargetWithSelectionTool, use FGeometrySelection to initialize target ROI

deprecate UPersistentMeshSelection and related functions in StoredMeshSelectionUtil.h. Deprecate Tool Input Selection APIs in USingleSelectionMeshEditingTool and Builder.

Repurpose old ModelingMode-level PolyModel tab operations for new Selection Tools UI, now support Inset, Outset, Cut Faces, Insert Edge Loop, PushPull, and Bevel.

diff --git a/Engine/Plugins/Editor/ModelingToolsEditorMode/Source/ModelingToolsEditorMode/Private/ModelingToolsEditorMode.cpp b/Engine/Plugins/Editor/ModelingToolsEditorMode/Source/ModelingToolsEditorMode/Private/ModelingToolsEditorMode.cpp
--- a/Engine/Plugins/Editor/ModelingToolsEditorMode/Source/ModelingToolsEditorMode/Private/ModelingToolsEditorMode.cpp
+++ b/Engine/Plugins/Editor/ModelingToolsEditorMode/Source/ModelingToolsEditorMode/Private/ModelingToolsEditorMode.cpp
@@ -190,9 +190,6 @@ bool UModelingToolsEditorMode::ProcessEditDelete()
 		return true;
 	}
 
-	// clear any active selection
-	UE::Geometry::ClearActiveToolSelection(GetToolManager());
-
 	return false;
 }
 
@@ -207,9 +204,6 @@ bool UModelingToolsEditorMode::ProcessEditCut()
 		return true;
 	}
 
-	// clear any active selection
-	UE::Geometry::ClearActiveToolSelection(GetToolManager());
-
 	return false;
 }
 
@@ -712,34 +706,49 @@ void UModelingToolsEditorMode::Enter()
 	auto ExtractCollisionGeoToolBuilder = NewObject<UExtractCollisionGeometryToolBuilder>();
 	RegisterTool(ToolManagerCommands.BeginExtractCollisionGeometryTool, TEXT("BeginExtractCollisionGeometryTool"), ExtractCollisionGeoToolBuilder);
 
-	// PolyModeling tools
-	auto RegisterPolyModelSelectTool = [&](EEditMeshPolygonsToolSelectionMode SelectionMode, TSharedPtr<FUICommandInfo> UICommand, FString StringName)
+
+	// this function registers and tracks an active UGeometrySelectionEditCommand and it's associated UICommand
+	auto RegisterSelectionTool = [&](TSharedPtr<FUICommandInfo> UICommand, FString ToolIdentifier, UInteractiveToolBuilder* Builder)
 	{
-		UEditMeshPolygonsSelectionModeToolBuilder* SelectionModeBuilder = NewObject<UEditMeshPolygonsSelectionModeToolBuilder>();
-		SelectionModeBuilder->SelectionMode = SelectionMode;
-		RegisterTool(UICommand, StringName, SelectionModeBuilder);
+		UEditorInteractiveToolsContext* UseToolsContext = GetInteractiveToolsContext(EToolsContextScope::EdMode);
+		const TSharedRef<FUICommandList>& CommandList = Toolkit->GetToolkitCommands();
+		UseToolsContext->ToolManager->RegisterToolType(ToolIdentifier, Builder);
+		CommandList->MapAction(UICommand,
+			FExecuteAction::CreateUObject(UseToolsContext, &UEdModeInteractiveToolsContext::StartTool, ToolIdentifier),
+			FCanExecuteAction::CreateWeakLambda(UseToolsContext, [this, ToolIdentifier, UseToolsContext]() {
+				return ShouldToolStartBeAllowed(ToolIdentifier) &&
+					GetSelectionManager()->HasSelection() &&
+					UseToolsContext->ToolManager->CanActivateTool(EToolSide::Mouse, ToolIdentifier);
+			}),
+			FIsActionChecked::CreateUObject(UseToolsContext, &UEdModeInteractiveToolsContext::IsToolActive, EToolSide::Mouse, ToolIdentifier),
+			FIsActionButtonVisible::CreateUObject(GetSelectionManager(), &UGeometrySelectionManager::HasSelection),
+			EUIActionRepeatMode::RepeatDisabled);
 	};
-	RegisterPolyModelSelectTool(EEditMeshPolygonsToolSelectionMode::Faces, ToolManagerCommands.BeginPolyModelTool_FaceSelect, TEXT("PolyEdit_FaceSelect"));
-	RegisterPolyModelSelectTool(EEditMeshPolygonsToolSelectionMode::Edges, ToolManagerCommands.BeginPolyModelTool_EdgeSelect, TEXT("PolyEdit_EdgeSelect"));
-	RegisterPolyModelSelectTool(EEditMeshPolygonsToolSelectionMode::Vertices, ToolManagerCommands.BeginPolyModelTool_VertexSelect, TEXT("PolyEdit_VertexSelect"));
-	RegisterPolyModelSelectTool(EEditMeshPolygonsToolSelectionMode::Loops, ToolManagerCommands.BeginPolyModelTool_LoopSelect, TEXT("PolyEdit_LoopSelect"));
-	RegisterPolyModelSelectTool(EEditMeshPolygonsToolSelectionMode::Rings, ToolManagerCommands.BeginPolyModelTool_RingSelect, TEXT("PolyEdit_RingSelect"));
-	RegisterPolyModelSelectTool(EEditMeshPolygonsToolSelectionMode::FacesEdgesVertices, ToolManagerCommands.BeginPolyModelTool_AllSelect, TEXT("PolyEdit_AllSelect"));
 
-	auto RegisterPolyModelActionTool = [&](EEditMeshPolygonsToolActions Action, TSharedPtr<FUICommandInfo> UICommand, FString StringName)
+	// register mesh-selection-driven tools
+	RegisterSelectionTool(ToolManagerCommands.BeginSelectionAction_Extrude, TEXT("BeginSelectionExtrudeTool"), NewObject<UExtrudeMeshSelectionToolBuilder>());
+
+
+	auto RegisterPolyModelActionTool = [&](EEditMeshPolygonsToolActions Action, TSharedPtr<FUICommandInfo> UICommand, FString StringName, bool bRequiresSelection)
 	{
 		UEditMeshPolygonsActionModeToolBuilder* ActionModeBuilder = NewObject<UEditMeshPolygonsActionModeToolBuilder>();
 		ActionModeBuilder->StartupAction = Action;
-		RegisterTool(UICommand, StringName, ActionModeBuilder);
+		if ( bRequiresSelection )
+		{
+			RegisterSelectionTool(UICommand, StringName, ActionModeBuilder);
+		}
+		else
+		{
+			RegisterTool(UICommand, StringName, ActionModeBuilder);
+		}
 	};
-	RegisterPolyModelActionTool(EEditMeshPolygonsToolActions::Extrude, ToolManagerCommands.BeginPolyModelTool_Extrude, TEXT("PolyEdit_Extrude"));
-	RegisterPolyModelActionTool(EEditMeshPolygonsToolActions::Inset, ToolManagerCommands.BeginPolyModelTool_Inset, TEXT("PolyEdit_Inset"));
-	RegisterPolyModelActionTool(EEditMeshPolygonsToolActions::Outset, ToolManagerCommands.BeginPolyModelTool_Outset, TEXT("PolyEdit_Outset"));
-	RegisterPolyModelActionTool(EEditMeshPolygonsToolActions::CutFaces, ToolManagerCommands.BeginPolyModelTool_CutFaces, TEXT("PolyEdit_CutFaces"));
-
-	// register mesh-selection-driven tools
-	RegisterTool(ToolManagerCommands.BeginSelectionAction_Extrude, TEXT("BeginSelectionExtrudeTool"), NewObject<UExtrudeMeshSelectionToolBuilder>());
-
+	RegisterPolyModelActionTool(EEditMeshPolygonsToolActions::Inset, ToolManagerCommands.BeginPolyModelTool_Inset, TEXT("PolyEdit_Inset"), true);
+	RegisterPolyModelActionTool(EEditMeshPolygonsToolActions::Outset, ToolManagerCommands.BeginPolyModelTool_Outset, TEXT("PolyEdit_Outset"), true);
+	RegisterPolyModelActionTool(EEditMeshPolygonsToolActions::CutFaces, ToolManagerCommands.BeginPolyModelTool_CutFaces, TEXT("PolyEdit_CutFaces"), true);
+	RegisterPolyModelActionTool(EEditMeshPolygonsToolActions::InsertEdgeLoop, ToolManagerCommands.BeginPolyModelTool_InsertEdgeLoop, TEXT("PolyEdit_InsertEdgeLoop"), false);
+	RegisterPolyModelActionTool(EEditMeshPolygonsToolActions::PushPull, ToolManagerCommands.BeginPolyModelTool_PushPull, TEXT("PolyEdit_PushPull"), true);
+	RegisterPolyModelActionTool(EEditMeshPolygonsToolActions::BevelAuto, ToolManagerCommands.BeginPolyModelTool_Bevel, TEXT("PolyEdit_Bevel"), true);
+	
 
 	// set up selection type toggles
 	auto RegisterSelectionTopologyType = [this](UGeometrySelectionManager::EMeshTopologyMode TopoMode, TSharedPtr<FUICommandInfo> UICommand)
@@ -799,7 +808,7 @@ void UModelingToolsEditorMode::Enter()
 
 
 	// this function registers and tracks an active UGeometrySelectionEditCommand and it's associated UICommand
-	auto RegisterSelectionCommand = [&](UGeometrySelectionEditCommand* Command, TSharedPtr<FUICommandInfo> UICommand)
+	auto RegisterSelectionCommand = [&](UGeometrySelectionEditCommand* Command, TSharedPtr<FUICommandInfo> UICommand, bool bAlwaysVisible)
 	{
 		ModelingModeCommands.Add(Command);
 		const TSharedRef<FUICommandList>& CommandList = Toolkit->GetToolkitCommands();
@@ -807,19 +816,20 @@ void UModelingToolsEditorMode::Enter()
 			FExecuteAction::CreateUObject(GetSelectionManager(), &UGeometrySelectionManager::ExecuteSelectionCommand, Command),
 			FCanExecuteAction::CreateUObject(GetSelectionManager(), &UGeometrySelectionManager::CanExecuteSelectionCommand, Command),
 			FIsActionChecked(),
-			FIsActionButtonVisible::CreateUObject(GetSelectionManager(), &UGeometrySelectionManager::CanExecuteSelectionCommand, Command),
+			(bAlwaysVisible) ? FIsActionButtonVisible() : 
+				FIsActionButtonVisible::CreateUObject(GetSelectionManager(), &UGeometrySelectionManager::CanExecuteSelectionCommand, Command),
 			EUIActionRepeatMode::RepeatDisabled);
 	};
 
 	// create and register InteractiveCommands for mesh selections
-	RegisterSelectionCommand(NewObject<UDeleteGeometrySelectionCommand>(), ToolManagerCommands.BeginSelectionAction_Delete);
-	RegisterSelectionCommand(NewObject<UDisconnectGeometrySelectionCommand>(), ToolManagerCommands.BeginSelectionAction_Disconnect);
-	RegisterSelectionCommand(NewObject<UModifyGeometrySelectionCommand>(), ToolManagerCommands.BeginSelectionAction_SelectAll);
-	RegisterSelectionCommand(NewObject<UModifyGeometrySelectionCommand_Invert>(), ToolManagerCommands.BeginSelectionAction_Invert);
-	RegisterSelectionCommand(NewObject<UModifyGeometrySelectionCommand_ExpandToConnected>(), ToolManagerCommands.BeginSelectionAction_ExpandToConnected);
-	RegisterSelectionCommand(NewObject<UModifyGeometrySelectionCommand_InvertConnected>(), ToolManagerCommands.BeginSelectionAction_InvertConnected);
-	RegisterSelectionCommand(NewObject<UModifyGeometrySelectionCommand_Expand>(), ToolManagerCommands.BeginSelectionAction_Expand);
-	RegisterSelectionCommand(NewObject<UModifyGeometrySelectionCommand_Contract>(), ToolManagerCommands.BeginSelectionAction_Contract);
+	RegisterSelectionCommand(NewObject<UDeleteGeometrySelectionCommand>(), ToolManagerCommands.BeginSelectionAction_Delete, false);
+	RegisterSelectionCommand(NewObject<UDisconnectGeometrySelectionCommand>(), ToolManagerCommands.BeginSelectionAction_Disconnect, false);
+	RegisterSelectionCommand(NewObject<UModifyGeometrySelectionCommand>(), ToolManagerCommands.BeginSelectionAction_SelectAll, true);
+	RegisterSelectionCommand(NewObject<UModifyGeometrySelectionCommand_Invert>(), ToolManagerCommands.BeginSelectionAction_Invert, true);
+	RegisterSelectionCommand(NewObject<UModifyGeometrySelectionCommand_ExpandToConnected>(), ToolManagerCommands.BeginSelectionAction_ExpandToConnected, true);
+	RegisterSelectionCommand(NewObject<UModifyGeometrySelectionCommand_InvertConnected>(), ToolManagerCommands.BeginSelectionAction_InvertConnected, true);
+	RegisterSelectionCommand(NewObject<UModifyGeometrySelectionCommand_Expand>(), ToolManagerCommands.BeginSelectionAction_Expand, true);
+	RegisterSelectionCommand(NewObject<UModifyGeometrySelectionCommand_Contract>(), ToolManagerCommands.BeginSelectionAction_Contract, true);
 
 	// register extensions
 	TArray<IModelingModeToolExtension*> Extensions = IModularFeatures::Get().GetModularFeatureImplementations<IModelingModeToolExtension>(
@@ -1275,7 +1285,6 @@ void UModelingToolsEditorMode::BindCommands()
 	CommandList->MapAction(
 		ToolManagerCommands.AcceptActiveTool,
 		FExecuteAction::CreateLambda([this]() { 
-			UE::Geometry::ClearActiveToolSelection(GetToolManager());
 			GetInteractiveToolsContext()->EndTool(EToolShutdownType::Accept); 
 		}),
 		FCanExecuteAction::CreateLambda([this]() { return GetInteractiveToolsContext()->CanAcceptActiveTool(); }),
@@ -1342,9 +1351,6 @@ void UModelingToolsEditorMode::AcceptActiveToolActionOrTool()
 		}
 	}
 
-	// clear existing selection
-	UE::Geometry::ClearActiveToolSelection(GetToolManager());
-
 	const EToolShutdownType ShutdownType = GetInteractiveToolsContext()->CanAcceptActiveTool() ? EToolShutdownType::Accept : EToolShutdownType::Completed;
 	GetInteractiveToolsContext()->EndTool(ShutdownType);
 }
diff --git a/Engine/Plugins/Editor/ModelingToolsEditorMode/Source/ModelingToolsEditorMode/Private/ModelingToolsEditorModeToolkit.cpp b/Engine/Plugins/Editor/ModelingToolsEditorMode/Source/ModelingToolsEditorMode/Private/ModelingToolsEditorModeToolkit.cpp
--- a/Engine/Plugins/Editor/ModelingToolsEditorMode/Source/ModelingToolsEditorMode/Private/ModelingToolsEditorModeToolkit.cpp
+++ b/Engine/Plugins/Editor/ModelingToolsEditorMode/Source/ModelingToolsEditorMode/Private/ModelingToolsEditorModeToolkit.cpp
@@ -676,6 +676,14 @@ void FModelingToolsEditorModeToolkit::BuildToolPalette(FName PaletteIndex, class
 		ToolbarBuilder.AddToolBarButton(Commands.BeginSelectionAction_Delete);
 		ToolbarBuilder.AddToolBarButton(Commands.BeginSelectionAction_Disconnect);
 		ToolbarBuilder.AddToolBarButton(Commands.BeginSelectionAction_Extrude);
+
+		ToolbarBuilder.AddToolBarButton(Commands.BeginPolyModelTool_PushPull);
+		ToolbarBuilder.AddToolBarButton(Commands.BeginPolyModelTool_InsertEdgeLoop);
+		ToolbarBuilder.AddToolBarButton(Commands.BeginPolyModelTool_CutFaces);
+		ToolbarBuilder.AddToolBarButton(Commands.BeginPolyModelTool_Inset);
+		ToolbarBuilder.AddToolBarButton(Commands.BeginPolyModelTool_Outset);
+		ToolbarBuilder.AddToolBarButton(Commands.BeginPolyModelTool_Bevel);
+		
 	}
 	else if (PaletteIndex == PrimitiveTabName)
 	{
@@ -820,23 +828,23 @@ void FModelingToolsEditorModeToolkit::BuildToolPalette(FName PaletteIndex, class
 		ToolbarBuilder.AddToolBarButton(Commands.BeginAddPatchTool);
 		ToolbarBuilder.AddToolBarButton(Commands.BeginShapeSprayTool);
 	}
-	else if (PaletteIndex == PolyEditTabName)
-	{
-		ToolbarBuilder.AddToolBarButton(Commands.BeginPolyModelTool_FaceSelect);
-		ToolbarBuilder.AddToolBarButton(Commands.BeginPolyModelTool_EdgeSelect);
-		ToolbarBuilder.AddToolBarButton(Commands.BeginPolyModelTool_VertexSelect);
-		ToolbarBuilder.AddToolBarButton(Commands.BeginPolyModelTool_AllSelect);
-		ToolbarBuilder.AddToolBarButton(Commands.BeginPolyModelTool_LoopSelect);
-		ToolbarBuilder.AddToolBarButton(Commands.BeginPolyModelTool_RingSelect);
-		ToolbarBuilder.AddSeparator();
-		ToolbarBuilder.AddToolBarButton(Commands.BeginPolyModelTool_Extrude);
-		ToolbarBuilder.AddToolBarButton(Commands.BeginPolyModelTool_Inset);
-		ToolbarBuilder.AddToolBarButton(Commands.BeginPolyModelTool_Outset);
-		ToolbarBuilder.AddToolBarButton(Commands.BeginPolyModelTool_CutFaces);
-		ToolbarBuilder.AddSeparator();
-		ToolbarBuilder.AddToolBarButton(Commands.BeginSubdividePolyTool);
-		ToolbarBuilder.AddToolBarButton(Commands.BeginPolyEditTool);
-	}
+	//else if (PaletteIndex == PolyEditTabName)
+	//{
+	//	ToolbarBuilder.AddToolBarButton(Commands.BeginPolyModelTool_FaceSelect);
+	//	ToolbarBuilder.AddToolBarButton(Commands.BeginPolyModelTool_EdgeSelect);
+	//	ToolbarBuilder.AddToolBarButton(Commands.BeginPolyModelTool_VertexSelect);
+	//	ToolbarBuilder.AddToolBarButton(Commands.BeginPolyModelTool_AllSelect);
+	//	ToolbarBuilder.AddToolBarButton(Commands.BeginPolyModelTool_LoopSelect);
+	//	ToolbarBuilder.AddToolBarButton(Commands.BeginPolyModelTool_RingSelect);
+	//	ToolbarBuilder.AddSeparator();
+	//	ToolbarBuilder.AddToolBarButton(Commands.BeginPolyModelTool_Extrude);
+	//	ToolbarBuilder.AddToolBarButton(Commands.BeginPolyModelTool_Inset);
+	//	ToolbarBuilder.AddToolBarButton(Commands.BeginPolyModelTool_Outset);
+	//	ToolbarBuilder.AddToolBarButton(Commands.BeginPolyModelTool_CutFaces);
+	//	ToolbarBuilder.AddSeparator();
+	//	ToolbarBuilder.AddToolBarButton(Commands.BeginSubdividePolyTool);
+	//	ToolbarBuilder.AddToolBarButton(Commands.BeginPolyEditTool);
+	//}
 	else
 	{
 		TArray<IModelingModeToolExtension*> Extensions = IModularFeatures::Get().GetModularFeatureImplementations<IModelingModeToolExtension>(
diff --git a/Engine/Plugins/Editor/ModelingToolsEditorMode/Source/ModelingToolsEditorMode/Private/ModelingToolsManagerActions.cpp b/Engine/Plugins/Editor/ModelingToolsEditorMode/Source/ModelingToolsEditorMode/Private/ModelingToolsManagerActions.cpp
--- a/Engine/Plugins/Editor/ModelingToolsEditorMode/Source/ModelingToolsEditorMode/Private/ModelingToolsManagerActions.cpp
+++ b/Engine/Plugins/Editor/ModelingToolsEditorMode/Source/ModelingToolsEditorMode/Private/ModelingToolsManagerActions.cpp
@@ -175,17 +175,12 @@ void FModelingToolsManagerCommands::RegisterCommands()
 	UI_COMMAND(BeginSkinWeightsPaintTool, "SkinWts", "Start the Paint Skin Weights Tool", EUserInterfaceActionType::Button, FInputChord());
 	UI_COMMAND(BeginSkinWeightsBindingTool, "SkinBind", "Start the Skin Weights Binding Tool", EUserInterfaceActionType::Button, FInputChord());
 
-	REGISTER_MODELING_TOOL_COMMAND(BeginPolyModelTool_FaceSelect, "Faces", "PolyGroup Face Selection Tool");
-	REGISTER_MODELING_TOOL_COMMAND(BeginPolyModelTool_EdgeSelect, "Edges", "PolyGroup Edge Selection Tool");
-	REGISTER_MODELING_TOOL_COMMAND(BeginPolyModelTool_VertexSelect, "Verts", "PolyGroup Vertex Selection Tool");
-	REGISTER_MODELING_TOOL_COMMAND(BeginPolyModelTool_AllSelect, "Any", "PolyGroup Face/Edge/Vertex Selection Tool");
-	REGISTER_MODELING_TOOL_COMMAND(BeginPolyModelTool_LoopSelect, "Loops", "PolyGroup Loop Selection Tool");
-	REGISTER_MODELING_TOOL_COMMAND(BeginPolyModelTool_RingSelect, "Rings", "PolyGroup Ring Selection Tool");
-
-	REGISTER_MODELING_TOOL_COMMAND(BeginPolyModelTool_Extrude, "Extrude", "PolyGroup Extrude Tool");
 	REGISTER_MODELING_TOOL_COMMAND(BeginPolyModelTool_Inset, "Inset", "PolyGroup Inset Tool");
 	REGISTER_MODELING_TOOL_COMMAND(BeginPolyModelTool_Outset, "Outset", "PolyGroup Outset Tool");
 	REGISTER_MODELING_TOOL_COMMAND(BeginPolyModelTool_CutFaces, "Cut", "PolyGroup Cut Faces Tool");
+	REGISTER_MODELING_TOOL_COMMAND(BeginPolyModelTool_InsertEdgeLoop, "ELoop", "PolyGroup Insert Edge Loop Tool");
+	REGISTER_MODELING_TOOL_COMMAND(BeginPolyModelTool_PushPull, "PushPull", "PolyGroup Push Pull Faces Tool");
+	REGISTER_MODELING_TOOL_COMMAND(BeginPolyModelTool_Bevel, "Bevel", "PolyGroup Bevel Tool");
 
 
 	REGISTER_MODELING_TOOL_COMMAND(BeginSelectionAction_ToObjectType, "None", "Disable Mesh Element Selection");
diff --git a/Engine/Plugins/Editor/ModelingToolsEditorMode/Source/ModelingToolsEditorMode/Public/ModelingToolsManagerActions.h b/Engine/Plugins/Editor/ModelingToolsEditorMode/Source/ModelingToolsEditorMode/Public/ModelingToolsManagerActions.h
--- a/Engine/Plugins/Editor/ModelingToolsEditorMode/Source/ModelingToolsEditorMode/Public/ModelingToolsManagerActions.h
+++ b/Engine/Plugins/Editor/ModelingToolsEditorMode/Source/ModelingToolsEditorMode/Public/ModelingToolsManagerActions.h
@@ -132,17 +132,13 @@ public:
 	TSharedPtr<FUICommandInfo> BeginGenerateStaticMeshLODAssetTool;
 	TSharedPtr<FUICommandInfo> BeginISMEditorTool;
 
-	TSharedPtr<FUICommandInfo> BeginPolyModelTool_FaceSelect;
-	TSharedPtr<FUICommandInfo> BeginPolyModelTool_EdgeSelect;
-	TSharedPtr<FUICommandInfo> BeginPolyModelTool_VertexSelect;
-	TSharedPtr<FUICommandInfo> BeginPolyModelTool_AllSelect;
-	TSharedPtr<FUICommandInfo> BeginPolyModelTool_LoopSelect;
-	TSharedPtr<FUICommandInfo> BeginPolyModelTool_RingSelect;
-
-	TSharedPtr<FUICommandInfo> BeginPolyModelTool_Extrude;
 	TSharedPtr<FUICommandInfo> BeginPolyModelTool_Inset;
 	TSharedPtr<FUICommandInfo> BeginPolyModelTool_Outset;
 	TSharedPtr<FUICommandInfo> BeginPolyModelTool_CutFaces;
+	TSharedPtr<FUICommandInfo> BeginPolyModelTool_InsertEdgeLoop;
+	TSharedPtr<FUICommandInfo> BeginPolyModelTool_PushPull;
+	TSharedPtr<FUICommandInfo> BeginPolyModelTool_Bevel;
+	
 
 
 	// old selection state actions
diff --git a/Engine/Plugins/Experimental/MeshModelingToolsetExp/Source/MeshModelingToolsExp/Private/Polymodeling/ExtrudeMeshSelectionTool.cpp b/Engine/Plugins/Experimental/MeshModelingToolsetExp/Source/MeshModelingToolsExp/Private/Polymodeling/ExtrudeMeshSelectionTool.cpp
--- a/Engine/Plugins/Experimental/MeshModelingToolsetExp/Source/MeshModelingToolsExp/Private/Polymodeling/ExtrudeMeshSelectionTool.cpp
+++ b/Engine/Plugins/Experimental/MeshModelingToolsetExp/Source/MeshModelingToolsExp/Private/Polymodeling/ExtrudeMeshSelectionTool.cpp
@@ -6,6 +6,7 @@
 #include "ToolSetupUtil.h"
 #include "ModelingToolTargetUtil.h"
 #include "Selections/GeometrySelectionUtil.h"
+#include "Selection/StoredMeshSelectionUtil.h"
 #include "ToolSceneQueriesUtil.h"
 #include "MeshQueries.h"
 
@@ -268,16 +269,28 @@ void UExtrudeMeshSelectionTool::OnShutdown(EToolShutdownType ShutdownType)
 		{
 			GetToolManager()->BeginUndoTransaction(LOCTEXT("ExtrudeSelection", "Extrude Selection"));
 
+			TArray<int32> SelectOutputTriangles;
+
 			// this function computes the final Extrude on full Mesh and cleans it up
-			auto ComputeFinalResult = [this](FDynamicMesh3& EditMesh) -> bool
+			auto ComputeFinalResult = [this, &SelectOutputTriangles](FDynamicMesh3& EditMesh) -> bool
 			{
 				TSharedPtr<FSharedConstDynamicMesh3> FullMeshShared = MakeShared<FSharedConstDynamicMesh3>(&EditMesh);
 				TUniquePtr<FLinearExtrusionOp> FinalOp = OperatorFactory->MakeNewExtrudeOp(FullMeshShared, this->ExtrudeROI);
+				SelectOutputTriangles = this->ExtrudeROI;
 				bool bResult = FinalOp->CalculateResultInPlace(EditMesh, nullptr);
 				if (EditMesh.IsCompact() == false)
 				{
-					EditMesh.CompactInPlace();
+					FCompactMaps CompactMaps;
+					EditMesh.CompactInPlace(&CompactMaps);
+					if (CompactMaps.TriangleMapIsSet())
+					{
+						for (int32& tid : SelectOutputTriangles)
+						{
+							tid = CompactMaps.GetTriangleMapping(tid);
+						}
+					}
 				}
+
 				FullMeshShared->ReleaseSharedObject();
 				return bResult;
 			};
@@ -313,6 +326,19 @@ void UExtrudeMeshSelectionTool::OnShutdown(EToolShutdownType ShutdownType)
 				UE::ToolTarget::CommitDynamicMeshUpdate(UseTarget, CurrentMesh, bModifiedTopology);
 			}
 
+			// Construct and send output triangle selection, assume selection system will convert to correct type.
+			// (todo: maybe be smarter about this, as for (eg) an input vertex selection, it will grow the selection...)
+			// Note that we cannot do this before we commit/complete the mesh update, as the selection system needs a chance
+			// to respond to the mesh change. Possibly should use UE::Geometry::InitializeSelectionFromTriangles() here...need to ProcessMesh() then though
+			FGeometrySelection OutputSelection;
+			OutputSelection.InitializeTypes(EGeometryElementType::Face, EGeometryTopologyType::Triangle);
+			for (int32 tid : SelectOutputTriangles)
+			{
+				OutputSelection.Selection.Add( FGeoSelectionID::MeshTriangle(tid).Encoded() );
+			}
+			UE::Geometry::SetToolOutputGeometrySelectionForTarget(this, UseTarget, OutputSelection);
+
+
 			GetToolManager()->EndUndoTransaction();
 		}
 	}
diff --git a/Engine/Plugins/Runtime/GeometryProcessing/Source/DynamicMesh/Private/Selections/GeometrySelectionUtil.cpp b/Engine/Plugins/Runtime/GeometryProcessing/Source/DynamicMesh/Private/Selections/GeometrySelectionUtil.cpp
--- a/Engine/Plugins/Runtime/GeometryProcessing/Source/DynamicMesh/Private/Selections/GeometrySelectionUtil.cpp
+++ b/Engine/Plugins/Runtime/GeometryProcessing/Source/DynamicMesh/Private/Selections/GeometrySelectionUtil.cpp
@@ -9,6 +9,8 @@
 #include "SegmentTypes.h"
 #include "GroupTopology.h"
 #include "Selections/MeshConnectedComponents.h"
+#include "Selections/MeshEdgeSelection.h"
+#include "Selections/MeshVertexSelection.h"
 #include "Algo/Find.h"
 
 
@@ -856,6 +858,149 @@ bool UE::Geometry::ConvertPolygroupSelectionToTopologySelection(
 
 
 
+bool UE::Geometry::InitializeSelectionFromTriangles(
+	const UE::Geometry::FDynamicMesh3& Mesh,
+	const FGroupTopology* GroupTopology,
+	TArrayView<const int> Triangles,
+	FGeometrySelection& SelectionOut)
+{
+	if (SelectionOut.TopologyType == EGeometryTopologyType::Triangle)
+	{
+		if (SelectionOut.ElementType == EGeometryElementType::Vertex)
+		{
+			for (int32 tid : Triangles)
+			{
+				if (Mesh.IsTriangle(tid))
+				{
+					FIndex3i TriVertices = Mesh.GetTriangle(tid);
+					SelectionOut.Selection.Add(FGeoSelectionID::MeshVertex(TriVertices.A).Encoded() );
+					SelectionOut.Selection.Add(FGeoSelectionID::MeshVertex(TriVertices.B).Encoded());
+					SelectionOut.Selection.Add(FGeoSelectionID::MeshVertex(TriVertices.C).Encoded());
+				}
+			}
+		}
+		else if (SelectionOut.ElementType == EGeometryElementType::Edge)
+		{
+			for (int32 tid : Triangles)
+			{
+				if (Mesh.IsTriangle(tid))
+				{
+					FIndex3i TriEdges = Mesh.GetTriEdges(tid);
+					SelectionOut.Selection.Add(FGeoSelectionID::MeshEdge(Mesh.GetTriEdgeIDFromEdgeID(TriEdges.A)).Encoded());
+					SelectionOut.Selection.Add(FGeoSelectionID::MeshEdge(Mesh.GetTriEdgeIDFromEdgeID(TriEdges.B)).Encoded());
+					SelectionOut.Selection.Add(FGeoSelectionID::MeshEdge(Mesh.GetTriEdgeIDFromEdgeID(TriEdges.C)).Encoded());
+				}
+			}
+		}
+		else if (SelectionOut.ElementType == EGeometryElementType::Face)
+		{
+			for (int32 tid : Triangles)
+			{
+				if (Mesh.IsTriangle(tid))
+				{
+					SelectionOut.Selection.Add(FGeoSelectionID::MeshTriangle(tid).Encoded());
+				}
+			}
+		}
+		else
+		{
+			return false;
+		}
+		return true;
+	}
+	else if (SelectionOut.TopologyType == EGeometryTopologyType::Polygroup)
+	{
+		if (!ensure(GroupTopology != nullptr))
+		{
+			return false;
+		}
+
+		if (SelectionOut.ElementType == EGeometryElementType::Vertex)
+		{
+			FMeshVertexSelection VertSelection(&Mesh);
+			VertSelection.SelectTriangleVertices(Triangles);
+			for (int32 vid : VertSelection)
+			{
+				int32 CornerID = GroupTopology->GetCornerIDFromVertexID(vid);
+				if (CornerID != IndexConstants::InvalidID)
+				{
+					const FGroupTopology::FCorner& Corner = GroupTopology->Corners[CornerID];
+					FGeoSelectionID ID = FGeoSelectionID(Corner.VertexID, CornerID);
+					SelectionOut.Selection.Add(ID.Encoded());
+				}
+			}
+		}
+		else if (SelectionOut.ElementType == EGeometryElementType::Edge)
+		{
+			FMeshEdgeSelection EdgeSelection(&Mesh);
+			EdgeSelection.SelectTriangleEdges(Triangles);
+			for (int32 eid : EdgeSelection)
+			{
+				int32 GroupEdgeID = GroupTopology->FindGroupEdgeID(eid);
+				if (GroupEdgeID != IndexConstants::InvalidID)
+				{
+					const FGroupTopology::FGroupEdge& GroupEdge = GroupTopology->Edges[GroupEdgeID];
+					FMeshTriEdgeID MeshEdgeID = Mesh.GetTriEdgeIDFromEdgeID(GroupEdge.Span.Edges[0]);
+					FGeoSelectionID ID = FGeoSelectionID(MeshEdgeID.Encoded(), GroupEdgeID);
+					SelectionOut.Selection.Add(ID.Encoded());
+				}
+			}
+		}
+		else if (SelectionOut.ElementType == EGeometryElementType::Face)
+		{
+			for (int32 tid : Triangles)
+			{
+				if (Mesh.IsTriangle(tid))
+				{
+					int32 GroupID = GroupTopology->GetGroupID(tid);
+					const FGroupTopology::FGroup* GroupFace = GroupTopology->FindGroupByID(GroupID);
+					if ( GroupFace )
+					{
+						FGeoSelectionID ID = FGeoSelectionID(GroupFace->Triangles[0], GroupFace->GroupID);
+						SelectionOut.Selection.Add(ID.Encoded());
+					}
+				}
+			}
+		}
+		else
+		{
+			return false;
+		}
+		return true;
+	}
+	return false;
+}
+
+
+bool UE::Geometry::ConvertSelection(
+	const UE::Geometry::FDynamicMesh3& Mesh,
+	const FGroupTopology* GroupTopology,
+	const FGeometrySelection& FromSelectionIn,
+	FGeometrySelection& ToSelectionOut)
+{
+	if (FromSelectionIn.IsSameType(ToSelectionOut))
+	{
+		ToSelectionOut.Selection = FromSelectionIn.Selection;
+		return true;
+	}
+
+	if (FromSelectionIn.TopologyType == EGeometryTopologyType::Triangle && FromSelectionIn.ElementType == EGeometryElementType::Face)
+	{
+		TArray<int32> Triangles;
+		for (uint64 ElemID : FromSelectionIn.Selection)
+		{
+			Triangles.Add( FGeoSelectionID(ElemID).GeometryID );
+		}
+		return InitializeSelectionFromTriangles(Mesh, GroupTopology, Triangles, ToSelectionOut);
+	}
+
+	// todo: support more conversions
+
+	return false;
+}
+
+
+
 bool UE::Geometry::MakeSelectAllSelection(
 	const UE::Geometry::FDynamicMesh3& Mesh,
 	const FGroupTopology* GroupTopology,
diff --git a/Engine/Plugins/Runtime/GeometryProcessing/Source/DynamicMesh/Public/Selections/GeometrySelectionUtil.h b/Engine/Plugins/Runtime/GeometryProcessing/Source/DynamicMesh/Public/Selections/GeometrySelectionUtil.h
--- a/Engine/Plugins/Runtime/GeometryProcessing/Source/DynamicMesh/Public/Selections/GeometrySelectionUtil.h
+++ b/Engine/Plugins/Runtime/GeometryProcessing/Source/DynamicMesh/Public/Selections/GeometrySelectionUtil.h
@@ -192,6 +192,29 @@ DYNAMICMESH_API bool ConvertPolygroupSelectionToTopologySelection(
 );
 
 
+/**
+ * Convert Triangle IDs to target Selection type
+ */
+DYNAMICMESH_API bool InitializeSelectionFromTriangles(
+	const UE::Geometry::FDynamicMesh3& Mesh,
+	const FGroupTopology* GroupTopology,
+	TArrayView<const int> Triangles,
+	FGeometrySelection& SelectionOut);
+
+
+/**
+ * Convert Selection from one type to another, based on geometry/topology types in FromSelectionIn and ToSelectionOut.
+ * Not all conversion types are necessarily supported
+ * (currently only Triangles -> All Others is working)
+ * @return true if conversion is supported and was computed successfully
+ */
+DYNAMICMESH_API bool ConvertSelection(
+	const UE::Geometry::FDynamicMesh3& Mesh,
+	const FGroupTopology* GroupTopology,
+	const FGeometrySelection& FromSelectionIn,
+	FGeometrySelection& ToSelectionOut);
+
+
 /**
  * Select all elements of the provided Mesh and GroupTopology that pass the provided SelectionIDPredicate, 
  * and store in the output AllSelection. The type of elements selected is defined by the existing configured
diff --git a/Engine/Plugins/Runtime/MeshModelingToolset/Source/MeshModelingTools/Private/EditMeshPolygonsTool.cpp b/Engine/Plugins/Runtime/MeshModelingToolset/Source/MeshModelingTools/Private/EditMeshPolygonsTool.cpp
--- a/Engine/Plugins/Runtime/MeshModelingToolset/Source/MeshModelingTools/Private/EditMeshPolygonsTool.cpp
+++ b/Engine/Plugins/Runtime/MeshModelingToolset/Source/MeshModelingTools/Private/EditMeshPolygonsTool.cpp
@@ -29,7 +29,6 @@
 #include "Operations/MinimalHoleFiller.h"
 #include "Operations/PolygroupRemesh.h"
 #include "Operations/WeldEdgeSequence.h"
-#include "Selection/PersistentMeshSelection.h"
 #include "Selection/StoredMeshSelectionUtil.h"
 #include "Selection/PolygonSelectionMechanic.h"
 #include "Selections/MeshConnectedComponents.h"
@@ -50,6 +49,7 @@
 #include "ToolTargetManager.h"
 #include "TransformTypes.h"
 #include "Util/CompactMaps.h"
+#include "Selections/GeometrySelection.h"
 
 #include UE_INLINE_GENERATED_CPP_BY_NAME(EditMeshPolygonsTool)
 
@@ -117,14 +117,14 @@ namespace EditMeshPolygonsToolLocals
  */
 
 
-USingleSelectionMeshEditingTool* UEditMeshPolygonsToolBuilder::CreateNewTool(const FToolBuilderState& SceneState) const
+USingleTargetWithSelectionTool* UEditMeshPolygonsToolBuilder::CreateNewTool(const FToolBuilderState& SceneState) const
 {
 	return NewObject<UEditMeshPolygonsTool>(SceneState.ToolManager);
 }
 
-void UEditMeshPolygonsToolBuilder::InitializeNewTool(USingleSelectionMeshEditingTool* Tool, const FToolBuilderState& SceneState) const
+void UEditMeshPolygonsToolBuilder::InitializeNewTool(USingleTargetWithSelectionTool* Tool, const FToolBuilderState& SceneState) const
 {
-	USingleSelectionMeshEditingToolBuilder::InitializeNewTool(Tool, SceneState);
+	USingleTargetWithSelectionToolBuilder::InitializeNewTool(Tool, SceneState);
 	UEditMeshPolygonsTool* EditPolygonsTool = CastChecked<UEditMeshPolygonsTool>(Tool);
 	if (bTriangleMode)
 	{
@@ -134,7 +134,7 @@ void UEditMeshPolygonsToolBuilder::InitializeNewTool(USingleSelectionMeshEditing
 
 
 
-void UEditMeshPolygonsActionModeToolBuilder::InitializeNewTool(USingleSelectionMeshEditingTool* Tool, const FToolBuilderState& SceneState) const
+void UEditMeshPolygonsActionModeToolBuilder::InitializeNewTool(USingleTargetWithSelectionTool* Tool, const FToolBuilderState& SceneState) const
 {
 	UEditMeshPolygonsToolBuilder::InitializeNewTool(Tool, SceneState);
 	UEditMeshPolygonsTool* EditPolygonsTool = CastChecked<UEditMeshPolygonsTool>(Tool);
@@ -143,13 +143,13 @@ void UEditMeshPolygonsActionModeToolBuilder::InitializeNewTool(USingleSelectionM
 	EditPolygonsTool->PostSetupFunction = [UseAction](UEditMeshPolygonsTool* PolyTool)
 	{
 		PolyTool->SetToSelectionModeInterface();
-		PolyTool->RequestAction(UseAction);
+		PolyTool->RequestSingleShotAction(UseAction);
 	};
 }
 
 
 
-void UEditMeshPolygonsSelectionModeToolBuilder::InitializeNewTool(USingleSelectionMeshEditingTool* Tool, const FToolBuilderState& SceneState) const
+void UEditMeshPolygonsSelectionModeToolBuilder::InitializeNewTool(USingleTargetWithSelectionTool* Tool, const FToolBuilderState& SceneState) const
 {
 	UEditMeshPolygonsToolBuilder::InitializeNewTool(Tool, SceneState);
 	UEditMeshPolygonsTool* EditPolygonsTool = CastChecked<UEditMeshPolygonsTool>(Tool);
@@ -414,7 +414,7 @@ void UEditMeshPolygonsTool::Setup()
 
 	// Create the preview object
 	Preview = NewObject<UMeshOpPreviewWithBackgroundCompute>();
-	Preview->Setup(TargetWorld);
+	Preview->Setup(GetTargetWorld());
 	ToolSetupUtil::ApplyRenderingConfigurationToPreview(Preview->PreviewMesh, Target); 
 	Preview->PreviewMesh->SetTransform((FTransform)WorldTransform);
 
@@ -466,18 +466,25 @@ void UEditMeshPolygonsTool::Setup()
 	}
 	SelectionMechanic->Initialize(CurrentMesh.Get(),
 		(FTransform3d)Preview->PreviewMesh->GetTransform(),
-		TargetWorld,
+		GetTargetWorld(),
 		Topology.Get(),
 		[this]() { return &GetSpatial(); }
 	);
 
 	LinearDeformer.Initialize(CurrentMesh.Get(), Topology.Get());
 
-	// Have to load selection after initializing the selection mechanic since we need to have
-	// the topology built.
-	if (HasInputSelection() && IsToolInputSelectionUsable(GetInputSelection()))
+	// initialize our selection from input selection, if available 
+	if (HasGeometrySelection())
 	{
-		SelectionMechanic->LoadSelection(*GetInputSelection());
+		const FGeometrySelection& CurSelection = GetGeometrySelection();
+		if (CurSelection.TopologyType == EGeometryTopologyType::Triangle && bTriangleMode)
+		{
+			SelectionMechanic->SetSelection_AsTriangleTopology(CurSelection);
+		}
+		else if (CurSelection.TopologyType == EGeometryTopologyType::Polygroup && bTriangleMode == false)
+		{
+			SelectionMechanic->SetSelection_AsGroupTopology(CurSelection);
+		}
 	}
 
 	bSelectionStateDirty = SelectionMechanic->HasSelection();
@@ -622,22 +629,6 @@ void UEditMeshPolygonsTool::ResetUserMessage()
 	}
 }
 
-bool UEditMeshPolygonsTool::IsToolInputSelectionUsable(const UPersistentMeshSelection* InputSelectionIn)
-{
-	// TODO: We currently don't support persistent selection on volume brushes because
-	// a conversion back to a brush involves a simplification step that may make the 
-	// same vids unrecoverable. Once we have persistence of dynamic meshes, this will
-	// hopefully not become a problem, and this function (along with stored selection
-	// identifying info) will change.
-	return !Cast<UBrushComponent>(UE::ToolTarget::GetTargetComponent(Target))
-
-		&& InputSelectionIn
-		&& InputSelectionIn->GetSelectionType() == (bTriangleMode ?
-			FGenericMeshSelection::ETopologyType::FTriangleGroupTopology
-			: FGenericMeshSelection::ETopologyType::FGroupTopology)
-		&& InputSelectionIn->GetTargetComponent() == UE::ToolTarget::GetTargetComponent(Target)
-		&& !InputSelectionIn->IsEmpty();
-}
 
 void UEditMeshPolygonsTool::OnShutdown(EToolShutdownType ShutdownType)
 {
@@ -681,22 +672,22 @@ void UEditMeshPolygonsTool::OnShutdown(EToolShutdownType ShutdownType)
 
 		if (ShutdownType == EToolShutdownType::Accept)
 		{
-			UPersistentMeshSelection* OutputSelection = nullptr;
-			FCompactMaps CompactMaps;
-
-			// Prep if we have a selection to store. We don't support storing selections for volumes
-			// because the conversion will change vids.
-			if (!SelectionMechanic->GetActiveSelection().IsEmpty()
-				&& !Cast<UBrushComponent>(UE::ToolTarget::GetTargetComponent(Target)))
+			FGeometrySelection OutputSelection;
+			EGeometryElementType CurElemType = EGeometryElementType::Face;
+			if (SelectionMechanic->GetActiveSelection().SelectedCornerIDs.Num() > 0)
+			{
+				CurElemType = EGeometryElementType::Vertex;
+			}
+			else if (SelectionMechanic->GetActiveSelection().SelectedEdgeIDs.Num() > 0)
 			{
-				OutputSelection = NewObject<UPersistentMeshSelection>();
-				FGenericMeshSelection NewSelection;
-				NewSelection.SourceComponent = UE::ToolTarget::GetTargetComponent(Target);
-				NewSelection.TopologyType = (bTriangleMode ?
-					FGenericMeshSelection::ETopologyType::FTriangleGroupTopology
-					: FGenericMeshSelection::ETopologyType::FGroupTopology);
-				OutputSelection->SetSelection(NewSelection);
+				CurElemType = EGeometryElementType::Edge;
 			}
+			OutputSelection.InitializeTypes( CurElemType, (bTriangleMode) ? EGeometryTopologyType::Triangle : EGeometryTopologyType::Polygroup);
+
+			FCompactMaps CompactMaps;
+			//bool bIsBrushComponent = (Cast<UBrushComponent>(UE::ToolTarget::GetTargetComponent(Target)) == nullptr);
+			bool bIsBrushComponent = false;		// can we allow this now?
+			bool bWantSelection = (SelectionMechanic->GetActiveSelection().IsEmpty() == false);
 
 			// Note: When not in triangle mode, ModifiedTopologyCounter refers to polygroup topology, so does not tell us
 			// about the triangle topology.  In this case, we just assume the triangle topology may have been modified.
@@ -706,13 +697,20 @@ void UEditMeshPolygonsTool::OnShutdown(EToolShutdownType ShutdownType)
 			if (bModifiedTriangleTopology)
 			{
 				// Store the compact maps if we have a selection that we need to update
-				CurrentMesh->CompactInPlace(OutputSelection ? &CompactMaps : nullptr);
+				CurrentMesh->CompactInPlace(bWantSelection ? &CompactMaps : nullptr);
 			}
 
 			// Finish prepping the stored selection
-			if (OutputSelection)
+			if (bWantSelection)
 			{
-				SelectionMechanic->GetSelection(*OutputSelection, bModifiedTriangleTopology ? &CompactMaps : nullptr);
+				if (bTriangleMode)
+				{
+					SelectionMechanic->GetSelection_AsTriangleTopology(OutputSelection, bModifiedTriangleTopology ? &CompactMaps : nullptr);
+				}
+				else
+				{
+					SelectionMechanic->GetSelection_AsGroupTopology(OutputSelection, bModifiedTriangleTopology ? &CompactMaps : nullptr);
+				}
 			}
 
 			// Bake CurrentMesh back to target inside an undo transaction
@@ -720,7 +718,10 @@ void UEditMeshPolygonsTool::OnShutdown(EToolShutdownType ShutdownType)
 			MeshTransforms::ApplyTransformInverse(*CurrentMesh, BakedTransform, true);
 			UE::ToolTarget::CommitDynamicMeshUpdate(Target, *CurrentMesh, bModifiedTriangleTopology);
 
-			UE::Geometry::SetToolOutputSelection(this, OutputSelection);
+			if (OutputSelection.IsEmpty() == false)
+			{
+				UE::Geometry::SetToolOutputGeometrySelectionForTarget(this, Target, OutputSelection);
+			}
 		
 			GetToolManager()->EndUndoTransaction();
 		}
@@ -734,7 +735,6 @@ void UEditMeshPolygonsTool::OnShutdown(EToolShutdownType ShutdownType)
 
 	// We null out as many pointers as we can because the tool pointer usually ends up sticking
 	// around in the undo stack.
-	TargetWorld = nullptr;
 	CommonProps = nullptr;
 	EditActions = nullptr;
 	EditActions_Triangles = nullptr;
@@ -832,6 +832,27 @@ void UEditMeshPolygonsTool::RequestAction(EEditMeshPolygonsToolActions ActionTyp
 	PendingAction = ActionType;
 }
 
+
+void UEditMeshPolygonsTool::RequestSingleShotAction(EEditMeshPolygonsToolActions ActionType)
+{
+	bTerminateOnPendingActionComplete = true;
+
+	if ( ActionType == EEditMeshPolygonsToolActions::BevelAuto )
+	{
+		if (SelectionMechanic->GetActiveSelection().SelectedEdgeIDs.Num() > 0)
+		{
+			ActionType = EEditMeshPolygonsToolActions::BevelEdges;
+		}
+		else
+		{
+			ActionType = EEditMeshPolygonsToolActions::BevelFaces;
+		}
+	}
+
+	RequestAction(ActionType);
+}
+
+
 FDynamicMeshAABBTree3& UEditMeshPolygonsTool::GetSpatial()
 {
 	if (bSpatialDirty)
@@ -1236,7 +1257,7 @@ void UEditMeshPolygonsTool::OnTick(float DeltaTime)
 
 void UEditMeshPolygonsTool::StartActivity(TObjectPtr<UInteractiveToolActivity> Activity)
 {
-	EndCurrentActivity();
+	EndCurrentActivity(EToolShutdownType::Accept);
 
 	// Right now we rely on the activity to fail to start or to issue an error message if the
 	// conditions are not right. Someday, we are going to disable the buttons based on a CanStart
@@ -1251,16 +1272,33 @@ void UEditMeshPolygonsTool::StartActivity(TObjectPtr<UInteractiveToolActivity> A
 		SetToolPropertySourceEnabled(SelectionMechanic->Properties, false);
 		SetToolPropertySourceEnabled(TopologyProperties, false);
 		CurrentActivity = Activity;
-		if (CurrentActivity->HasAccept())
+
+		if ( bTerminateOnPendingActionComplete == false )
 		{
-			SetToolPropertySourceEnabled(AcceptCancelAction, true);
+			if (CurrentActivity->HasAccept())
+			{
+				SetToolPropertySourceEnabled(AcceptCancelAction, true);
+			}
+			else
+			{
+				SetToolPropertySourceEnabled(CancelAction, true);
+			}
 		}
 		else
 		{
-			SetToolPropertySourceEnabled(CancelAction, true);
+			SetToolPropertySourceEnabled(CommonProps, false);			
 		}
+
 		SetActionButtonPanelsVisible(false);
 	}
+	else
+	{
+		if (bTerminateOnPendingActionComplete)
+		{
+			GetToolManager()->PostActiveToolShutdownRequest(this, EToolShutdownType::Cancel);
+			return;
+		}
+	}
 }
 
 void UEditMeshPolygonsTool::EndCurrentActivity(EToolShutdownType ShutdownType)
@@ -1275,6 +1313,12 @@ void UEditMeshPolygonsTool::EndCurrentActivity(EToolShutdownType ShutdownType)
 		CurrentActivity = nullptr;
 		++ActivityTimestamp;
 
+		if (bTerminateOnPendingActionComplete)
+		{
+			GetToolManager()->PostActiveToolShutdownRequest(this, ShutdownType);
+			return;
+		}
+
 		SetToolPropertySourceEnabled(CancelAction, false);
 		SetToolPropertySourceEnabled(AcceptCancelAction, false);
 		SetActionButtonPanelsVisible(true);
@@ -1291,7 +1335,7 @@ void UEditMeshPolygonsTool::EndCurrentActivity(EToolShutdownType ShutdownType)
 
 void UEditMeshPolygonsTool::NotifyActivitySelfEnded(UInteractiveToolActivity* Activity)
 {
-	EndCurrentActivity();
+	EndCurrentActivity(EToolShutdownType::Accept);
 }
 
 void UEditMeshPolygonsTool::UpdateGizmoVisibility()
@@ -2481,6 +2525,8 @@ void UEditMeshPolygonsTool::SetActionButtonPanelsVisible(bool bVisible)
 
 bool UEditMeshPolygonsTool::CanCurrentlyNestedCancel()
 {
+	if ( bTerminateOnPendingActionComplete ) return false;
+
 	return CurrentActivity != nullptr 
 		|| (SelectionMechanic && !SelectionMechanic->GetActiveSelection().IsEmpty());
 }
@@ -2504,6 +2550,7 @@ bool UEditMeshPolygonsTool::ExecuteNestedCancelCommand()
 
 bool UEditMeshPolygonsTool::CanCurrentlyNestedAccept()
 {
+	if ( bTerminateOnPendingActionComplete ) return false;
 	return CurrentActivity != nullptr;
 }
 
@@ -2550,7 +2597,8 @@ FString FEditMeshPolygonsToolMeshChange::ToString() const
 
 void FPolyEditActivityStartChange::Revert(UObject* Object)
 {
-	Cast<UEditMeshPolygonsTool>(Object)->EndCurrentActivity();
+	//note: previously called EndCurrentActivity() which defauled to Cancel, so leaving that behavior here...
+	Cast<UEditMeshPolygonsTool>(Object)->EndCurrentActivity(EToolShutdownType::Cancel);
 	bHaveDoneUndo = true;
 }
 bool FPolyEditActivityStartChange::HasExpired(UObject* Object) const
diff --git a/Engine/Plugins/Runtime/MeshModelingToolset/Source/MeshModelingTools/Private/UVProjectionTool.cpp b/Engine/Plugins/Runtime/MeshModelingToolset/Source/MeshModelingTools/Private/UVProjectionTool.cpp
--- a/Engine/Plugins/Runtime/MeshModelingToolset/Source/MeshModelingTools/Private/UVProjectionTool.cpp
+++ b/Engine/Plugins/Runtime/MeshModelingToolset/Source/MeshModelingTools/Private/UVProjectionTool.cpp
@@ -11,6 +11,7 @@
 #include "BaseGizmos/TransformGizmoUtil.h"
 #include "ModelingToolTargetUtil.h"
 #include "Selection/StoredMeshSelectionUtil.h"
+#include "Selections/GeometrySelectionUtil.h"
 
 #include "DynamicMesh/DynamicMesh3.h"
 #include "DynamicMesh/DynamicMeshAABBTree3.h"
@@ -32,14 +33,14 @@ using namespace UE::Geometry;
  * ToolBuilder
  */
 
-USingleSelectionMeshEditingTool* UUVProjectionToolBuilder::CreateNewTool(const FToolBuilderState& SceneState) const
+USingleTargetWithSelectionTool* UUVProjectionToolBuilder::CreateNewTool(const FToolBuilderState& SceneState) const
 {
 	return NewObject<UUVProjectionTool>(SceneState.ToolManager);
 }
 
 bool UUVProjectionToolBuilder::CanBuildTool(const FToolBuilderState& SceneState) const
 {
-	return USingleSelectionMeshEditingToolBuilder::CanBuildTool(SceneState) && 
+	return USingleTargetWithSelectionToolBuilder::CanBuildTool(SceneState) && 
 		SceneState.TargetManager->CountSelectedAndTargetableWithPredicate(SceneState, GetTargetRequirements(),
 			[](UActorComponent& Component) { return ToolBuilderUtil::ComponentTypeCouldHaveUVs(Component); }) > 0;
 }
@@ -267,6 +268,12 @@ void UUVProjectionTool::OnShutdown(EToolShutdownType ShutdownType)
 		{
 			UE::ToolTarget::CommitDynamicMeshUVUpdate(Target, NewDynamicMesh);
 		}
+
+		if (HasGeometrySelection())
+		{
+			UE::Geometry::SetToolOutputGeometrySelectionForTarget(this, Target, GetGeometrySelection());
+		}
+
 		GetToolManager()->EndUndoTransaction();
 	}
 
@@ -324,9 +331,11 @@ void UUVProjectionTool::InitializeMesh()
 
 	// initialize triangle ROI if one exists
 	TriangleROI = MakeShared<TArray<int32>, ESPMode::ThreadSafe>();
-	if (USingleSelectionMeshEditingTool::HasInputSelection())
+	if (HasGeometrySelection())
 	{
-		UE::Geometry::GetStoredSelectionAsTriangles(GetInputSelection(), *InputMesh, *TriangleROI);
+		const FGeometrySelection& InputSelection = GetGeometrySelection();
+		UE::Geometry::EnumerateSelectionTriangles(InputSelection, *InputMesh, 
+			[&](int32 TriangleID) { TriangleROI->Add(TriangleID); });
 		TriangleROISet.Append(*TriangleROI);
 	}
 	VertexROI = MakeShared<TArray<int32>, ESPMode::ThreadSafe>();
diff --git a/Engine/Plugins/Runtime/MeshModelingToolset/Source/MeshModelingTools/Public/Commands/ModifyGeometrySelectionCommand.h b/Engine/Plugins/Runtime/MeshModelingToolset/Source/MeshModelingTools/Public/Commands/ModifyGeometrySelectionCommand.h
--- a/Engine/Plugins/Runtime/MeshModelingToolset/Source/MeshModelingTools/Public/Commands/ModifyGeometrySelectionCommand.h
+++ b/Engine/Plugins/Runtime/MeshModelingToolset/Source/MeshModelingTools/Public/Commands/ModifyGeometrySelectionCommand.h
@@ -30,7 +30,9 @@ public:
 	};
 	virtual EModificationType GetModificationType() const { return EModificationType::SelectAll; }
 
-	virtual bool IsModifySelectionCommand() const { return true; }
+	virtual bool AllowEmptySelection() const override { return GetModificationType() == EModificationType::SelectAll || GetModificationType() == EModificationType::Invert; }
+
+	virtual bool IsModifySelectionCommand() const override { return true; }
 	virtual FText GetCommandShortString() const override;
 
 	virtual bool CanExecuteCommandForSelection(UGeometrySelectionEditCommandArguments* Arguments) override;
diff --git a/Engine/Plugins/Runtime/MeshModelingToolset/Source/MeshModelingTools/Public/EditMeshPolygonsTool.h b/Engine/Plugins/Runtime/MeshModelingToolset/Source/MeshModelingTools/Public/EditMeshPolygonsTool.h
--- a/Engine/Plugins/Runtime/MeshModelingToolset/Source/MeshModelingTools/Public/EditMeshPolygonsTool.h
+++ b/Engine/Plugins/Runtime/MeshModelingToolset/Source/MeshModelingTools/Public/EditMeshPolygonsTool.h
@@ -11,7 +11,7 @@
 #include "InteractiveToolBuilder.h"
 #include "InteractiveToolQueryInterfaces.h" // IInteractiveToolNestedAcceptCancelAPI
 #include "Operations/GroupTopologyDeformer.h"
-#include "BaseTools/SingleSelectionMeshEditingTool.h"
+#include "BaseTools/SingleTargetWithSelectionTool.h"
 
 #include "GeometryBase.h"
 
@@ -24,7 +24,6 @@ class UCombinedTransformGizmo;
 class UDragAlignmentMechanic;
 class UMeshOpPreviewWithBackgroundCompute; 
 class FMeshVertexChangeBuilder;
-class UPersistentMeshSelection;
 class UEditMeshPolygonsTool;
 class UPolyEditActivityContext;
 class UPolyEditInsertEdgeActivity;
@@ -42,16 +41,16 @@ class UTransformProxy;
  * ToolBuilder
  */
 UCLASS()
-class MESHMODELINGTOOLS_API UEditMeshPolygonsToolBuilder : public USingleSelectionMeshEditingToolBuilder
+class MESHMODELINGTOOLS_API UEditMeshPolygonsToolBuilder : public USingleTargetWithSelectionToolBuilder
 {
 	GENERATED_BODY()
 public:
 	bool bTriangleMode = false;
 
-	virtual USingleSelectionMeshEditingTool* CreateNewTool(const FToolBuilderState& SceneState) const override;
-	virtual void InitializeNewTool(USingleSelectionMeshEditingTool* Tool, const FToolBuilderState& SceneState) const override;
+	virtual USingleTargetWithSelectionTool* CreateNewTool(const FToolBuilderState& SceneState) const override;
+	virtual void InitializeNewTool(USingleTargetWithSelectionTool* Tool, const FToolBuilderState& SceneState) const override;
 
-	virtual bool WantsInputSelectionIfAvailable() const { return true; }
+	virtual bool RequiresInputSelection() const override { return false; }
 };
 
 
@@ -141,7 +140,10 @@ enum class EEditMeshPolygonsToolActions
 	PokeSingleFace,
 	SplitSingleEdge,
 	FlipSingleEdge,
-	CollapseSingleEdge
+	CollapseSingleEdge,
+
+	// for external use
+	BevelAuto
 };
 
 UCLASS()
@@ -151,7 +153,7 @@ class MESHMODELINGTOOLS_API UEditMeshPolygonsActionModeToolBuilder : public UEdi
 public:
 	EEditMeshPolygonsToolActions StartupAction = EEditMeshPolygonsToolActions::Extrude;
 
-	virtual void InitializeNewTool(USingleSelectionMeshEditingTool* Tool, const FToolBuilderState& SceneState) const override;
+	virtual void InitializeNewTool(USingleTargetWithSelectionTool* Tool, const FToolBuilderState& SceneState) const override;
 };
 
 UENUM()
@@ -172,7 +174,7 @@ class MESHMODELINGTOOLS_API UEditMeshPolygonsSelectionModeToolBuilder : public U
 public:
 	EEditMeshPolygonsToolSelectionMode SelectionMode = EEditMeshPolygonsToolSelectionMode::Faces;
 
-	virtual void InitializeNewTool(USingleSelectionMeshEditingTool* Tool, const FToolBuilderState& SceneState) const override;
+	virtual void InitializeNewTool(USingleTargetWithSelectionTool* Tool, const FToolBuilderState& SceneState) const override;
 };
 
 
@@ -480,7 +482,7 @@ public:
  *
  */
 UCLASS()
-class MESHMODELINGTOOLS_API UEditMeshPolygonsTool : public USingleSelectionMeshEditingTool,
+class MESHMODELINGTOOLS_API UEditMeshPolygonsTool : public USingleTargetWithSelectionTool,
 	public IToolActivityHost, 
 	public IMeshVertexCommandChangeTarget,
 	public IInteractiveToolNestedAcceptCancelAPI
@@ -494,8 +496,6 @@ public:
 	virtual void RegisterActions(FInteractiveToolActionSet& ActionSet) override;
 	void EnableTriangleMode();
 
-	virtual void SetWorld(UWorld* World) { this->TargetWorld = World; }
-
 	// used by undo/redo
 	void RebuildTopologyWithGivenExtraCorners(const TSet<int32>& Vids);
 
@@ -533,6 +533,8 @@ public:
 
 	virtual void RequestAction(EEditMeshPolygonsToolActions ActionType);
 
+	virtual void RequestSingleShotAction(EEditMeshPolygonsToolActions ActionType);
+
 	void SetActionButtonsVisibility(bool bVisible);
 
 protected:
@@ -546,8 +548,6 @@ protected:
 	// hanging the editor.
 	bool bToolDisabled = false;
 
-	TObjectPtr<UWorld> TargetWorld = nullptr;
-
 	UPROPERTY()
 	TObjectPtr<UMeshOpPreviewWithBackgroundCompute> Preview = nullptr;
 
@@ -622,7 +622,6 @@ protected:
 
 	void ResetUserMessage();
 
-	bool IsToolInputSelectionUsable(const UPersistentMeshSelection* InputSelection);
 	bool bSelectionStateDirty = false;
 	void OnSelectionModifiedEvent();
 
@@ -653,11 +652,12 @@ protected:
 	float UVScaleFactor = 1.0f;
 
 	EEditMeshPolygonsToolActions PendingAction = EEditMeshPolygonsToolActions::NoAction;
+	bool bTerminateOnPendingActionComplete = false;
 
 	int32 ActivityTimestamp = 1;
 
 	void StartActivity(TObjectPtr<UInteractiveToolActivity> Activity);
-	void EndCurrentActivity(EToolShutdownType ShutdownType = EToolShutdownType::Cancel);
+	void EndCurrentActivity(EToolShutdownType ShutdownType);
 	void SetActionButtonPanelsVisible(bool bVisible);
 
 	// Emit an undoable change to CurrentMesh and update related structures (preview, spatial, etc)
diff --git a/Engine/Plugins/Runtime/MeshModelingToolset/Source/MeshModelingTools/Public/UVProjectionTool.h b/Engine/Plugins/Runtime/MeshModelingToolset/Source/MeshModelingTools/Public/UVProjectionTool.h
--- a/Engine/Plugins/Runtime/MeshModelingToolset/Source/MeshModelingTools/Public/UVProjectionTool.h
+++ b/Engine/Plugins/Runtime/MeshModelingToolset/Source/MeshModelingTools/Public/UVProjectionTool.h
@@ -3,7 +3,7 @@
 #pragma once
 
 #include "CoreMinimal.h"
-#include "BaseTools/SingleSelectionMeshEditingTool.h"
+#include "BaseTools/SingleTargetWithSelectionTool.h"
 #include "MeshOpPreviewHelpers.h"
 #include "ToolDataVisualizer.h"
 #include "ParameterizationOps/UVProjectionOp.h"
@@ -28,13 +28,13 @@ class UUVProjectionTool;
  *
  */
 UCLASS()
-class MESHMODELINGTOOLS_API UUVProjectionToolBuilder : public USingleSelectionMeshEditingToolBuilder
+class MESHMODELINGTOOLS_API UUVProjectionToolBuilder : public USingleTargetWithSelectionToolBuilder
 {
 	GENERATED_BODY()
 public:
-	virtual USingleSelectionMeshEditingTool* CreateNewTool(const FToolBuilderState& SceneState) const override;
+	virtual USingleTargetWithSelectionTool* CreateNewTool(const FToolBuilderState& SceneState) const override;
 
-	virtual bool WantsInputSelectionIfAvailable() const override { return true; }
+	virtual bool RequiresInputSelection() const override { return false; }
 	virtual bool CanBuildTool(const FToolBuilderState& SceneState) const override;
 };
 
@@ -202,7 +202,7 @@ public:
  * UV projection tool
  */
 UCLASS()
-class MESHMODELINGTOOLS_API UUVProjectionTool : public USingleSelectionMeshEditingTool
+class MESHMODELINGTOOLS_API UUVProjectionTool : public USingleTargetWithSelectionTool
 {
 	GENERATED_BODY()
 
diff --git a/Engine/Plugins/Runtime/MeshModelingToolset/Source/ModelingComponents/Private/BaseTools/SingleSelectionMeshEditingTool.cpp b/Engine/Plugins/Runtime/MeshModelingToolset/Source/ModelingComponents/Private/BaseTools/SingleSelectionMeshEditingTool.cpp
--- a/Engine/Plugins/Runtime/MeshModelingToolset/Source/ModelingComponents/Private/BaseTools/SingleSelectionMeshEditingTool.cpp
+++ b/Engine/Plugins/Runtime/MeshModelingToolset/Source/ModelingComponents/Private/BaseTools/SingleSelectionMeshEditingTool.cpp
@@ -46,14 +46,16 @@ void USingleSelectionMeshEditingToolBuilder::InitializeNewTool(USingleSelectionM
 	NewTool->SetTarget(Target);
 	NewTool->SetWorld(SceneState.World);
 
+PRAGMA_DISABLE_DEPRECATION_WARNINGS
 	if (WantsInputSelectionIfAvailable())
 	{
-		const UPersistentMeshSelection* InputSelection = UE::Geometry::GetCurrentToolInputSelection(SceneState, Target);
+		const UDEPRECATED_PersistentMeshSelection* InputSelection = UE::Geometry::GetCurrentToolInputSelection(SceneState, Target);
 		if (InputSelection != nullptr)
 		{
 			NewTool->SetInputSelection(InputSelection);
 		}
 	}
+PRAGMA_ENABLE_DEPRECATION_WARNINGS
 }
 
 
diff --git a/Engine/Plugins/Runtime/MeshModelingToolset/Source/ModelingComponents/Private/BaseTools/SingleTargetWithSelectionTool.cpp b/Engine/Plugins/Runtime/MeshModelingToolset/Source/ModelingComponents/Private/BaseTools/SingleTargetWithSelectionTool.cpp
--- a/Engine/Plugins/Runtime/MeshModelingToolset/Source/ModelingComponents/Private/BaseTools/SingleTargetWithSelectionTool.cpp
+++ b/Engine/Plugins/Runtime/MeshModelingToolset/Source/ModelingComponents/Private/BaseTools/SingleTargetWithSelectionTool.cpp
@@ -31,6 +31,11 @@ const FToolTargetTypeRequirements& USingleTargetWithSelectionToolBuilder::GetTar
 
 bool USingleTargetWithSelectionToolBuilder::CanBuildTool(const FToolBuilderState& SceneState) const
 {
+	if (RequiresInputSelection() && UE::Geometry::HaveAvailableGeometrySelection(SceneState) == false )
+	{
+		return false;
+	}
+
 	return SceneState.TargetManager->CountSelectedAndTargetable(SceneState, GetTargetRequirements()) == 1;
 }
 
diff --git a/Engine/Plugins/Runtime/MeshModelingToolset/Source/ModelingComponents/Private/Selection/DynamicMeshSelector.cpp b/Engine/Plugins/Runtime/MeshModelingToolset/Source/ModelingComponents/Private/Selection/DynamicMeshSelector.cpp
--- a/Engine/Plugins/Runtime/MeshModelingToolset/Source/ModelingComponents/Private/Selection/DynamicMeshSelector.cpp
+++ b/Engine/Plugins/Runtime/MeshModelingToolset/Source/ModelingComponents/Private/Selection/DynamicMeshSelector.cpp
@@ -64,7 +64,6 @@ void FBaseDynamicMeshSelector::InvalidateOnMeshChange(FDynamicMeshChangeInfo Cha
 	NotifyGeometryModified();
 }
 
-
 void FBaseDynamicMeshSelector::Shutdown()
 {
 	if (TargetMesh.IsValid())
@@ -168,7 +167,24 @@ void FBaseDynamicMeshSelector::UpdateSelectionFromSelection(
 		return;
 	}
 
-	check(false);		// todo: convert selection
+	bool bConverted = false;
+	GetDynamicMesh()->ProcessMesh([&](const FDynamicMesh3& Mesh) 
+	{ 
+		FGeometrySelection TempSelection;
+		TempSelection.InitializeTypes(SelectionEditor.GetElementType(), SelectionEditor.GetTopologyType());
+		bConverted = UE::Geometry::ConvertSelection(Mesh, GetGroupTopology(), FromSelection, TempSelection);		// todo do not always need group topology here...
+		if ( bConverted )
+		{
+			UE::Geometry::UpdateSelectionWithNewElements(&SelectionEditor, UpdateConfig.ChangeType,
+				TempSelection.Selection.Array(), SelectionDelta);
+		}
+	});
+	if (bConverted)
+	{
+		return;
+	}
+
+	ensure(false);		// todo: support this type of convert selection
 }
 
 
diff --git a/Engine/Plugins/Runtime/MeshModelingToolset/Source/ModelingComponents/Private/Selection/GeometrySelectionManager.cpp b/Engine/Plugins/Runtime/MeshModelingToolset/Source/ModelingComponents/Private/Selection/GeometrySelectionManager.cpp
--- a/Engine/Plugins/Runtime/MeshModelingToolset/Source/ModelingComponents/Private/Selection/GeometrySelectionManager.cpp
+++ b/Engine/Plugins/Runtime/MeshModelingToolset/Source/ModelingComponents/Private/Selection/GeometrySelectionManager.cpp
@@ -16,8 +16,6 @@
 #include "DynamicMeshBuilder.h"
 #include "ToolSetupUtil.h"
 
-#include "Selection/PersistentMeshSelection.h"
-
 #include UE_INLINE_GENERATED_CPP_BY_NAME(GeometrySelectionManager)
 
 static TAutoConsoleVariable<int32> CVarGeometrySelectionManager_FullSelectionHoverHighlights(
@@ -669,7 +667,33 @@ void UGeometrySelectionManager::EndTrackedSelectionChange()
 }
 
 
+bool UGeometrySelectionManager::SetSelectionForComponent(UPrimitiveComponent* Component, const FGeometrySelection& NewSelection)
+{
+	for (TSharedPtr<FGeometrySelectionTarget> Target : ActiveTargetReferences)
+	{
+		if (Target->TargetIdentifier.TargetObject == Component)
+		{
+			FGeometrySelection InitialSelection = Target->Selection;
+			FGeometrySelectionDelta AfterDelta;
+			Target->Selector->UpdateSelectionFromSelection(
+				NewSelection, true, *Target->SelectionEditor, 
+				FGeometrySelectionUpdateConfig{EGeometrySelectionChangeType::Replace}, &AfterDelta);
+			if ( AfterDelta.IsEmpty() == false )
+			{
+				TUniquePtr<FGeometrySelectionReplaceChange> NewSelectionChange = MakeUnique<FGeometrySelectionReplaceChange>();
+				NewSelectionChange->Identifier = Target->Selector->GetIdentifier();
+				NewSelectionChange->After = Target->Selection;
+				NewSelectionChange->Before = InitialSelection;
+				GetTransactionsAPI()->AppendChange(this, MoveTemp(NewSelectionChange), LOCTEXT("NewSelection", "New Selection"));
 
+				bSelectionRenderCachesDirty = true;
+				OnSelectionModified.Broadcast();
+			}
+			return true;
+		}
+	}
+	return false;
+}
 
 
 bool UGeometrySelectionManager::UpdateSelectionPreviewViaRaycast(
@@ -1180,56 +1204,5 @@ void UGeometrySelectionManager::DebugRender(IToolsContextRenderAPI* RenderAPI)
 
 
 
-UPersistentMeshSelection* UGeometrySelectionManager::GetActiveSingleSelectionConverted_Legacy(UPrimitiveComponent* ForComponent)
-{
-	OldSelection = nullptr;
-
-	if (HasSelection() == false)
-	{
-		return nullptr;
-	}
-
-	for (TSharedPtr<FGeometrySelectionTarget> Target : ActiveTargetReferences)
-	{
-		if (Target->TargetIdentifier.TargetObject == ForComponent)
-		{
-			const FGeometrySelection& GeoSelection = Target->Selection;
-
-			OldSelection = NewObject<UPersistentMeshSelection>(this);
-			FGenericMeshSelection Selection;
-			Selection.SourceComponent = ForComponent;
-
-			if (GeoSelection.TopologyType == EGeometryTopologyType::Triangle)
-			{
-				Selection.TopologyType = FGenericMeshSelection::ETopologyType::FTriangleGroupTopology;
-				if (GeoSelection.ElementType == EGeometryElementType::Face)
-				{
-					for (uint64 tid : GeoSelection.Selection)
-					{
-						Selection.FaceIDs.Add((int32)tid);
-					}
-				}
-			}
-			else
-			{
-				Selection.TopologyType = FGenericMeshSelection::ETopologyType::FGroupTopology;
-				if (GeoSelection.ElementType == EGeometryElementType::Face)
-				{
-					for (uint64 EncodedID : GeoSelection.Selection)
-					{
-						FGeoSelectionID SelectionID(EncodedID);
-						Selection.FaceIDs.Add((int32)SelectionID.TopologyID);
-					}
-				}
-			}
-
-			
-			OldSelection->SetSelection(Selection);
-		}
-	}
-
-	return OldSelection;
-}
-
 
 #undef LOCTEXT_NAMESPACE
diff --git a/Engine/Plugins/Runtime/MeshModelingToolset/Source/ModelingComponents/Private/Selection/PersistentMeshSelection.cpp b/Engine/Plugins/Runtime/MeshModelingToolset/Source/ModelingComponents/Private/Selection/PersistentMeshSelection.cpp
--- a/Engine/Plugins/Runtime/MeshModelingToolset/Source/ModelingComponents/Private/Selection/PersistentMeshSelection.cpp
+++ b/Engine/Plugins/Runtime/MeshModelingToolset/Source/ModelingComponents/Private/Selection/PersistentMeshSelection.cpp
@@ -12,7 +12,7 @@
 using namespace UE::Geometry;
 
 
-void UPersistentMeshSelection::SetSelection(const FGroupTopology& TopologyIn, const FGroupTopologySelection& SelectionIn,
+void UDEPRECATED_PersistentMeshSelection::SetSelection(const FGroupTopology& TopologyIn, const FGroupTopologySelection& SelectionIn,
 											const FCompactMaps* CompactMaps)
 {
 	FGenericMeshSelection& Data = this->Selection;
@@ -81,7 +81,7 @@ void UPersistentMeshSelection::SetSelection(const FGroupTopology& TopologyIn, co
 }
 
 
-void UPersistentMeshSelection::ExtractIntoSelectionObject(const FGroupTopology& TopologyIn, FGroupTopologySelection& SelectionOut) const
+void UDEPRECATED_PersistentMeshSelection::ExtractIntoSelectionObject(const FGroupTopology& TopologyIn, FGroupTopologySelection& SelectionOut) const
 {
 	const FGenericMeshSelection& Data = this->Selection;
 
diff --git a/Engine/Plugins/Runtime/MeshModelingToolset/Source/ModelingComponents/Private/Selection/PersistentMeshSelectionManager.cpp b/Engine/Plugins/Runtime/MeshModelingToolset/Source/ModelingComponents/Private/Selection/PersistentMeshSelectionManager.cpp
--- a/Engine/Plugins/Runtime/MeshModelingToolset/Source/ModelingComponents/Private/Selection/PersistentMeshSelectionManager.cpp
+++ b/Engine/Plugins/Runtime/MeshModelingToolset/Source/ModelingComponents/Private/Selection/PersistentMeshSelectionManager.cpp
@@ -33,20 +33,20 @@ void UDEPRECATED_PersistentMeshSelectionManager::Shutdown()
 
 bool UDEPRECATED_PersistentMeshSelectionManager::HasActiveSelection()
 {
-	return (ActiveSelection != nullptr);
+	return (ActiveSelection_DEPRECATED != nullptr);
 }
 
-UPersistentMeshSelection* UDEPRECATED_PersistentMeshSelectionManager::GetActiveSelection()
+UDEPRECATED_PersistentMeshSelection* UDEPRECATED_PersistentMeshSelectionManager::GetActiveSelection()
 {
-	return ActiveSelection;
+	return ActiveSelection_DEPRECATED;
 }
 
-void UDEPRECATED_PersistentMeshSelectionManager::SetNewActiveSelection(UPersistentMeshSelection* Selection)
+void UDEPRECATED_PersistentMeshSelectionManager::SetNewActiveSelection(UDEPRECATED_PersistentMeshSelection* Selection)
 {
 	TUniquePtr<FPersistentMeshSelectionChange> SelectionChange = MakeUnique<FPersistentMeshSelectionChange>();
-	if (ActiveSelection != nullptr)
+	if (ActiveSelection_DEPRECATED != nullptr)
 	{
-		SelectionChange->From = ActiveSelection->GetSelection();
+		SelectionChange->From = ActiveSelection_DEPRECATED->GetSelection();
 	}
 	if (Selection != nullptr)
 	{
@@ -55,21 +55,21 @@ void UDEPRECATED_PersistentMeshSelectionManager::SetNewActiveSelection(UPersiste
 	ParentContext->ToolManager->GetContextTransactionsAPI()->AppendChange(this, MoveTemp(SelectionChange),
 		LOCTEXT("SelectionChange", "Selection Change"));
 
-	ActiveSelection = Selection;
+	ActiveSelection_DEPRECATED = Selection;
 	OnSelectionModified();
 }
 
 
-void UDEPRECATED_PersistentMeshSelectionManager::SetNewActiveSelectionInternal(UPersistentMeshSelection* Selection)
+void UDEPRECATED_PersistentMeshSelectionManager::SetNewActiveSelectionInternal(UDEPRECATED_PersistentMeshSelection* Selection)
 {
-	ActiveSelection = Selection;
+	ActiveSelection_DEPRECATED = Selection;
 	OnSelectionModified();
 }
 
 
 void UDEPRECATED_PersistentMeshSelectionManager::ClearActiveSelection()
 {
-	if (ActiveSelection == nullptr)
+	if (ActiveSelection_DEPRECATED == nullptr)
 	{
 		return;
 	}
@@ -82,17 +82,17 @@ void UDEPRECATED_PersistentMeshSelectionManager::ClearActiveSelection()
 
 void UDEPRECATED_PersistentMeshSelectionManager::OnSelectionModified()
 {
-	if (ActiveSelection != nullptr && SelectionDisplay == nullptr)
+	if (ActiveSelection_DEPRECATED != nullptr && SelectionDisplay == nullptr)
 	{
 		SelectionDisplay = NewObject<UPreviewGeometry>(ParentContext);
-		SelectionDisplay->CreateInWorld(ActiveSelection->GetTargetComponent()->GetWorld(), FTransform::Identity);
+		SelectionDisplay->CreateInWorld(ActiveSelection_DEPRECATED->GetTargetComponent()->GetWorld(), FTransform::Identity);
 	}
 	if (SelectionDisplay == nullptr)
 	{
 		return;
 	}
 
-	const FGenericMeshSelection* SelectionData = (ActiveSelection != nullptr) ? &ActiveSelection->GetSelection() : nullptr;
+	const FGenericMeshSelection* SelectionData = (ActiveSelection_DEPRECATED != nullptr) ? &ActiveSelection_DEPRECATED->GetSelection() : nullptr;
 
 	bool bShowLines = (SelectionData != nullptr) && (SelectionData->HasRenderableLines());
 
@@ -103,7 +103,7 @@ void UDEPRECATED_PersistentMeshSelectionManager::OnSelectionModified()
 		//const float ROIBorderDepthBias = 0.1f * (float)(WorldBounds.DiagonalLength() * 0.01);
 		const float ROIBorderDepthBias = 0.01f;
 
-		FTransform3d Transform(ActiveSelection->GetTargetComponent()->GetComponentToWorld());
+		FTransform3d Transform(ActiveSelection_DEPRECATED->GetTargetComponent()->GetComponentToWorld());
 
 		const TArray<UE::Geometry::FSegment3d>& Lines = SelectionData->RenderEdges;
 		SelectionDisplay->CreateOrUpdateLineSet(TEXT("SelectionEdges"),  Lines.Num(),
@@ -190,7 +190,7 @@ void FPersistentMeshSelectionChange::Apply(UObject* Object)
 	UDEPRECATED_PersistentMeshSelectionManager* SelectionManager = Cast<UDEPRECATED_PersistentMeshSelectionManager>(Object);
 	if (SelectionManager)
 	{
-		UPersistentMeshSelection* NewSelection = NewObject<UPersistentMeshSelection>(SelectionManager);
+		UDEPRECATED_PersistentMeshSelection* NewSelection = NewObject<UDEPRECATED_PersistentMeshSelection>(SelectionManager);
 		NewSelection->SetSelection(To);
 		SelectionManager->SetNewActiveSelectionInternal(NewSelection);
 	}
@@ -201,7 +201,7 @@ void FPersistentMeshSelectionChange::Revert(UObject* Object)
 	UDEPRECATED_PersistentMeshSelectionManager* SelectionManager = Cast<UDEPRECATED_PersistentMeshSelectionManager>(Object);
 	if (SelectionManager)
 	{
-		UPersistentMeshSelection* NewSelection = NewObject<UPersistentMeshSelection>(SelectionManager);
+		UDEPRECATED_PersistentMeshSelection* NewSelection = NewObject<UDEPRECATED_PersistentMeshSelection>(SelectionManager);
 		NewSelection->SetSelection(From);
 		SelectionManager->SetNewActiveSelectionInternal(NewSelection);
 	}
diff --git a/Engine/Plugins/Runtime/MeshModelingToolset/Source/ModelingComponents/Private/Selection/PolygonSelectionMechanic.cpp b/Engine/Plugins/Runtime/MeshModelingToolset/Source/ModelingComponents/Private/Selection/PolygonSelectionMechanic.cpp
--- a/Engine/Plugins/Runtime/MeshModelingToolset/Source/ModelingComponents/Private/Selection/PolygonSelectionMechanic.cpp
+++ b/Engine/Plugins/Runtime/MeshModelingToolset/Source/ModelingComponents/Private/Selection/PolygonSelectionMechanic.cpp
@@ -4,6 +4,7 @@
 #include "Engine/World.h"
 #include "Selection/GroupTopologySelector.h"
 #include "Selection/PersistentMeshSelection.h"
+#include "Selections/GeometrySelection.h"
 
 #include UE_INLINE_GENERATED_CPP_BY_NAME(PolygonSelectionMechanic)
 
@@ -60,16 +61,190 @@ void UPolygonSelectionMechanic::Setup(UInteractiveTool* ParentToolIn)
 }
 
 
-void UPolygonSelectionMechanic::GetSelection(UPersistentMeshSelection& SelectionOut, const FCompactMaps* CompactMapsToApply) const
+void UPolygonSelectionMechanic::GetSelection_AsGroupTopology(UE::Geometry::FGeometrySelection& SelectionOut, const FCompactMaps* CompactMapsToApply) const
 {
-	SelectionOut.SetSelection(*Topology, PersistentSelection, CompactMapsToApply);
+	const FGroupTopologySelection& CurSelection = PersistentSelection;
+	if (SelectionOut.TopologyType != EGeometryTopologyType::Polygroup)
+	{
+		return;
+	}
+	if (SelectionOut.ElementType == EGeometryElementType::Vertex)
+	{
+		for (int32 CornerID : CurSelection.SelectedCornerIDs)
+		{
+			int32 VertexID = Topology->GetCornerVertexID(CornerID);
+			if (CompactMapsToApply != nullptr)
+			{
+				VertexID = CompactMapsToApply->GetVertexMapping(VertexID);
+			}
+			SelectionOut.Selection.Add( FGeoSelectionID(VertexID, CornerID).Encoded() );
+		}
+	}
+	else if (SelectionOut.ElementType == EGeometryElementType::Edge)
+	{
+		// TODO: not sure how we can apply compact maps here because mapping does not included compacted edges? are edges even compacted?
+		if ( CompactMapsToApply == nullptr )
+		{
+			for (int32 GroupEdgeID : CurSelection.SelectedEdgeIDs)
+			{
+				const TArray<int>& GroupEdge = Topology->GetGroupEdgeEdges(GroupEdgeID);
+				FMeshTriEdgeID TriEdgeID = Topology->GetMesh()->GetTriEdgeIDFromEdgeID(GroupEdge[0]);
+				SelectionOut.Selection.Add( FGeoSelectionID(TriEdgeID.Encoded(), GroupEdgeID).Encoded() );
+			}
+		}
+	}
+	else if (SelectionOut.ElementType == EGeometryElementType::Face)
+	{
+		for (int32 GroupID : CurSelection.SelectedGroupIDs)
+		{
+			const FGroupTopology::FGroup* GroupFace = Topology->FindGroupByID(GroupID);
+			if ( GroupFace )
+			{
+				FGeoSelectionID ID = FGeoSelectionID(GroupFace->Triangles[0], GroupFace->GroupID);
+				SelectionOut.Selection.Add(ID.Encoded());
+			}
+		}
+	}
+}
+
+
+void UPolygonSelectionMechanic::GetSelection_AsTriangleTopology(UE::Geometry::FGeometrySelection& SelectionOut, const FCompactMaps* CompactMapsToApply) const
+{
+	// note: this is currently the same code as GetSelection_AsGroupTopology() except for the topology-type verification check
+
+	const FGroupTopologySelection& CurSelection = PersistentSelection;
+	if (SelectionOut.TopologyType != EGeometryTopologyType::Triangle)
+	{
+		return;
+	}
+	if (SelectionOut.ElementType == EGeometryElementType::Vertex)
+	{
+		for (int32 CornerID : CurSelection.SelectedCornerIDs)
+		{
+			int32 VertexID = Topology->GetCornerVertexID(CornerID);
+			if (CompactMapsToApply != nullptr)
+			{
+				VertexID = CompactMapsToApply->GetVertexMapping(VertexID);
+			}
+			SelectionOut.Selection.Add( FGeoSelectionID::MeshVertex(VertexID).Encoded() );
+		}
+	}
+	else if (SelectionOut.ElementType == EGeometryElementType::Edge)
+	{
+		// TODO: not sure how we can apply compact maps here because mapping does not included compacted edges? are edges even compacted?
+		if ( CompactMapsToApply == nullptr )
+		{
+			for (int32 GroupEdgeID : CurSelection.SelectedEdgeIDs)
+			{
+				const TArray<int>& GroupEdge = Topology->GetGroupEdgeEdges(GroupEdgeID);
+				FMeshTriEdgeID TriEdgeID = Topology->GetMesh()->GetTriEdgeIDFromEdgeID(GroupEdge[0]);
+				SelectionOut.Selection.Add( FGeoSelectionID::MeshEdge(TriEdgeID).Encoded() );
+			}
+		}
+	}
+	else if (SelectionOut.ElementType == EGeometryElementType::Face)
+	{
+		for (int32 GroupID : CurSelection.SelectedGroupIDs)
+		{
+			const FGroupTopology::FGroup* GroupFace = Topology->FindGroupByID(GroupID);
+			if ( GroupFace )
+			{
+				SelectionOut.Selection.Add( FGeoSelectionID::MeshTriangle(GroupFace->Triangles[0]).Encoded() );
+			}
+		}
+	}
 }
 
-void UPolygonSelectionMechanic::LoadSelection(const UPersistentMeshSelection& SelectionIn)
+
+void UPolygonSelectionMechanic::SetSelection_AsGroupTopology(const UE::Geometry::FGeometrySelection& Selection)
 {
-	SelectionIn.ExtractIntoSelectionObject(*Topology, PersistentSelection);
+	if (Selection.TopologyType != EGeometryTopologyType::Polygroup)
+	{
+		return;
+	}
+	if (Selection.ElementType == EGeometryElementType::Vertex)
+	{
+		for (uint64 ElementID : Selection.Selection)
+		{
+			int32 VertexID = FGeoSelectionID(ElementID).GeometryID;
+			int32 CornerID = Topology->GetCornerIDFromVertexID(VertexID);
+			if (CornerID != IndexConstants::InvalidID)
+			{
+				PersistentSelection.SelectedCornerIDs.Add(CornerID);
+			}
+		}
+	}
+	else if (Selection.ElementType == EGeometryElementType::Edge)
+	{
+		for (uint64 ElementID : Selection.Selection)
+		{
+			FMeshTriEdgeID TriEdgeID(FGeoSelectionID(ElementID).GeometryID);
+			int32 GroupEdgeID = Topology->FindGroupEdgeID(TriEdgeID);
+			if (GroupEdgeID != IndexConstants::InvalidID)
+			{
+				PersistentSelection.SelectedEdgeIDs.Add(GroupEdgeID);
+			}
+		}
+	}
+	else if (Selection.ElementType == EGeometryElementType::Face)
+	{
+		for (uint64 ElementID : Selection.Selection)
+		{
+			int32 TriangleID = FGeoSelectionID(ElementID).GeometryID;
+			int32 GroupID = Topology->GetGroupID(TriangleID);
+			if (Topology->FindGroupByID(GroupID) != nullptr)
+			{
+				PersistentSelection.SelectedGroupIDs.Add(GroupID);
+			}
+		}
+	}
 }
 
+void UPolygonSelectionMechanic::SetSelection_AsTriangleTopology(const UE::Geometry::FGeometrySelection& Selection)
+{
+	if (Selection.TopologyType != EGeometryTopologyType::Triangle)
+	{
+		return;
+	}
+	if (Selection.ElementType == EGeometryElementType::Vertex)
+	{
+		for (uint64 ElementID : Selection.Selection)
+		{
+			int32 VertexID = FGeoSelectionID(ElementID).GeometryID;
+			int32 CornerID = Topology->GetCornerIDFromVertexID(VertexID);
+			if (CornerID != IndexConstants::InvalidID)
+			{
+				PersistentSelection.SelectedCornerIDs.Add(CornerID);
+			}
+		}
+	}
+	else if (Selection.ElementType == EGeometryElementType::Edge)
+	{
+		for (uint64 ElementID : Selection.Selection)
+		{
+			FMeshTriEdgeID TriEdgeID(FGeoSelectionID(ElementID).GeometryID);
+			int32 GroupEdgeID = Topology->FindGroupEdgeID(TriEdgeID);
+			if (GroupEdgeID != IndexConstants::InvalidID)
+			{
+				PersistentSelection.SelectedEdgeIDs.Add(GroupEdgeID);
+			}
+		}
+	}
+	else if (Selection.ElementType == EGeometryElementType::Face)
+	{
+		for (uint64 ElementID : Selection.Selection)
+		{
+			int32 TriangleID = FGeoSelectionID(ElementID).GeometryID;
+			int32 GroupID = Topology->GetGroupID(TriangleID);
+			if (Topology->FindGroupByID(GroupID) != nullptr)
+			{
+				PersistentSelection.SelectedGroupIDs.Add(GroupID);
+			}
+		}
+	}
+}
+
+
 bool UPolygonSelectionMechanic::UpdateHighlight(const FRay& WorldRay)
 {
 	checkf(DrawnTriangleSetComponent != nullptr, TEXT("Initialize() not called on UMeshTopologySelectionMechanic."));
diff --git a/Engine/Plugins/Runtime/MeshModelingToolset/Source/ModelingComponents/Private/Selection/StoredMeshSelectionUtil.cpp b/Engine/Plugins/Runtime/MeshModelingToolset/Source/ModelingComponents/Private/Selection/StoredMeshSelectionUtil.cpp
--- a/Engine/Plugins/Runtime/MeshModelingToolset/Source/ModelingComponents/Private/Selection/StoredMeshSelectionUtil.cpp
+++ b/Engine/Plugins/Runtime/MeshModelingToolset/Source/ModelingComponents/Private/Selection/StoredMeshSelectionUtil.cpp
@@ -13,11 +13,22 @@
 using namespace UE::Geometry;
 
 
+bool UE::Geometry::HaveAvailableGeometrySelection(const FToolBuilderState& SceneState)
+{
+	UGeometrySelectionManager* SelectionManager = SceneState.ToolManager->GetContextObjectStore()->FindContext<UGeometrySelectionManager>();
+	return (SelectionManager != nullptr && SelectionManager->HasSelection());
+}
 
 
 bool UE::Geometry::GetCurrentGeometrySelectionForTarget(const FToolBuilderState& SceneState, UToolTarget* Target, FGeometrySelection& SelectionOut)
 {
-	UGeometrySelectionManager* SelectionManager = SceneState.ToolManager->GetContextObjectStore()->FindContext<UGeometrySelectionManager>();
+	return GetCurrentGeometrySelectionForTarget(SceneState.ToolManager, Target, SelectionOut);
+}
+
+
+bool UE::Geometry::GetCurrentGeometrySelectionForTarget(UInteractiveToolManager* ToolManager, UToolTarget* Target, FGeometrySelection& SelectionOut)
+{
+	UGeometrySelectionManager* SelectionManager = ToolManager->GetContextObjectStore()->FindContext<UGeometrySelectionManager>();
 	if (SelectionManager == nullptr)
 	{
 		return false;
@@ -33,118 +44,56 @@ bool UE::Geometry::GetCurrentGeometrySelectionForTarget(const FToolBuilderState&
 }
 
 
-const UPersistentMeshSelection* UE::Geometry::GetCurrentToolInputSelection(const FToolBuilderState& SceneState, UToolTarget* Target)
+bool UE::Geometry::SetToolOutputGeometrySelectionForTarget(UInteractiveTool* Tool, UToolTarget* Target, const FGeometrySelection& OutputSelection)
 {
-	UGeometrySelectionManager* SelectionManager = SceneState.ToolManager->GetContextObjectStore()->FindContext<UGeometrySelectionManager>();
+	UGeometrySelectionManager* SelectionManager = Tool->GetToolManager()->GetContextObjectStore()->FindContext<UGeometrySelectionManager>();
 	if (SelectionManager == nullptr)
 	{
-		return nullptr;
+		return false;
 	}
 	IPrimitiveComponentBackedTarget* TargetInterface = Cast<IPrimitiveComponentBackedTarget>(Target);
 	if (TargetInterface == nullptr)
 	{
-		return nullptr;
+		return false;
 	}
+	return SelectionManager->SetSelectionForComponent(TargetInterface->GetOwnerComponent(), OutputSelection);
+}
 
-	UPersistentMeshSelection* ActiveSelection = SelectionManager->GetActiveSingleSelectionConverted_Legacy(TargetInterface->GetOwnerComponent());
-	if (ActiveSelection == nullptr)
-	{
-		return nullptr;
-	}
 
-	if (TargetInterface->GetOwnerComponent() != ActiveSelection->GetTargetComponent())
-	{
-		return nullptr;
-	}
 
-	return ActiveSelection;
+
+
+
+
+const UDEPRECATED_PersistentMeshSelection* UE::Geometry::GetCurrentToolInputSelection(const FToolBuilderState& SceneState, UToolTarget* Target)
+{
+	ensureMsgf(false, TEXT("This path no longer functions and will be removed"));
+	return nullptr;
 }
 
 
 
 bool UE::Geometry::GetStoredSelectionAsTriangles(
-	const UPersistentMeshSelection* Selection,
+	const UDEPRECATED_PersistentMeshSelection* Selection,
 	const FDynamicMesh3& Mesh,
 	TArray<int32>& TrianglesOut)
 {
-	// don't support UV selection currently - unclear how reproducible this is?
-	if (Selection->GetSelectionType() == FGenericMeshSelection::ETopologyType::FUVGroupTopology)
-	{
-		return false;
-	}
-
-	if (Selection->GetSelection().FaceIDs.Num() == 0)
-	{
-		return false;
-	}
-
-	if (Selection->GetSelectionType() == FGenericMeshSelection::ETopologyType::FGroupTopology)
-	{
-		TSet<int32> SelectedGroups;
-		for (int32 gid : Selection->GetSelection().FaceIDs)
-		{
-			SelectedGroups.Add(gid);
-		}
-
-		for (int32 tid : Mesh.TriangleIndicesItr())
-		{
-			int32 gid = Mesh.GetTriangleGroup(tid);
-			if (SelectedGroups.Contains(gid))
-			{
-				TrianglesOut.Add(tid);
-			}
-		}
-	}
-	else if (Selection->GetSelectionType() == FGenericMeshSelection::ETopologyType::FTriangleGroupTopology)
-	{
-		for (int gid : Selection->GetSelection().FaceIDs)
-		{
-			if (Mesh.IsTriangle(gid))
-			{
-				TrianglesOut.Add(gid);
-			}
-		}
-	}
-
-	return true;
+	return false;
 }
 
 
-
-
 void UE::Geometry::ClearActiveToolSelection(UInteractiveToolManager* ToolManager)
 {
-PRAGMA_DISABLE_DEPRECATION_WARNINGS
-	UDEPRECATED_PersistentMeshSelectionManager* SelectionManager = FindPersistentMeshSelectionManager(ToolManager);
-	if (SelectionManager != nullptr)
-	{
-		SelectionManager->ClearActiveSelection();
-	}
-PRAGMA_ENABLE_DEPRECATION_WARNINGS
 }
 
 
-bool UE::Geometry::SetToolOutputSelection(UInteractiveTool* Tool, UPersistentMeshSelection* Selection)
+bool UE::Geometry::SetToolOutputSelection(UInteractiveTool* Tool, UDEPRECATED_PersistentMeshSelection* Selection)
 {
-PRAGMA_DISABLE_DEPRECATION_WARNINGS
-	UDEPRECATED_PersistentMeshSelectionManager* SelectionManager = FindPersistentMeshSelectionManager(Tool->GetToolManager());
-	if (SelectionManager == nullptr)
-	{
-		return false;
-	}
+	return false;
+}
+
 
-	if (Selection != nullptr)
-	{
-		SelectionManager->SetNewActiveSelection(Selection);
-	}
-	else
-	{
-		SelectionManager->ClearActiveSelection();
-	}
 
-	return true;
-PRAGMA_ENABLE_DEPRECATION_WARNINGS
-}
 
 
 
diff --git a/Engine/Plugins/Runtime/MeshModelingToolset/Source/ModelingComponents/Public/BaseTools/SingleSelectionMeshEditingTool.h b/Engine/Plugins/Runtime/MeshModelingToolset/Source/ModelingComponents/Public/BaseTools/SingleSelectionMeshEditingTool.h
--- a/Engine/Plugins/Runtime/MeshModelingToolset/Source/ModelingComponents/Public/BaseTools/SingleSelectionMeshEditingTool.h
+++ b/Engine/Plugins/Runtime/MeshModelingToolset/Source/ModelingComponents/Public/BaseTools/SingleSelectionMeshEditingTool.h
@@ -7,7 +7,7 @@
 #include "SingleSelectionMeshEditingTool.generated.h"
 
 class USingleSelectionMeshEditingTool;
-class UPersistentMeshSelection;
+class UDEPRECATED_PersistentMeshSelection;
 
 /**
  * USingleSelectionMeshEditingToolBuilder is a base tool builder for single
@@ -33,6 +33,7 @@ public:
 	virtual void InitializeNewTool(USingleSelectionMeshEditingTool* Tool, const FToolBuilderState& SceneState) const;
 
 	/** @return true if this Tool would like access to an available Input Selection object */
+	UE_DEPRECATED(5.2, "UPersistentMeshSelection and related functions are deprecated")
 	virtual bool WantsInputSelectionIfAvailable() const { return false; }
 
 protected:
@@ -58,6 +59,8 @@ protected:
 	UPROPERTY()
 	TWeakObjectPtr<UWorld> TargetWorld = nullptr;
 
+PRAGMA_DISABLE_DEPRECATION_WARNINGS
+
 	//
 	// Mesh Selection support
 	//
@@ -68,21 +71,24 @@ public:
 	 * to allow the Tool to configure it's behavior based on the Selection (which may or may-not exist).
 	 * If the Tool requires a Selection, this needs to be handled at the Builder level.
 	 */
-	virtual void SetInputSelection(const UPersistentMeshSelection* SelectionIn)
+	UE_DEPRECATED(5.2, "UPersistentMeshSelection and related functions are deprecated")
+	virtual void SetInputSelection(const UDEPRECATED_PersistentMeshSelection* SelectionIn)
 	{
-		InputSelection = SelectionIn;
+		InputSelection_DEPRECATED = SelectionIn;
 	}
 
 	/** @return true if an InputSelection is available */
+	UE_DEPRECATED(5.2, "UPersistentMeshSelection and related functions are deprecated")
 	virtual bool HasInputSelection() const
 	{
-		return InputSelection != nullptr;
+		return InputSelection_DEPRECATED != nullptr;
 	}
 
 	/** @return the input Selection, or nullptr if one was not configured */
-	virtual const UPersistentMeshSelection* GetInputSelection() const
+	UE_DEPRECATED(5.2, "UPersistentMeshSelection and related functions are deprecated")
+	virtual const UDEPRECATED_PersistentMeshSelection* GetInputSelection() const
 	{
-		return InputSelection;
+		return InputSelection_DEPRECATED;
 	}
 
 protected:
@@ -90,7 +96,9 @@ protected:
 	/**
 	 * (optional) Mesh Selection provided on Tool Input. This should never be modified after the Tool's Setup() is called
 	 */
+	UE_DEPRECATED(5.2, "UPersistentMeshSelection and related functions are deprecated")
 	UPROPERTY()
-	TObjectPtr<const UPersistentMeshSelection> InputSelection = nullptr;
+	TObjectPtr<const UDEPRECATED_PersistentMeshSelection> InputSelection_DEPRECATED = nullptr;
 
+PRAGMA_ENABLE_DEPRECATION_WARNINGS
 };
diff --git a/Engine/Plugins/Runtime/MeshModelingToolset/Source/ModelingComponents/Public/Selection/GeometrySelectionManager.h b/Engine/Plugins/Runtime/MeshModelingToolset/Source/ModelingComponents/Public/Selection/GeometrySelectionManager.h
--- a/Engine/Plugins/Runtime/MeshModelingToolset/Source/ModelingComponents/Public/Selection/GeometrySelectionManager.h
+++ b/Engine/Plugins/Runtime/MeshModelingToolset/Source/ModelingComponents/Public/Selection/GeometrySelectionManager.h
@@ -13,7 +13,6 @@ class IGeometrySelector;
 class UInteractiveToolsContext;
 class IToolsContextRenderAPI;
 class IToolsContextTransactionsAPI;
-class UPersistentMeshSelection;
 class UGeometrySelectionEditCommand;
 class UGeometrySelectionEditCommandArguments;
 
@@ -208,6 +207,15 @@ public:
 	virtual void EndTrackedSelectionChange();
 
 
+	/**
+	 * Directly set the current Selection for the specified Component to NewSelection.
+	 * This function allows external code to construct explicit selections, eg for a Tool or Command to emit a new Selection.
+	 * Component must already be an active target, ie set via AddActiveTarget.
+	 * If the selection of the Target would be modified, a selection-change transaction will be emitted.
+	 */
+	virtual bool SetSelectionForComponent(UPrimitiveComponent* Component, const FGeometrySelection& NewSelection);
+
+
 	DECLARE_MULTICAST_DELEGATE(FModelingSelectionInteraction_SelectionModified);
 	/**
 	 * OnSelectionModified is broadcast if the selection is modified via the above functions.
@@ -325,14 +333,6 @@ public:
 	/** Visualize the active selection using PDI drawing */
 	virtual void DebugRender(IToolsContextRenderAPI* RenderAPI);
 
-	/**
-	 * Convert the active selection to a UPersistentMeshSelection. 
-	 * SelectionManager will hold a reference to UPersistentMeshSelection to prevent it from being garbage collected,
-	 * but only until the next call to GetActiveSingleSelectionConverted_Legacy()
-	 * This function will be removed in future when UPersistentMeshSelection is no longer needed.
-	 */
-	UPersistentMeshSelection* GetActiveSingleSelectionConverted_Legacy(UPrimitiveComponent* ForComponent);
-
 
 protected:
 
@@ -429,13 +429,6 @@ protected:
 	UE::Geometry::FGeometrySelectionDelta ActiveTrackedDelta;
 	bool bSelectionModifiedDuringTrackedChange = false;
 
-
-	// legacy / to-deprecate/remove
-	
-	// this is to support GetActiveSingleSelectionConverted_Legacy, which will eventually be removed
-	UPROPERTY()
-	TObjectPtr<UPersistentMeshSelection> OldSelection;
-
 };
 
 
diff --git a/Engine/Plugins/Runtime/MeshModelingToolset/Source/ModelingComponents/Public/Selection/PersistentMeshSelection.h b/Engine/Plugins/Runtime/MeshModelingToolset/Source/ModelingComponents/Public/Selection/PersistentMeshSelection.h
--- a/Engine/Plugins/Runtime/MeshModelingToolset/Source/ModelingComponents/Public/Selection/PersistentMeshSelection.h
+++ b/Engine/Plugins/Runtime/MeshModelingToolset/Source/ModelingComponents/Public/Selection/PersistentMeshSelection.h
@@ -102,8 +102,9 @@ struct MODELINGCOMPONENTS_API FGenericMeshSelection
 /**
  * UPersistentMeshSelection is a UObject wrapper for a FGenericMeshSelection
  */
-UCLASS()
-class MODELINGCOMPONENTS_API UPersistentMeshSelection : public UObject
+//UE_DEPRECATED(5.2, "UPersistentMeshSelection and related functions are deprecated")
+UCLASS(Deprecated)
+class MODELINGCOMPONENTS_API UDEPRECATED_PersistentMeshSelection : public UObject
 {
 	GENERATED_BODY()
 
diff --git a/Engine/Plugins/Runtime/MeshModelingToolset/Source/ModelingComponents/Public/Selection/PersistentMeshSelectionManager.h b/Engine/Plugins/Runtime/MeshModelingToolset/Source/ModelingComponents/Public/Selection/PersistentMeshSelectionManager.h
--- a/Engine/Plugins/Runtime/MeshModelingToolset/Source/ModelingComponents/Public/Selection/PersistentMeshSelectionManager.h
+++ b/Engine/Plugins/Runtime/MeshModelingToolset/Source/ModelingComponents/Public/Selection/PersistentMeshSelectionManager.h
@@ -16,7 +16,7 @@ class UPreviewGeometry;
 
 
 /**
- * UPersistentMeshSelectionManager manages an active Mesh Selection.
+ * UDEPRECATED_PersistentMeshSelectionManager manages an active Mesh Selection.
  * The assumption is that this class will be registered with the ContextStore of an InteractiveToolsContext,
  * and the various Tool implementations will get/set the selection.
  * 
@@ -26,6 +26,7 @@ class UPreviewGeometry;
  * those functions are preferable to directly using the Manager.
  * 
  */
+//UE_DEPRECATED(5.1, "UPersistentMeshSelection and related functions are deprecated")
 UCLASS(Deprecated) 
 class MODELINGCOMPONENTS_API UDEPRECATED_PersistentMeshSelectionManager : public UObject
 {
@@ -40,10 +41,10 @@ public:
 	virtual bool HasActiveSelection();
 
 	/** @return the active selection */
-	virtual UPersistentMeshSelection* GetActiveSelection();
+	virtual UDEPRECATED_PersistentMeshSelection* GetActiveSelection();
 
 	/** Set a new active selection. This is undoable, ie a FPersistentMeshSelectionChange will be emitted */
-	virtual void SetNewActiveSelection(UPersistentMeshSelection* Selection);
+	virtual void SetNewActiveSelection(UDEPRECATED_PersistentMeshSelection* Selection);
 
 	/** Clear the active selection. This is undoable, ie a FPersistentMeshSelectionChange will be emitted */
 	virtual void ClearActiveSelection();
@@ -53,15 +54,16 @@ protected:
 	UPROPERTY()
 	TObjectPtr<UInteractiveToolsContext> ParentContext;
 
+	UE_DEPRECATED(5.2, "UPersistentMeshSelection and related functions are deprecated")
 	UPROPERTY()
-	TObjectPtr<UPersistentMeshSelection> ActiveSelection;
+	TObjectPtr<UDEPRECATED_PersistentMeshSelection> ActiveSelection_DEPRECATED;
 
 	UPROPERTY()
 	TObjectPtr<UPreviewGeometry> SelectionDisplay;
 
 	virtual void OnSelectionModified();
 
-	virtual void SetNewActiveSelectionInternal(UPersistentMeshSelection* Selection);
+	virtual void SetNewActiveSelectionInternal(UDEPRECATED_PersistentMeshSelection* Selection);
 	friend class FPersistentMeshSelectionChange;
 };
 
@@ -71,10 +73,10 @@ protected:
 
 /**
  * FPersistentMeshSelectionChange stores a selection change, ie before-and-after states.
- * UPersistentMeshSelectionManager emits these changes on selection modifications.
+ * UDEPRECATED_PersistentMeshSelectionManager emits these changes on selection modifications.
  */
 class MODELINGCOMPONENTS_API 
-UE_DEPRECATED(5.1, "UPersistentMeshSelectionManager has been deprecated and FPersistentMeshSelectionChange is no longer used") 
+UE_DEPRECATED(5.1, "UDEPRECATED_PersistentMeshSelectionManager has been deprecated and FPersistentMeshSelectionChange is no longer used") 
 FPersistentMeshSelectionChange : public FToolCommandChange
 {
 public:
@@ -99,28 +101,28 @@ namespace UE
 namespace Geometry
 {
 	//
-	// The functions below are helper functions that simplify usage of a UPersistentMeshSelectionManager 
+	// The functions below are helper functions that simplify usage of a UDEPRECATED_PersistentMeshSelectionManager 
 	// that is registered as a ContextStoreObject in an InteractiveToolsContext
 	//
 
 	/**
-	 * If one does not already exist, create a new instance of UPersistentMeshSelectionManager and add it to the
+	 * If one does not already exist, create a new instance of UDEPRECATED_PersistentMeshSelectionManager and add it to the
 	 * ToolsContext's ContextObjectStore
-	 * @return true if the ContextObjectStore now has a UPersistentMeshSelectionManager (whether it already existed, or was created)
+	 * @return true if the ContextObjectStore now has a UDEPRECATED_PersistentMeshSelectionManager (whether it already existed, or was created)
 	 */
 	UE_DEPRECATED(5.1, "UE::Geometry::RegisterPersistentMeshSelectionManager is deprecated")
 	MODELINGCOMPONENTS_API bool RegisterPersistentMeshSelectionManager(UInteractiveToolsContext* ToolsContext);
 
 	/**
-	 * Remove any existing UPersistentMeshSelectionManager from the ToolsContext's ContextObjectStore
-	 * @return true if the ContextObjectStore no longer has a UPersistentMeshSelectionManager (whether it was removed, or did not exist)
+	 * Remove any existing UDEPRECATED_PersistentMeshSelectionManager from the ToolsContext's ContextObjectStore
+	 * @return true if the ContextObjectStore no longer has a UDEPRECATED_PersistentMeshSelectionManager (whether it was removed, or did not exist)
 	 */
 	UE_DEPRECATED(5.1, "UE::Geometry::DeregisterPersistentMeshSelectionManager is deprecated")
 	MODELINGCOMPONENTS_API bool DeregisterPersistentMeshSelectionManager(UInteractiveToolsContext* ToolsContext);
 
 
 	/**
-	 * Find an existing UPersistentMeshSelectionManager in the ToolsContext's ContextObjectStore
+	 * Find an existing UDEPRECATED_PersistentMeshSelectionManager in the ToolsContext's ContextObjectStore
 	 * @return SelectionManager pointer or nullptr if not found
 	 */
 	UE_DEPRECATED(5.1, "UE::Geometry::FindPersistentMeshSelectionManager is deprecated")
diff --git a/Engine/Plugins/Runtime/MeshModelingToolset/Source/ModelingComponents/Public/Selection/PolygonSelectionMechanic.h b/Engine/Plugins/Runtime/MeshModelingToolset/Source/ModelingComponents/Public/Selection/PolygonSelectionMechanic.h
--- a/Engine/Plugins/Runtime/MeshModelingToolset/Source/ModelingComponents/Public/Selection/PolygonSelectionMechanic.h
+++ b/Engine/Plugins/Runtime/MeshModelingToolset/Source/ModelingComponents/Public/Selection/PolygonSelectionMechanic.h
@@ -5,6 +5,8 @@
 #include "MeshTopologySelectionMechanic.h"
 #include "PolygonSelectionMechanic.generated.h"
 
+namespace UE::Geometry { struct FGeometrySelection; }
+
 // DEPRECATED: Use UMeshTopologySelectionMechanicProperties
 UCLASS(Deprecated)
 class MODELINGCOMPONENTS_API UDEPRECATED_PolygonSelectionMechanicProperties : public UMeshTopologySelectionMechanicProperties
@@ -48,16 +50,25 @@ public:
 	virtual bool UpdateSelection(const FRay& WorldRay, FVector3d& LocalHitPositionOut, FVector3d& LocalHitNormalOut) override;
 
 	/**
-	 * Gives the current selection as a storable selection object. Can optionally apply the passed-in
-	 * compact maps to the object if the topology in the mechanic was not updated after compacting.
+	 * Convert the Active Selection to a PolyGroup-topology FGeometrySelection, with optional CompactMaps
+	 */
+	void GetSelection_AsGroupTopology(UE::Geometry::FGeometrySelection& SelectionOut, const FCompactMaps* CompactMapsToApply = nullptr) const;
+
+	/**
+	 * Convert the Active Selection to a Triangle-topology FGeometrySelection, with optional CompactMaps
+	 */
+	void GetSelection_AsTriangleTopology(UE::Geometry::FGeometrySelection& SelectionOut, const FCompactMaps* CompactMapsToApply = nullptr) const;
+
+
+	/**
+	 * Initialize the Active Selection based on the provided PolyGroup-topology FGeometrySelection
 	 */
-	void GetSelection(UPersistentMeshSelection& SelectionOut, const FCompactMaps* CompactMapsToApply = nullptr) const;
+	void SetSelection_AsGroupTopology(const UE::Geometry::FGeometrySelection& Selection);
 
 	/**
-	 * Sets the current selection using the given storable selection object. The topology in the
-	 * mechanic must already be initialized for this to work.
+	 * Initialize the Active Selection based on the provided Triangle-topology FGeometrySelection
 	 */
-	void LoadSelection(const UPersistentMeshSelection& Selection);
+	void SetSelection_AsTriangleTopology(const UE::Geometry::FGeometrySelection& Selection);
 
 
 	UE_DEPRECATED(5.2, "Use GetTopologySelector in base class")
diff --git a/Engine/Plugins/Runtime/MeshModelingToolset/Source/ModelingComponents/Public/Selection/StoredMeshSelectionUtil.h b/Engine/Plugins/Runtime/MeshModelingToolset/Source/ModelingComponents/Public/Selection/StoredMeshSelectionUtil.h
--- a/Engine/Plugins/Runtime/MeshModelingToolset/Source/ModelingComponents/Public/Selection/StoredMeshSelectionUtil.h
+++ b/Engine/Plugins/Runtime/MeshModelingToolset/Source/ModelingComponents/Public/Selection/StoredMeshSelectionUtil.h
@@ -7,7 +7,7 @@
 #include "Selections/GeometrySelection.h"
 
 struct FToolBuilderState;
-class UPersistentMeshSelection;
+class UDEPRECATED_PersistentMeshSelection;
 class UToolTarget;
 class UInteractiveTool;
 class UInteractiveToolManager;
@@ -21,44 +21,72 @@ namespace UE
 {
 namespace Geometry
 {
-
 	/**
-	 * Given a FToolBuilderState, determine if an Input Selection exists that is applicable to the ToolTarget
-	 * @return Selection if found, or nullptr if not
+	 * @return true if there is currently any geometry selection available in the SceneState
 	 */
-	MODELINGCOMPONENTS_API const UPersistentMeshSelection* GetCurrentToolInputSelection(const FToolBuilderState& SceneState, UToolTarget* Target);
+	MODELINGCOMPONENTS_API bool HaveAvailableGeometrySelection(const FToolBuilderState& SceneState);
 
 	/**
-	* Given a FToolBuilderState, find geometry selection for the given ToolTarget
-	* @param SelectionOut non-empty selection will be returned here
-	* @return true if a non-empty selection was found
-	*/
+	 * Given a FToolBuilderState, find an available FGeometrySelection for the given ToolTarget (via the UGeometrySelectionManager in the ContextObjectStore)
+	 * @param SelectionOut non-empty selection will be returned here
+	 * @return true if a non-empty selection was found
+	 */
 	MODELINGCOMPONENTS_API bool GetCurrentGeometrySelectionForTarget(const FToolBuilderState& SceneState, UToolTarget* Target,
 		FGeometrySelection& SelectionOut);
 
+	/**
+	 * Find an available FGeometrySelection for the given ToolTarget (via the UGeometrySelectionManager in the ContextObjectStore)
+	 * @param SelectionOut non-empty selection will be returned here
+	 * @return true if a non-empty selection was found
+	 */
+	MODELINGCOMPONENTS_API bool GetCurrentGeometrySelectionForTarget(UInteractiveToolManager* ToolManager, UToolTarget* Target,
+		FGeometrySelection& SelectionOut);
+
+	/**
+	 * Allow a Tool to return an "output" FGeometrySelection for the given ToolTarget (presumably on Tool Shutdown)
+	 * Target must be an "active" selection target in the current UGeometrySelectionManager.
+	 * This function will emit a selection-change transaction and should in most cases be nested inside a tool-shutdown transaction.
+	 * @return true if the selection could be set, ie Target was valid
+	 */
+	MODELINGCOMPONENTS_API bool SetToolOutputGeometrySelectionForTarget(UInteractiveTool* Tool, UToolTarget* Target, const FGeometrySelection& OutputSelection);
+
+
+
+
+	/**
+	 * Given a FToolBuilderState, determine if an Input Selection exists that is applicable to the ToolTarget
+	 * @return Selection if found, or nullptr if not
+	 */
+	UE_DEPRECATED(5.2, "PersistentMeshSelection and related functions are deprecated. GetCurrentToolInputSelection now returns nullptr.")
+	MODELINGCOMPONENTS_API const UDEPRECATED_PersistentMeshSelection* GetCurrentToolInputSelection(const FToolBuilderState& SceneState, UToolTarget* Target);
+
 	/**
 	 * Convert the given Seletion to a list of Triangles of the specified Mesh. 
-	 * This will attempt to convert the different selection types represented by UPersistentMeshSelection
+	 * This will attempt to convert the different selection types represented by PersistentMeshSelection
 	 * (eg group topology, UV island, etc) to triangle indices in TrianglesOut
 	 * @return true if a list of triangles could be created.
 	 */
+	UE_DEPRECATED(5.2, "PersistentMeshSelection and related functions are deprecated. GetStoredSelectionAsTriangles now always returns an empty TrianglesOut list.")
 	MODELINGCOMPONENTS_API bool GetStoredSelectionAsTriangles(
-		const UPersistentMeshSelection* Selection,
+		const UDEPRECATED_PersistentMeshSelection* Selection,
 		const FDynamicMesh3& Mesh,
 		TArray<int32>& TrianglesOut );
 
+
 	/**
 	 * Set the active output selection of a Tool, which (if supported) will be persist between Tools
-	 * via a UPersistentMeshSelectionManager context object. Passing a null Selection will clear any active selection.
+	 * via a PersistentMeshSelectionManager context object. Passing a null Selection will clear any active selection.
 	 * @return true if the selection was saved
 	 */
+	UE_DEPRECATED(5.2, "PersistentMeshSelection and related functions are deprecated. SetToolOutputSelection is now ignored.")
 	MODELINGCOMPONENTS_API bool SetToolOutputSelection(
 		UInteractiveTool* Tool,
-		UPersistentMeshSelection* Selection);
+		UDEPRECATED_PersistentMeshSelection* Selection);
 
 	/**
 	 * Clear any currently-active stored Tool selection
 	 */
+	UE_DEPRECATED(5.2, "PersistentMeshSelection and related functions are deprecated. ClearActiveToolSelection no longer has any effect.")
 	MODELINGCOMPONENTS_API void ClearActiveToolSelection(UInteractiveToolManager* ToolManager);
 
 
