ModelingComponents: fix paste-o found by static analysis

diff --git a/Engine/Plugins/Runtime/MeshModelingToolset/Source/ModelingComponents/Private/Selection/DynamicMeshSelector.cpp b/Engine/Plugins/Runtime/MeshModelingToolset/Source/ModelingComponents/Private/Selection/DynamicMeshSelector.cpp
--- a/Engine/Plugins/Runtime/MeshModelingToolset/Source/ModelingComponents/Private/Selection/DynamicMeshSelector.cpp
+++ b/Engine/Plugins/Runtime/MeshModelingToolset/Source/ModelingComponents/Private/Selection/DynamicMeshSelector.cpp
@@ -464,7 +464,7 @@ void FBaseDynamicMeshSelector::UpdateSelectionViaShape(
 						bool bFullyContained = false;
 						bool bIntersects = ShapeInfo.Convex.IntersectLineSegment(A,B);
 						bFullyContained = bIntersects && ShapeInfo.Convex.IntersectPoint(A) && ShapeInfo.Convex.IntersectPoint(B);
-						if ( bFullyContained == false || bFullyContained == false)
+						if ( ! bFullyContained )
 						{
 							bAllInside = false;
 							break;
@@ -494,7 +494,7 @@ void FBaseDynamicMeshSelector::UpdateSelectionViaShape(
 						C = WorldTransform.TransformPosition(C);
 						bool bFullyContained = false;
 						bool bIntersects = ShapeInfo.Convex.IntersectTriangle(A,B,C, bFullyContained);
-						if ( bIntersects == false || bFullyContained == false)
+						if ( ! (bIntersects && bFullyContained) )
 						{
 							bAllInside = false;
 							break;
