ModelingComponents: add FDynamicMeshPolygroupDeformer, applies GroupTopologyDeformer to a MeshSelection. Use in DynamicMeshSelector when selection has Polygroup topology.

diff --git a/Engine/Plugins/Runtime/GeometryProcessing/Source/DynamicMesh/Private/Operations/GroupTopologyDeformer.cpp b/Engine/Plugins/Runtime/GeometryProcessing/Source/DynamicMesh/Private/Operations/GroupTopologyDeformer.cpp
--- a/Engine/Plugins/Runtime/GeometryProcessing/Source/DynamicMesh/Private/Operations/GroupTopologyDeformer.cpp
+++ b/Engine/Plugins/Runtime/GeometryProcessing/Source/DynamicMesh/Private/Operations/GroupTopologyDeformer.cpp
@@ -7,6 +7,8 @@
 #include "Async/ParallelFor.h"
 #include "Containers/BitArray.h"
 #include "DynamicMesh/MeshNormals.h"
+#include "Selections/GeometrySelection.h"
+#include "Selections/GeometrySelectionUtil.h"
 
 using namespace UE::Geometry;
 
@@ -120,6 +122,53 @@ void FGroupTopologyDeformer::SetActiveHandleCorners(const TArray<int>& CornerIDs
 }
 
 
+void FGroupTopologyDeformer::SetActiveHandleFromSelection(const FGeometrySelection& Selection)
+{
+	if (!ensure(Selection.TopologyType == EGeometryTopologyType::Polygroup))
+	{
+		return;
+	}
+
+	TArray<int> Elements;
+	if (Selection.ElementType == EGeometryElementType::Face)
+	{
+		for (uint64 Element : Selection.Selection)
+		{
+			int32 GroupID = FGeoSelectionID(Element).TopologyID;
+			if (Topology->FindGroupByID(GroupID) != nullptr)
+			{
+				Elements.Add(GroupID);
+			}
+		}
+		SetActiveHandleFaces(Elements);
+	}
+	else if (Selection.ElementType == EGeometryElementType::Edge)
+	{
+		for (uint64 Element : Selection.Selection)
+		{
+			int32 GroupEdgeID = FGeoSelectionID(Element).TopologyID;
+			if (GroupEdgeID > 0 && GroupEdgeID < Topology->Edges.Num())
+			{
+				Elements.Add(GroupEdgeID);
+			}
+		}
+		SetActiveHandleEdges(Elements);
+	}
+	else if (Selection.ElementType == EGeometryElementType::Vertex)
+	{
+		for (uint64 Element : Selection.Selection)
+		{
+			int32 CornerID = FGeoSelectionID(Element).TopologyID;
+			if (CornerID > 0 && CornerID < Topology->Corners.Num())
+			{
+				Elements.Add(CornerID);
+			}
+		}
+		SetActiveHandleCorners(Elements);
+	}
+
+}
+
 
 
 void FGroupTopologyDeformer::SaveInitialPositions()
diff --git a/Engine/Plugins/Runtime/GeometryProcessing/Source/DynamicMesh/Public/Operations/GroupTopologyDeformer.h b/Engine/Plugins/Runtime/GeometryProcessing/Source/DynamicMesh/Public/Operations/GroupTopologyDeformer.h
--- a/Engine/Plugins/Runtime/GeometryProcessing/Source/DynamicMesh/Public/Operations/GroupTopologyDeformer.h
+++ b/Engine/Plugins/Runtime/GeometryProcessing/Source/DynamicMesh/Public/Operations/GroupTopologyDeformer.h
@@ -12,6 +12,7 @@ namespace Geometry
 
 class FDynamicMesh3;
 class FGroupTopology;
+struct FGeometrySelection;
 
 
 /**
@@ -94,6 +95,12 @@ public:
 	 */
 	virtual void SetActiveHandleEdges(const TArray<int>& TopologyEdgeIDs);
 
+
+	/**
+	 * Set the active handle from GeometrySelection
+	 */
+	virtual void SetActiveHandleFromSelection(const FGeometrySelection& Selection);
+
 	//
 	// Solving
 	// 
diff --git a/Engine/Plugins/Runtime/MeshModelingToolset/Source/ModelingComponents/Private/Selection/DynamicMeshPolygroupTransformer.cpp b/Engine/Plugins/Runtime/MeshModelingToolset/Source/ModelingComponents/Private/Selection/DynamicMeshPolygroupTransformer.cpp
new file mode 100644
--- /dev/null
+++ b/Engine/Plugins/Runtime/MeshModelingToolset/Source/ModelingComponents/Private/Selection/DynamicMeshPolygroupTransformer.cpp
@@ -0,0 +1,145 @@
+// Copyright Epic Games, Inc. All Rights Reserved.
+
+
+#include "Selection/DynamicMeshPolygroupTransformer.h"
+#include "Selections/GeometrySelectionUtil.h"
+#include "ToolContextInterfaces.h"
+#include "DynamicMesh/DynamicMeshAttributeSet.h"
+#include "DynamicMesh/MeshIndexUtil.h"
+#include "DynamicMesh/MeshNormals.h"
+#include "Changes/MeshVertexChange.h"
+#include "GroupTopology.h"
+#include "Operations/GroupTopologyDeformer.h"
+
+using namespace UE::Geometry;
+
+#define LOCTEXT_NAMESPACE "FDynamicMeshPolygroupTransformer"
+
+void FDynamicMeshPolygroupTransformer::BeginTransform(const FGeometrySelection& Selection)
+{
+	if (!ensure(Selection.TopologyType == EGeometryTopologyType::Polygroup))
+	{
+		bFallbackToSimpleTransform = true;
+		FBasicDynamicMeshSelectionTransformer::BeginTransform(Selection);
+		return;
+	}
+
+	const FGroupTopology* UseTopology = Selector->GetGroupTopology();
+
+	LinearDeformer = MakePimpl<FGroupTopologyDeformer>();
+
+	TSet<int32> VertexIDs;
+	Selector->GetDynamicMesh()->ProcessMesh([&](const FDynamicMesh3& SourceMesh)
+	{
+		// danger! this currently holds onto SourceMesh pointer...
+		// Also this is expensive, would be nice if we could save it across multiple transforms...
+		LinearDeformer->Initialize(&SourceMesh, UseTopology);
+
+		LinearDeformer->SetActiveHandleFromSelection(Selection);
+
+		// get set of selected vertex IDs, and then vector (only converting to vector for VertexToTriangleOneRing...)
+		UE::Geometry::EnumeratePolygroupSelectionVertices(Selection, SourceMesh, UseTopology, FTransform::Identity,
+			[&](uint32 VertexID, const FVector3d& Position) { VertexIDs.Add((int32)VertexID); }
+		);
+		MeshVertices = VertexIDs.Array();
+		int32 NumVertices = MeshVertices.Num();
+		InitialPositions.SetNum(NumVertices);
+		for (int32 k = 0; k < NumVertices; ++k)
+		{
+			InitialPositions[k] = SourceMesh.GetVertex(MeshVertices[k]);
+			ROIMap.Add(MeshVertices[k], k);
+		}
+
+		//UpdatedPositions = InitialPositions;
+	});
+
+	ActiveVertexChange = MakePimpl<FMeshVertexChangeBuilder>(EMeshVertexChangeComponents::VertexPositions | EMeshVertexChangeComponents::OverlayNormals);
+
+	//UpdatePendingVertexChange(false);
+	Selector->GetDynamicMesh()->ProcessMesh([&](const FDynamicMesh3& SourceMesh)
+	{
+		ActiveVertexChange->SaveVertices(&SourceMesh, LinearDeformer->GetModifiedVertices(), true);
+		ActiveVertexChange->SaveOverlayNormals(&SourceMesh, LinearDeformer->GetModifiedOverlayNormals(), true);
+	});
+}
+
+void FDynamicMeshPolygroupTransformer::UpdateTransform( 
+	TFunctionRef<FVector3d(int32 VertexID, const FVector3d& InitialPosition, const FTransform& WorldTransform)> PositionTransformFunc  )
+{
+	if (bFallbackToSimpleTransform)
+	{
+		FBasicDynamicMeshSelectionTransformer::UpdateTransform(PositionTransformFunc);
+		return;
+	}
+
+
+	int32 N = MeshVertices.Num();
+	FTransform WorldTransform = Selector->GetLocalToWorldTransform();
+
+	//for (int32 k = 0; k < N; ++k)
+	//{
+	//	UpdatedPositions[k] = PositionTransformFunc(MeshVertices[k], InitialPositions[k], WorldTransform);
+	//}
+
+	Selector->GetDynamicMesh()->EditMesh([&](FDynamicMesh3& EditMesh)
+	{
+
+		// UEditMeshPolygonsTool::ComputeUpdate_Gizmo() calls ClearSolution if the transform was zero...
+
+		LinearDeformer->UpdateSolution(&EditMesh, [&](FDynamicMesh3* TargetMesh, int VertIdx)
+		{
+			if ( int32* ArrayIndex = ROIMap.Find(VertIdx) )
+			{
+				return PositionTransformFunc( VertIdx, InitialPositions[*ArrayIndex], WorldTransform );
+			}
+			else
+			{
+				return TargetMesh->GetVertex(VertIdx);
+			}
+		});
+
+
+		// TODO: the right thing to do here would be to only recompute the modified overlay normals...but we have no utility for that??
+		//FMeshNormals::QuickRecomputeOverlayNormals(EditMesh);
+
+	}, EDynamicMeshChangeType::DeformationEdit,
+		//EDynamicMeshAttributeChangeFlags::VertexPositions, false);
+	   EDynamicMeshAttributeChangeFlags::VertexPositions | EDynamicMeshAttributeChangeFlags::NormalsTangents, false);
+
+
+	// why does simple deformer do this every time??
+	//UpdatePendingVertexChange(false);
+}
+
+
+void FDynamicMeshPolygroupTransformer::EndTransform(IToolsContextTransactionsAPI* TransactionsAPI)
+{
+	if (bFallbackToSimpleTransform)
+	{
+		FBasicDynamicMeshSelectionTransformer::EndTransform(TransactionsAPI);
+		return;
+	}
+
+	//UpdatePendingVertexChange(true);
+	Selector->GetDynamicMesh()->ProcessMesh([&](const FDynamicMesh3& SourceMesh)
+	{
+		ActiveVertexChange->SaveVertices(&SourceMesh, LinearDeformer->GetModifiedVertices(), false);
+		ActiveVertexChange->SaveOverlayNormals(&SourceMesh, LinearDeformer->GetModifiedOverlayNormals(), false);
+	});
+
+	if (TransactionsAPI != nullptr)
+	{
+		TUniquePtr<FToolCommandChange> Result = MoveTemp(ActiveVertexChange->Change);
+		TransactionsAPI->AppendChange(Selector->GetDynamicMesh(), MoveTemp(Result), LOCTEXT("DynamicMeshTransformChange", "Transform"));
+	}
+
+	ActiveVertexChange.Reset();
+
+	if (OnEndTransformFunc)
+	{
+		OnEndTransformFunc(TransactionsAPI);
+	}
+}
+
+
+#undef LOCTEXT_NAMESPACE 
\ No newline at end of file
diff --git a/Engine/Plugins/Runtime/MeshModelingToolset/Source/ModelingComponents/Private/Selection/DynamicMeshSelector.cpp b/Engine/Plugins/Runtime/MeshModelingToolset/Source/ModelingComponents/Private/Selection/DynamicMeshSelector.cpp
--- a/Engine/Plugins/Runtime/MeshModelingToolset/Source/ModelingComponents/Private/Selection/DynamicMeshSelector.cpp
+++ b/Engine/Plugins/Runtime/MeshModelingToolset/Source/ModelingComponents/Private/Selection/DynamicMeshSelector.cpp
@@ -12,6 +12,7 @@
 #include "GroupTopology.h"
 #include "Spatial/SegmentTree3.h"
 #include "ToolSceneQueriesUtil.h"
+#include "Selection/DynamicMeshPolygroupTransformer.h"
 
 using namespace UE::Geometry;
 
@@ -733,7 +734,16 @@ IGeometrySelectionTransformer* FDynamicMeshSelector::InitializeTransformation(co
 		ParentIdentifier.GetAsComponentType<UDynamicMeshComponent>()->SetTransientDeferCollisionUpdates(true);
 	}
 
-	ActiveTransformer = MakePimpl<FBasicDynamicMeshSelectionTransformer>();
+	if (Selection.TopologyType == EGeometryTopologyType::Polygroup)
+	{
+		ActiveTransformer = MakeShared<FDynamicMeshPolygroupTransformer>();
+	}
+	else
+	{
+		ActiveTransformer = MakeShared<FBasicDynamicMeshSelectionTransformer>();
+	}
+
+
 	ActiveTransformer->Initialize(this);
 	return ActiveTransformer.Get();
 }
diff --git a/Engine/Plugins/Runtime/MeshModelingToolset/Source/ModelingComponents/Public/Selection/DynamicMeshPolygroupTransformer.h b/Engine/Plugins/Runtime/MeshModelingToolset/Source/ModelingComponents/Public/Selection/DynamicMeshPolygroupTransformer.h
new file mode 100644
--- /dev/null
+++ b/Engine/Plugins/Runtime/MeshModelingToolset/Source/ModelingComponents/Public/Selection/DynamicMeshPolygroupTransformer.h
@@ -0,0 +1,30 @@
+// Copyright Epic Games, Inc. All Rights Reserved.
+
+#pragma once
+
+#include "CoreMinimal.h"
+#include "Templates/PimplPtr.h"
+#include "Selections/GeometrySelection.h"
+#include "Selection/GeometrySelector.h"
+#include "Selection/DynamicMeshSelector.h"
+
+namespace UE::Geometry { class FGroupTopologyDeformer; }
+
+/**
+* IGeometrySelectionTransformer for Dynamic Mesh Polygroups, that does various types of soft deformation.
+*/
+class MODELINGCOMPONENTS_API FDynamicMeshPolygroupTransformer : public FBasicDynamicMeshSelectionTransformer
+{
+public:
+
+	virtual void BeginTransform(const UE::Geometry::FGeometrySelection& Selection) override;
+	virtual void UpdateTransform( TFunctionRef<FVector3d(int32 VertexID, const FVector3d& InitialPosition, const FTransform& WorldTransform)> PositionTransformFunc ) override;
+	virtual void EndTransform(IToolsContextTransactionsAPI* TransactionsAPI) override;
+
+protected:
+	bool bFallbackToSimpleTransform = false;
+
+	TMap<int32, int32> ROIMap;
+	TPimplPtr<UE::Geometry::FGroupTopologyDeformer> LinearDeformer;
+};
+
diff --git a/Engine/Plugins/Runtime/MeshModelingToolset/Source/ModelingComponents/Public/Selection/DynamicMeshSelector.h b/Engine/Plugins/Runtime/MeshModelingToolset/Source/ModelingComponents/Public/Selection/DynamicMeshSelector.h
--- a/Engine/Plugins/Runtime/MeshModelingToolset/Source/ModelingComponents/Public/Selection/DynamicMeshSelector.h
+++ b/Engine/Plugins/Runtime/MeshModelingToolset/Source/ModelingComponents/Public/Selection/DynamicMeshSelector.h
@@ -186,7 +186,7 @@ public:
 	virtual void ShutdownTransformation(IGeometrySelectionTransformer* Transformer) override;
 
 protected:
-	TPimplPtr<FBasicDynamicMeshSelectionTransformer> ActiveTransformer;
+	TSharedPtr<FBasicDynamicMeshSelectionTransformer> ActiveTransformer;
 
 	// give Transformer access to internals 
 	friend class FBasicDynamicMeshSelectionTransformer;
