MeshModelingTools: Add gizmo numerical UI.

- Added FTransformGizmoDataBinder which allows for gizmos to be bound to vectors such that the vectors are updated when the gizmo changes, and the gizmo is updated from vectors when asked. This should make it easy-ish to create displays of gizmos.
- Used the data binder in STransformGizmoNumericalUIOverlay, an overlay containing a draggable numerical UI panel for gizmos. Note that currently all changes, including typing, intentionally go down a gizmo drag path, meaning that they trigger the begin/end transform edit sequence calls, to make the tools treat them as a gizmo drag (in part because some tools might otherwise assume that a gizmo update is from an undo transaction or some other path that is special-cased).
- Added some slate things to ModelingEditorUI to support the numerical UI, in particular a viewport draggable box, and a widget that emits a tick delegate.
- Added some delegates to gizmo context object and combined transform gizmos to allow things to bind/unbind to/from them appropriately.

diff --git a/Engine/Plugins/Editor/ModelingToolsEditorMode/Source/ModelingToolsEditorMode/Private/ModelingModeToolkit_Toolbars.cpp b/Engine/Plugins/Editor/ModelingToolsEditorMode/Source/ModelingToolsEditorMode/Private/ModelingModeToolkit_Toolbars.cpp
--- a/Engine/Plugins/Editor/ModelingToolsEditorMode/Source/ModelingToolsEditorMode/Private/ModelingModeToolkit_Toolbars.cpp
+++ b/Engine/Plugins/Editor/ModelingToolsEditorMode/Source/ModelingToolsEditorMode/Private/ModelingModeToolkit_Toolbars.cpp
@@ -17,6 +17,7 @@
 #include "ISettingsModule.h"
 
 #include "Framework/MultiBox/MultiBoxBuilder.h"
+#include "STransformGizmoNumericalUIOverlay.h"
 #include "Widgets/Images/SImage.h"
 #include "Widgets/Layout/SBorder.h"
 #include "Widgets/Layout/SSpacer.h"
@@ -107,6 +108,17 @@ void MakeSubMenu_GizmoVisibilityMode(FModelingToolsEditorModeToolkit* Toolkit, F
 	MenuBuilder.AddMenuEntry(LOCTEXT("GizmoMode_AbsoluteWorldSnap", "World Grid Snapping"), 
 		LOCTEXT("GizmoMode_AbsoluteWorldSnap_Tooltip", "Snap Translation/Rotation to Absolute Grid Coordinates in the World Coordinate System, instead of Relative to the initial position"),
 		FSlateIcon(), GizmoMode_AbsoluteWorldSnap, NAME_None, EUserInterfaceActionType::ToggleButton);
+	
+
+	MenuBuilder.AddSubMenu(
+		LOCTEXT("NumericalUISubMenu", "Numerical UI"), LOCTEXT("NumericalUISubMenu_ToolTip", "Configure the gizmo numerical UI."),
+		FNewMenuDelegate::CreateLambda([Toolkit](FMenuBuilder& SubMenuBuilder) {
+			TSharedPtr<STransformGizmoNumericalUIOverlay> NumericalUI = Toolkit->GetGizmoNumericalUIOverlayWidget();
+			if (ensure(NumericalUI.IsValid()))
+			{
+				NumericalUI->MakeNumericalUISubMenu(SubMenuBuilder);
+			}
+		}));
 }
 
 
diff --git a/Engine/Plugins/Editor/ModelingToolsEditorMode/Source/ModelingToolsEditorMode/Private/ModelingToolsEditorModeToolkit.cpp b/Engine/Plugins/Editor/ModelingToolsEditorMode/Source/ModelingToolsEditorMode/Private/ModelingToolsEditorModeToolkit.cpp
--- a/Engine/Plugins/Editor/ModelingToolsEditorMode/Source/ModelingToolsEditorMode/Private/ModelingToolsEditorModeToolkit.cpp
+++ b/Engine/Plugins/Editor/ModelingToolsEditorMode/Source/ModelingToolsEditorMode/Private/ModelingToolsEditorModeToolkit.cpp
@@ -16,6 +16,7 @@
 #include "Toolkits/AssetEditorModeUILayer.h"
 
 #include "SSimpleButton.h"
+#include "STransformGizmoNumericalUIOverlay.h"
 #include "Tools/UEdMode.h"
 #include "Widgets/Layout/SExpandableArea.h"
 #include "Widgets/Input/SEditableComboBox.h"
@@ -101,6 +102,12 @@ FModelingToolsEditorModeToolkit::~FModelingToolsEditorModeToolkit()
 			GetToolkitHost()->RemoveViewportOverlayWidget(SelectionPaletteOverlayWidget.ToSharedRef());
 			SelectionPaletteOverlayWidget.Reset();
 		}
+
+		if (GizmoNumericalUIOverlayWidget.IsValid())
+		{
+			GetToolkitHost()->RemoveViewportOverlayWidget(GizmoNumericalUIOverlayWidget.ToSharedRef());
+			GizmoNumericalUIOverlayWidget.Reset();
+		}
 	}
 
 	UModelingToolsEditorModeSettings* Settings = GetMutableDefault<UModelingToolsEditorModeSettings>();
@@ -201,6 +208,11 @@ void FModelingToolsEditorModeToolkit::Init(const TSharedPtr<IToolkitHost>& InitT
 
 	MakeToolShutdownOverlayWidget();
 
+	// Note that the numerical UI widget should be created before making the selection palette so that
+	// it can be bound to the buttons there.
+	MakeGizmoNumericalUIOverlayWidget();
+	GetToolkitHost()->AddViewportOverlayWidget(GizmoNumericalUIOverlayWidget.ToSharedRef());
+
 	const UModelingToolsEditorModeSettings* ModelingModeSettings = GetDefault<UModelingToolsEditorModeSettings>();
 	bool bEnableSelectionUI = ModelingModeSettings && ModelingModeSettings->bEnableMeshSelections;
 	if ( bEnableSelectionUI )
@@ -287,6 +299,15 @@ void FModelingToolsEditorModeToolkit::MakeToolShutdownOverlayWidget()
 }
 
 
+void FModelingToolsEditorModeToolkit::MakeGizmoNumericalUIOverlayWidget()
+{
+	GizmoNumericalUIOverlayWidget = SNew(STransformGizmoNumericalUIOverlay)
+		.DefaultLeftPadding(15)
+		// Position above the little axis visualization
+		.DefaultVerticalPadding(75)
+		.bPositionRelativeToBottom(true);
+}
+
 
 TSharedPtr<SWidget> FModelingToolsEditorModeToolkit::MakeAssetConfigPanel()
 {
@@ -777,6 +798,13 @@ void FModelingToolsEditorModeToolkit::InitializeAfterModeSetup()
 		// ToolTarget Factories that are only available once ModelingToolsEditorMode has been initialized
 		AssetLODMode->SetSelectedItem(AssetLODModes[0]);
 
+		// Attach the gizmo numerical UI to the gizmo context object so that it can show the values for any newly
+		// created gizmos.
+		if (ensure(GizmoNumericalUIOverlayWidget.IsValid()))
+		{
+			GizmoNumericalUIOverlayWidget->BindToGizmoContextObject(GetScriptableEditorMode()->GetInteractiveToolsContext(EToolsContextScope::EdMode));
+		}
+
 		bFirstInitializeAfterModeSetup = false;
 	}
 
@@ -1500,10 +1528,20 @@ void FModelingToolsEditorModeToolkit::OnActiveViewportChanged(TSharedPtr<IAssetV
 		if (OldViewport)	
 		{
 			GetToolkitHost()->RemoveViewportOverlayWidget(ToolShutdownViewportOverlayWidget.ToSharedRef(), OldViewport);
+
+			if (GizmoNumericalUIOverlayWidget.IsValid())
+			{
+				GetToolkitHost()->RemoveViewportOverlayWidget(GizmoNumericalUIOverlayWidget.ToSharedRef(), OldViewport);
+			}
 		}
 
 		// Add the hud to the new viewport
 		GetToolkitHost()->AddViewportOverlayWidget(ToolShutdownViewportOverlayWidget.ToSharedRef(), NewViewport);
+
+		if (GizmoNumericalUIOverlayWidget.IsValid())
+		{
+			GetToolkitHost()->AddViewportOverlayWidget(GizmoNumericalUIOverlayWidget.ToSharedRef(), NewViewport);
+		}
 	}
 }
 
diff --git a/Engine/Plugins/Editor/ModelingToolsEditorMode/Source/ModelingToolsEditorMode/Public/ModelingToolsEditorModeToolkit.h b/Engine/Plugins/Editor/ModelingToolsEditorMode/Source/ModelingToolsEditorMode/Public/ModelingToolsEditorModeToolkit.h
--- a/Engine/Plugins/Editor/ModelingToolsEditorMode/Source/ModelingToolsEditorMode/Public/ModelingToolsEditorModeToolkit.h
+++ b/Engine/Plugins/Editor/ModelingToolsEditorMode/Source/ModelingToolsEditorMode/Public/ModelingToolsEditorModeToolkit.h
@@ -9,7 +9,7 @@
 #include "Widgets/Input/SEditableTextBox.h"
 #include "StatusBarSubsystem.h"
 
-
+class STransformGizmoNumericalUIOverlay;
 class IDetailsView;
 class SButton;
 class STextBlock;
@@ -68,6 +68,9 @@ public:
 
 	virtual void ForceToolPaletteRebuild();
 
+	// This is exposed only for the convenience of being able to create the numerical UI submenu
+	// in a non-member function in ModelingModeToolkit_Toolbars.cpp
+	TSharedPtr<STransformGizmoNumericalUIOverlay> GetGizmoNumericalUIOverlayWidget() { return GizmoNumericalUIOverlayWidget; }
 private:
 	const static TArray<FName> PaletteNames_Standard;
 
@@ -86,6 +89,8 @@ private:
 	TSharedPtr<SWidget> SelectionPaletteOverlayWidget;
 	void MakeSelectionPaletteOverlayWidget();
 
+	void MakeGizmoNumericalUIOverlayWidget();
+	TSharedPtr<STransformGizmoNumericalUIOverlay> GizmoNumericalUIOverlayWidget;
 
 	TSharedPtr<STextBlock> ModeWarningArea;
 	TSharedPtr<STextBlock> ModeHeaderArea;
diff --git a/Engine/Plugins/Experimental/MeshModelingToolsetExp/Source/ModelingEditorUI/ModelingEditorUI.build.cs b/Engine/Plugins/Experimental/MeshModelingToolsetExp/Source/ModelingEditorUI/ModelingEditorUI.build.cs
--- a/Engine/Plugins/Experimental/MeshModelingToolsetExp/Source/ModelingEditorUI/ModelingEditorUI.build.cs
+++ b/Engine/Plugins/Experimental/MeshModelingToolsetExp/Source/ModelingEditorUI/ModelingEditorUI.build.cs
@@ -31,7 +31,9 @@ public class ModelingEditorUI : ModuleRules
 				"UnrealEd",
                 "ContentBrowser",
 				"ContentBrowserData",
+				"InteractiveToolsFramework",
                 "LevelEditor",
+				"ModelingComponents", // FTransformGizmoDataBinder
 				"StatusBar",
                 "Projects",
 				"ToolWidgets",
diff --git a/Engine/Plugins/Experimental/MeshModelingToolsetExp/Source/ModelingEditorUI/Private/ModelingWidgets/SDraggableBox.cpp b/Engine/Plugins/Experimental/MeshModelingToolsetExp/Source/ModelingEditorUI/Private/ModelingWidgets/SDraggableBox.cpp
new file mode 100644
--- /dev/null
+++ b/Engine/Plugins/Experimental/MeshModelingToolsetExp/Source/ModelingEditorUI/Private/ModelingWidgets/SDraggableBox.cpp
@@ -0,0 +1,176 @@
+// Copyright Epic Games, Inc. All Rights Reserved.
+
+#include "ModelingWidgets/SDraggableBox.h"
+
+#include "Viewports/InViewportUIDragOperation.h"
+#include "Widgets/Layout/SBackgroundBlur.h"
+
+
+void SDraggableBoxOverlay::Construct(const FArguments& InArgs)
+{
+	// Depending on whether we're positioning relative to the bottom or top, we'll update the margin of the
+	// draggable box differently.
+	TAttribute<FMargin> DraggableBoxPaddingAttribute;
+	SDraggableBox::FOnDragComplete OnDragCompleteFunction;
+
+	if (InArgs._bPositionRelativeToBottom)
+	{
+		DraggableBoxPaddingAttribute = MakeAttributeLambda([this]()
+		{
+			return FMargin(DraggableBoxPaddingHorizontal, 0, 0, DraggableBoxPaddingVertical);
+		});
+		OnDragCompleteFunction = SDraggableBox::FOnDragComplete::CreateLambda([this](const FVector2D& ScreenSpacePosition)
+		{
+			FVector2D TopLeftPosition = ContainingBox->GetTickSpaceGeometry().AbsoluteToLocal(ScreenSpacePosition);
+
+			DraggableBoxPaddingHorizontal = TopLeftPosition.X;
+
+			// When aligning to bottom, we need the distance between the draggable box bottom and our bottom.
+			// I.e. OurHeight - TopLeftY - BoxHeight
+			DraggableBoxPaddingVertical = ContainingBox->GetTickSpaceGeometry().GetAbsoluteSize().Y - TopLeftPosition.Y - DraggableBox->GetTickSpaceGeometry().GetAbsoluteSize().Y;
+		});
+	}
+	else
+	{
+		DraggableBoxPaddingAttribute = MakeAttributeLambda([this]()
+		{
+			return FMargin(DraggableBoxPaddingHorizontal, DraggableBoxPaddingVertical, 0, 0);
+		});
+		OnDragCompleteFunction = SDraggableBox::FOnDragComplete::CreateLambda([this](const FVector2D& ScreenSpacePosition)
+		{
+			FVector2D Position = ContainingBox->GetTickSpaceGeometry().AbsoluteToLocal(ScreenSpacePosition);
+			DraggableBoxPaddingHorizontal = Position.X;
+			DraggableBoxPaddingVertical = Position.Y;
+		});
+	}
+
+	ChildSlot
+	[
+		SAssignNew(ContainingBox, SVerticalBox)
+		.Visibility(EVisibility::SelfHitTestInvisible)
+
+		+ SVerticalBox::Slot()
+
+		.HAlign(HAlign_Left)
+		.VAlign(InArgs._bPositionRelativeToBottom ? VAlign_Bottom : VAlign_Top)
+		.Padding(DraggableBoxPaddingAttribute)
+		[
+			// The draggable box itself
+			SAssignNew(DraggableBox, SDraggableBox)
+			.OnDragComplete(OnDragCompleteFunction)
+			[
+				InArgs._Content.Widget
+			]
+		]
+	]
+	;
+}
+
+void SDraggableBoxOverlay::SetBoxPosition(float HorizontalPosition, float VerticalPosition)
+{
+	DraggableBoxPaddingHorizontal = HorizontalPosition;
+	DraggableBoxPaddingVertical = VerticalPosition;
+}
+
+// Draggable box:
+
+void SDraggableBox::Construct(const FArguments& InArgs)
+{
+	OnDragComplete = InArgs._OnDragComplete;
+
+	ChildSlot
+	[
+		SAssignNew(InnerWidget, SBackgroundBlur)
+		.BlurStrength(2)
+		.Padding(0.0f)
+		.CornerRadius(FVector4(4.0f, 4.0f, 4.0f, 4.0f))
+		[
+			SNew(SBorder)
+			.Padding(0.0f)
+			[
+				InArgs._Content.Widget
+			]
+		]
+	];
+}
+
+FReply SDraggableBox::OnMouseButtonDown(const FGeometry& MyGeometry, const FPointerEvent& MouseEvent)
+{
+	return FReply::Handled().DetectDrag(SharedThis(this), EKeys::LeftMouseButton);
+};
+
+FReply SDraggableBox::OnDragDetected(const FGeometry& MyGeometry, const FPointerEvent& MouseEvent)
+{
+	// Need to remember where within the box we grabbed
+	const FVector2D ScreenSpaceOffsetOfGrab = MouseEvent.GetScreenSpacePosition() - MyGeometry.GetAbsolutePosition();
+
+	// The drag/drop operation moves a transparent copy of our contents, so we should hide ourselves to make it
+	// look like the contents are moving.
+	SetVisibility(EVisibility::Collapsed);
+
+	// This will be called when the drag/drop operation is done
+	FOnDragComplete OnUIDropped = FOnDragComplete::CreateLambda([this] (const FVector2D& ScreenSpacePosition) {
+		OnDragComplete.ExecuteIfBound(ScreenSpacePosition);
+		SetVisibility(EVisibility::SelfHitTestInvisible);
+	});
+
+	TSharedRef<FDraggableBoxUIDragOperation> DragDropOperation =
+		FDraggableBoxUIDragOperation::New(
+			// Note that it is important to give the inner widget here, since SDraggableBox itself is invisible now
+			InnerWidget.ToSharedRef(),
+			ScreenSpaceOffsetOfGrab,
+			GetDesiredSize(),
+			OnUIDropped
+		);
+
+	return FReply::Handled().BeginDragDrop(DragDropOperation);
+}
+
+// Drag operation:
+
+void FDraggableBoxUIDragOperation::OnDrop(bool bDropWasHandled, const FPointerEvent& MouseEvent)
+{
+	FVector2D ScreenSpaceDropLocation = MouseEvent.GetScreenSpacePosition() - DecoratorOffsetFromCursor;
+
+	OnDragComplete.ExecuteIfBound(ScreenSpaceDropLocation);
+
+	// Call super
+	FDragDropOperation::OnDrop(bDropWasHandled, MouseEvent);
+
+	UIBeingDragged.Reset();
+}
+
+// These are mostly copied from FInViewportUIDragOperation:
+void FDraggableBoxUIDragOperation::OnDragged(const FDragDropEvent& DragDropEvent)
+{
+	// The tab is being dragged. Move the the decorator window to match the cursor position.
+	FVector2D TargetPosition = DragDropEvent.GetScreenSpacePosition() - DecoratorOffsetFromCursor;
+	CursorDecoratorWindow->UpdateMorphTargetShape(FSlateRect(TargetPosition.X, TargetPosition.Y, TargetPosition.X + LastContentSize.X, TargetPosition.Y + LastContentSize.Y));
+	CursorDecoratorWindow->MoveWindowTo(TargetPosition);
+}
+
+TSharedRef<FDraggableBoxUIDragOperation> FDraggableBoxUIDragOperation::New(const TSharedRef<SWidget>& InUIToBeDragged, 
+	const FVector2D InDecoratorOffset, const FVector2D& OwnerAreaSize, SDraggableBox::FOnDragComplete& OnDragComplete)
+{
+	const TSharedRef<FDraggableBoxUIDragOperation> Operation = MakeShareable(new FDraggableBoxUIDragOperation(InUIToBeDragged, InDecoratorOffset, OwnerAreaSize, OnDragComplete));
+	return Operation;
+}
+
+FDraggableBoxUIDragOperation::FDraggableBoxUIDragOperation(const TSharedRef<class SWidget>& InUIToBeDragged, 
+	const FVector2D InDecoratorOffset, const FVector2D& OwnerAreaSize, SDraggableBox::FOnDragComplete& InOnDragComplete)
+	: UIBeingDragged(InUIToBeDragged)
+	, DecoratorOffsetFromCursor(InDecoratorOffset)
+	, LastContentSize(OwnerAreaSize)
+	, OnDragComplete(InOnDragComplete)
+{
+	// Create the decorator window that we will use during this drag and drop to make the user feel like
+	// they are actually dragging a piece of UI.
+	const bool bShowImmediately = true;
+	CursorDecoratorWindow = FSlateApplication::Get().AddWindow(SWindow::MakeStyledCursorDecorator(FAppStyle::Get().GetWidgetStyle<FWindowStyle>("InViewportDecoratorWindow")), bShowImmediately);
+	CursorDecoratorWindow->SetOpacity(0.45f);
+	CursorDecoratorWindow->SetContent
+	(
+		InUIToBeDragged
+	);
+
+}
\ No newline at end of file
diff --git a/Engine/Plugins/Experimental/MeshModelingToolsetExp/Source/ModelingEditorUI/Private/ModelingWidgets/STickDelegateWidget.cpp b/Engine/Plugins/Experimental/MeshModelingToolsetExp/Source/ModelingEditorUI/Private/ModelingWidgets/STickDelegateWidget.cpp
new file mode 100644
--- /dev/null
+++ b/Engine/Plugins/Experimental/MeshModelingToolsetExp/Source/ModelingEditorUI/Private/ModelingWidgets/STickDelegateWidget.cpp
@@ -0,0 +1,20 @@
+// Copyright Epic Games, Inc. All Rights Reserved.
+
+#include "ModelingWidgets/STickDelegateWidget.h"
+
+void STickDelegateWidget::Construct(const FArguments& InArgs)
+{
+	OnTickDelegate = InArgs._OnTickDelegate;
+
+	ChildSlot
+	[
+		InArgs._Content.Widget
+	];
+}
+
+void STickDelegateWidget::Tick(const FGeometry& AllottedGeometry, const double InCurrentTime, const float InDeltaTime)
+{
+	OnTickDelegate.ExecuteIfBound(AllottedGeometry, InCurrentTime, InDeltaTime);
+
+	SCompoundWidget::Tick(AllottedGeometry, InCurrentTime, InDeltaTime);
+}
\ No newline at end of file
diff --git a/Engine/Plugins/Experimental/MeshModelingToolsetExp/Source/ModelingEditorUI/Private/STransformGizmoNumericalUIOverlay.cpp b/Engine/Plugins/Experimental/MeshModelingToolsetExp/Source/ModelingEditorUI/Private/STransformGizmoNumericalUIOverlay.cpp
new file mode 100644
--- /dev/null
+++ b/Engine/Plugins/Experimental/MeshModelingToolsetExp/Source/ModelingEditorUI/Private/STransformGizmoNumericalUIOverlay.cpp
@@ -0,0 +1,533 @@
+// Copyright Epic Games, Inc. All Rights Reserved.
+
+#include "STransformGizmoNumericalUIOverlay.h"
+
+#include "BaseGizmos/CombinedTransformGizmo.h"
+#include "BaseGizmos/TransformGizmoUtil.h" // UCombinedTransformGizmoContextObject
+#include "ContextObjectStore.h"
+#include "InteractiveToolsContext.h"
+#include "ModelingWidgets/SDraggableBox.h" // SDraggableBoxOverlay
+#include "ModelingWidgets/STickDelegateWidget.h"
+#include "Transforms/TransformGizmoDataBinder.h"
+#include "Widgets/SBoxPanel.h" // SHorizontalBox
+#include "Widgets/Images/SImage.h"
+
+#define LOCTEXT_NAMESPACE "STransformGizmoNumericalUIOverlay"
+
+namespace STransformGizmoNumericalUIOverlayLocals
+{
+	// Used for placing the icon in the translate/rotate/scale rows
+	float ICON_RIGHT_PADDING = 8;
+}
+
+void STransformGizmoNumericalUIOverlay::Construct(const FArguments& InArgs)
+{
+	bPositionRelativeToBottom = InArgs._bPositionRelativeToBottom;
+	DefaultLeftPadding = InArgs._DefaultLeftPadding;
+	DefaultVerticalPadding = InArgs._DefaultVerticalPadding;
+	TranslationScrubSensitivity = InArgs._TranslationScrubSensitivity;
+
+	// Set up the binding between our displayed values and the gizmo
+	DataBinder = MakeShared<FTransformGizmoDataBinder>();
+	DataBinder->InitializeBoundVectors(&DisplayTranslation, &DisplayEulerAngles, &DisplayScale);
+	DataBinder->SetDefaultLocalReferenceTransform(InArgs._DefaultLocalReferenceTransform);
+	DataBinder->SetTranslationConversionFunctions(InArgs._InternalToDisplayFunction, InArgs._DisplayToInternalFunction);
+	
+	// Rebuild our display whenever we start tracking a new gizmo
+	DataBinder->OnTrackedGizmoChanged.AddLambda([this](UCombinedTransformGizmo* Gizmo) {
+		if (Gizmo)
+		{
+			WidgetContents->SetContent(CreateWidgetForGizmo(Gizmo->GetGizmoElements()));
+		}
+	});
+
+	ChildSlot
+	[
+		SAssignNew(DraggableBoxOverlay, SDraggableBoxOverlay)
+		.bPositionRelativeToBottom(bPositionRelativeToBottom)
+		.Visibility_Lambda([this]()
+		{
+			return bIsEnabled && DataBinder->HasVisibleGizmo() ? EVisibility::SelfHitTestInvisible : EVisibility::Collapsed;
+		})
+		[
+			SNew(STickDelegateWidget)
+			.OnTickDelegate(STickDelegateWidget::FOnTick::CreateSP(this, &STransformGizmoNumericalUIOverlay::OnTickWidget))
+			[
+				SAssignNew(WidgetContents, SBorder)
+				.BorderImage(FAppStyle::Get().GetBrush("EditorViewport.OverlayBrush"))
+				.Padding(8.f)
+			]
+		]
+	]
+	;
+
+	// Don't block things under the overlay.
+	SetVisibility(EVisibility::SelfHitTestInvisible);
+
+	ResetPositionInViewport();
+}
+
+STransformGizmoNumericalUIOverlay::~STransformGizmoNumericalUIOverlay()
+{
+	Reset();
+	DataBinder = nullptr;
+}
+
+void STransformGizmoNumericalUIOverlay::Reset()
+{
+	if (DataBinder.IsValid())
+	{
+		DataBinder->Reset();
+	}
+}
+
+bool STransformGizmoNumericalUIOverlay::BindToGizmoContextObject(UInteractiveToolsContext* ToolsContext)
+{
+	if (ToolsContext && ToolsContext->ContextObjectStore)
+	{
+		UCombinedTransformGizmoContextObject* ContextObject = ToolsContext->ContextObjectStore->FindContext<UCombinedTransformGizmoContextObject>();
+		if (ContextObject)
+		{
+			DataBinder->BindToGizmoContextObject(ContextObject);
+			return true;
+		}
+	}
+	return ensure(false);
+}
+
+void STransformGizmoNumericalUIOverlay::OnTickWidget(const FGeometry& AllottedGeometry, const double InCurrentTime, const float InDeltaTime)
+{
+	// We update the coordinate system each tick since context queries api doesn't tell us when the value changes.
+	// In the editor, we could actually bind to a notification callback instead of ticking here.
+	DataBinder->UpdateCoordinateSystem();
+}
+
+// Helper to reuse code when creating translation/rotation/scale UI pieces. We originally built these individually
+// so that we could individually pick and choose components based on active gizmo components. However, we have since
+// decided to disable unused components rather than removing them entirely. Otherwise, maybe we could have used
+// SVectorInputBox. There didn't seem to be enough generality here to justify making a utility function, but perhaps
+// we will change our mind about that.
+SNumericEntryBox<double>::FArguments STransformGizmoNumericalUIOverlay::CreateComponentBaseArgs(FVector3d& TargetVector, int ComponentIndex, TAttribute<bool> IsEnabled)
+{
+	// Can be indexed via ComponentIndex:
+	const FText Tooltips[3] {
+		LOCTEXT("X_ToolTip", "X: {0}"),
+		LOCTEXT("Y_ToolTip", "Y: {0}"),
+		LOCTEXT("Z_ToolTip", "Z: {0}")
+	};
+	const FLinearColor Colors[3] {
+		SNumericEntryBox<double>::RedLabelBackgroundColor,
+		SNumericEntryBox<double>::GreenLabelBackgroundColor,
+		SNumericEntryBox<double>::BlueLabelBackgroundColor
+	};
+
+	SNumericEntryBox<double>::FArguments Args;
+
+	TWeakPtr<FTransformGizmoDataBinder> DataBinderWeakPtr = DataBinder;
+
+	Args.AllowSpin(true)
+		.EditableTextBoxStyle(&FCoreStyle::Get().GetWidgetStyle<FEditableTextBoxStyle>("NormalEditableTextBox"))
+		.Value(MakeAttributeLambda([&TargetVector, ComponentIndex]()->TOptional<double> { return TargetVector[ComponentIndex]; }))
+		.OnValueChanged(SNumericEntryBox<double>::FOnValueChanged::CreateLambda([DataBinderWeakPtr, &TargetVector, ComponentIndex](double Value) {
+			TargetVector[ComponentIndex] = Value;
+			if (DataBinderWeakPtr.IsValid())
+			{
+				DataBinderWeakPtr.Pin()->UpdateAfterDataEdit();
+			}
+		}))
+		// This callback doesn't seem to be that useful because we get OnValueChanged callbacks regardless of the input method.
+		//.OnValueCommitted(SNumericEntryBox<double>::FOnValueCommitted::CreateLambda(this, [](double Value, ETextCommit::Type TextCommitType) {}))
+		
+		.OnBeginSliderMovement(FSimpleDelegate::CreateLambda([DataBinderWeakPtr]() {
+			if (DataBinderWeakPtr.IsValid())
+			{
+				DataBinderWeakPtr.Pin()->BeginDataEditSequence();
+			}
+			}))
+		.OnEndSliderMovement(SNumericEntryBox<double>::FOnValueChanged::CreateLambda([DataBinderWeakPtr](double Value) {
+			if (DataBinderWeakPtr.IsValid())
+			{
+				DataBinderWeakPtr.Pin()->EndDataEditSequence();
+			}
+		}))
+		.ToolTipText(MakeAttributeLambda([this, &TargetVector, Tooltip = Tooltips[ComponentIndex], ComponentIndex]
+		{
+			return FText::Format(Tooltip, TargetVector[ComponentIndex]);
+		}))
+		.IsEnabled(IsEnabled)
+
+		.LabelPadding(FMargin(3.f))
+		.LabelLocation(SNumericEntryBox<double>::ELabelLocation::Inside)
+
+		.Label()
+		[
+			SNumericEntryBox<double>::BuildNarrowColorLabel(Colors[ComponentIndex])
+		]
+		;
+
+	return Args;
+}
+
+TSharedRef<SWidget> STransformGizmoNumericalUIOverlay::CreateTranslationComponent(int ComponentIndex, TAttribute<bool> IsEnabled)
+{
+	TSharedRef<SWidget> WidgetToReturn = SArgumentNew(CreateComponentBaseArgs(DisplayTranslation, ComponentIndex, IsEnabled), SNumericEntryBox<double>)
+
+		// All four of these have to be explicitly unset for the spin range to be unlimited.
+		.MinValue(TOptional<double>())
+		.MaxValue(TOptional<double>())
+		.MinSliderValue(TOptional<double>())
+		.MaxSliderValue(TOptional<double>())
+
+		// Controls the linear scrolling behavior
+		.LinearDeltaSensitivity(1)
+		.Delta(MakeAttributeLambda([this]() { return TranslationScrubSensitivity; }))
+		;
+
+	return WidgetToReturn;
+}
+
+TSharedRef<SWidget> STransformGizmoNumericalUIOverlay::CreateRotationComponent(int ComponentIndex, TAttribute<bool> IsEnabled)
+{
+	TSharedRef<SWidget> WidgetToReturn = SArgumentNew(CreateComponentBaseArgs(DisplayEulerAngles, ComponentIndex, IsEnabled), SNumericEntryBox<double>)
+		.MinSliderValue(-360)
+		.MaxSliderValue(360)
+		;
+
+	return WidgetToReturn;
+}
+
+TSharedRef<SWidget> STransformGizmoNumericalUIOverlay::CreateScaleComponent(int ComponentIndex, TAttribute<bool> IsEnabled)
+{
+	TSharedRef<SWidget> WidgetToReturn = SArgumentNew(CreateComponentBaseArgs(DisplayScale, ComponentIndex, IsEnabled), SNumericEntryBox<double>)
+
+		// All four of these have to be explicitly unset for the spin range to be unlimited.
+		.MinValue(TOptional<double>())
+		.MaxValue(TOptional<double>())
+		.MinSliderValue(TOptional<double>())
+		.MaxSliderValue(TOptional<double>())
+
+		// Controls the linear scrolling behavior
+		.LinearDeltaSensitivity(1)
+		.Delta(0.01)
+		;
+
+	return WidgetToReturn;
+}
+
+TSharedRef<SWidget> STransformGizmoNumericalUIOverlay::CreateWidgetForGizmo(ETransformGizmoSubElements GizmoElements)
+{
+	using namespace STransformGizmoNumericalUIOverlayLocals;
+
+	// We'll add the translation/rotation/scale rows to this
+	TSharedRef<SVerticalBox> Rows = SNew(SVerticalBox);
+
+	TSharedRef<SBox> WidgetToReturn = SNew(SBox)
+		.MinDesiredWidth(100.0f * 3.0f)
+		.MaxDesiredWidth(100.0f * 3.0f)
+		[
+			Rows
+		];
+
+	// Create the delta mode toggle row.
+	Rows->AddSlot()
+		.Padding(FMargin(0.f, 0.f, 0.f, 6.f))
+	[
+		SNew(SHorizontalBox)
+
+		// Label
+		+ SHorizontalBox::Slot()
+		.AutoWidth()
+		.VAlign(VAlign_Center)
+		.Padding(FMargin(0.f, 0.f, 8.f, 0.f))
+		[
+			SNew(STextBlock)
+			.Text(LOCTEXT("DeltaModeLabel", "Delta Mode:"))
+			.ToolTipText(LOCTEXT("DeltaModeTooltip", "In Delta mode, the gizmo values are interpreted as changes relative to last gizmo position."))
+		]
+
+		// Actual on/off toggle
+		+ SHorizontalBox::Slot()
+		.AutoWidth()
+		.Padding(FMargin(0.0, 0.f, 2.f, 0.f))
+		[
+			SNew(SCheckBox)
+			//.Style(FAppStyle::Get(), "DetailsView.SectionButton")
+			.Padding(FMargin(5, 2))
+			.HAlign(HAlign_Center)
+			.OnCheckStateChanged_Lambda([this](ECheckBoxState NewState)
+			{
+				if (NewState == ECheckBoxState::Checked)
+				{
+					DataBinder->ResetToDeltaMode();
+				}
+				else
+				{
+					DataBinder->SetToDestinationMode();
+				}
+			})
+			.IsChecked_Lambda([this]() -> ECheckBoxState { return DataBinder->IsInDeltaMode() ? ECheckBoxState::Checked : ECheckBoxState::Unchecked; })
+			[
+				SNew(SHorizontalBox)
+				+ SHorizontalBox::Slot()
+				.VAlign(VAlign_Center)
+				.Padding(FMargin(2.0, 0.f, 2.0, 0.f))
+				.AutoWidth()
+				[
+					SNew(STextBlock)
+					.Justification(ETextJustify::Center)
+					.Text(LOCTEXT("DeltaModeOnOff", "On/Off"))
+					.ToolTipText(MakeAttributeLambda([this]() -> FText
+					{
+						if (!DataBinder->ShouldDestinationModeBeAllowed())
+						{
+							return LOCTEXT("GizmoOnlySupportsDeltaMode", "Current gizmo only supports delta mode.");
+						}
+						return LOCTEXT("DeltaModeToggleTooltip", "Toggle delta mode on or off");
+					}))
+				]
+			]
+			.IsEnabled_Lambda([this]() { 
+				return DataBinder->ShouldDestinationModeBeAllowed();
+			})
+		]
+
+		// Reset button for delta mode
+		+ SHorizontalBox::Slot()
+		.AutoWidth()
+		.Padding(FMargin(2.0, 0.f, 0.f, 0.f))
+		[
+			SNew(SButton)
+			.Text(LOCTEXT("ResetToDeltaModeButton", "Reset"))
+			.ToolTipText(LOCTEXT("ResetToDeltaModeTooltip", "Reset the location relative to which delta values are measured to be current gizmo position."))
+			.HAlign(HAlign_Center)
+			.OnClicked_Lambda([this]() { 
+				DataBinder->ResetToDeltaMode();
+				return FReply::Handled(); 
+			})
+			.IsEnabled_Lambda([this]() { return DataBinder->IsInDeltaMode(); })
+		]
+	];
+
+	// Create the translation row
+	if ((GizmoElements & ETransformGizmoSubElements::TranslateAllAxes) != ETransformGizmoSubElements::None
+		|| (GizmoElements & ETransformGizmoSubElements::TranslateAllPlanes) != ETransformGizmoSubElements::None)
+	{
+		TSharedRef<SHorizontalBox> Row = SNew(SHorizontalBox);
+		Rows->AddSlot()
+		[
+			Row
+		];
+
+		// Translation icon
+		Row->AddSlot()
+			.AutoWidth()
+			.VAlign(VAlign_Center)
+			.Padding(FMargin(0.f, 0.f, ICON_RIGHT_PADDING, 0.f))
+		[
+			SNew(SImage)
+			.Image(FAppStyle::Get().GetBrush("EditorViewport.TranslateMode"))
+			.ColorAndOpacity(FSlateColor::UseForeground())
+		];
+
+		bool bHaveXComponent = (GizmoElements & ETransformGizmoSubElements::TranslateAxisX) != ETransformGizmoSubElements::None
+			|| (GizmoElements & ETransformGizmoSubElements::TranslatePlaneXY) != ETransformGizmoSubElements::None
+			|| (GizmoElements & ETransformGizmoSubElements::TranslatePlaneXZ) != ETransformGizmoSubElements::None;
+		bool bHaveYComponent = (GizmoElements & ETransformGizmoSubElements::TranslateAxisY) != ETransformGizmoSubElements::None
+			|| (GizmoElements & ETransformGizmoSubElements::TranslatePlaneXY) != ETransformGizmoSubElements::None
+			|| (GizmoElements & ETransformGizmoSubElements::TranslatePlaneYZ) != ETransformGizmoSubElements::None;
+		bool bHaveZComponent = (GizmoElements & ETransformGizmoSubElements::TranslateAxisZ) != ETransformGizmoSubElements::None
+			|| (GizmoElements & ETransformGizmoSubElements::TranslatePlaneXZ) != ETransformGizmoSubElements::None
+			|| (GizmoElements & ETransformGizmoSubElements::TranslatePlaneYZ) != ETransformGizmoSubElements::None;
+
+		Row->AddSlot()
+		[
+			CreateTranslationComponent(0, bHaveXComponent)
+		];
+		Row->AddSlot()
+		[
+			CreateTranslationComponent(1, bHaveYComponent)
+		];
+		Row->AddSlot()
+		[
+			CreateTranslationComponent(2, bHaveZComponent)
+		];
+	}
+
+	// Add rotate row
+	if ((GizmoElements & ETransformGizmoSubElements::RotateAllAxes) != ETransformGizmoSubElements::None)
+	{
+		TSharedRef<SHorizontalBox> Row = SNew(SHorizontalBox);
+		Rows->AddSlot()
+		[
+			Row
+		];
+
+		// Rotate icon
+		Row->AddSlot()
+			.AutoWidth()
+			.VAlign(VAlign_Center)
+			.Padding(FMargin(0.f, 0.f, ICON_RIGHT_PADDING, 0.f))
+		[
+			SNew(SImage)
+			.Image(FAppStyle::Get().GetBrush("EditorViewport.RotateMode"))
+			.ColorAndOpacity(FSlateColor::UseForeground())
+		];
+
+		bool bHaveXComponent = (GizmoElements & ETransformGizmoSubElements::RotateAxisX) != ETransformGizmoSubElements::None;
+		bool bHaveYComponent = (GizmoElements & ETransformGizmoSubElements::RotateAxisY) != ETransformGizmoSubElements::None;
+		bool bHaveZComponent = (GizmoElements & ETransformGizmoSubElements::RotateAxisZ) != ETransformGizmoSubElements::None;
+
+		Row->AddSlot()
+		[
+			CreateRotationComponent(0, bHaveXComponent)
+		];
+		Row->AddSlot()
+		[
+			CreateRotationComponent(1, bHaveYComponent)
+		];
+		Row->AddSlot()
+		[
+			CreateRotationComponent(2, bHaveZComponent)
+		];
+	}
+
+	// Scale row
+	if ((GizmoElements & ETransformGizmoSubElements::ScaleAllAxes) != ETransformGizmoSubElements::None
+		|| (GizmoElements & ETransformGizmoSubElements::ScaleAllPlanes) != ETransformGizmoSubElements::None
+		|| (GizmoElements & ETransformGizmoSubElements::ScaleUniform) != ETransformGizmoSubElements::None)
+	{
+		TSharedRef<SHorizontalBox> Row = SNew(SHorizontalBox);
+		Rows->AddSlot()
+		[
+			Row
+		];
+
+		// Icon
+		Row->AddSlot()
+			.AutoWidth()
+			.VAlign(VAlign_Center)
+			.Padding(FMargin(0.f, 0.f, ICON_RIGHT_PADDING, 0.f))
+		[
+			SNew(SImage)
+			.Image(FAppStyle::Get().GetBrush("EditorViewport.ScaleMode"))
+			.ColorAndOpacity(FSlateColor::UseForeground())
+		];
+
+		bool bHaveXComponent = (GizmoElements & ETransformGizmoSubElements::ScaleAxisX) != ETransformGizmoSubElements::None
+			|| (GizmoElements & ETransformGizmoSubElements::ScalePlaneXY) != ETransformGizmoSubElements::None
+			|| (GizmoElements & ETransformGizmoSubElements::ScalePlaneXZ) != ETransformGizmoSubElements::None
+			// Alternatively, create this component if the only scaling we have is uniform
+			|| ((GizmoElements & ETransformGizmoSubElements::ScaleUniform) != ETransformGizmoSubElements::None
+				&& (GizmoElements & ETransformGizmoSubElements::ScaleAllAxes) == ETransformGizmoSubElements::None
+				&& (GizmoElements & ETransformGizmoSubElements::ScaleAllPlanes) == ETransformGizmoSubElements::None);
+		bool bHaveYComponent = (GizmoElements & ETransformGizmoSubElements::ScaleAxisY) != ETransformGizmoSubElements::None
+			|| (GizmoElements & ETransformGizmoSubElements::ScalePlaneXY) != ETransformGizmoSubElements::None
+			|| (GizmoElements & ETransformGizmoSubElements::ScalePlaneYZ) != ETransformGizmoSubElements::None;
+		bool bHaveZComponent = (GizmoElements & ETransformGizmoSubElements::ScaleAxisZ) != ETransformGizmoSubElements::None
+			|| (GizmoElements & ETransformGizmoSubElements::ScalePlaneXZ) != ETransformGizmoSubElements::None
+			|| (GizmoElements & ETransformGizmoSubElements::ScalePlaneYZ) != ETransformGizmoSubElements::None;
+
+		Row->AddSlot()
+		[
+			CreateScaleComponent(0, bHaveXComponent)
+		];
+		Row->AddSlot()
+		[
+			CreateScaleComponent(1, bHaveYComponent)
+		];
+		Row->AddSlot()
+		[
+			CreateScaleComponent(2, bHaveZComponent)
+		];
+	}
+
+	return WidgetToReturn;
+}
+
+void STransformGizmoNumericalUIOverlay::SetEnabled(bool bEnabledIn)
+{
+	bIsEnabled = bEnabledIn;
+
+	// Note that currently, bIsEnabled just acts as a visibility toggle, meaning that the numerical UI
+	// system will still receive updates when disabled. This makes it easy for the system to show the
+	// correct values if it is re-enabled after being disabled.
+	// If we decide that we don't want this to be the case, we'd need to add disabling support to 
+	// FTransformGizmoDataBinder, which shouldn't be too hard.
+}
+
+void STransformGizmoNumericalUIOverlay::ResetPositionInViewport()
+{
+	if (DraggableBoxOverlay.IsValid())
+	{
+		DraggableBoxOverlay->SetBoxPosition(DefaultLeftPadding, DefaultVerticalPadding);
+	}
+}
+
+void STransformGizmoNumericalUIOverlay::MakeNumericalUISubMenu(FMenuBuilder& MenuBuilder)
+{
+	TWeakPtr<STransformGizmoNumericalUIOverlay> NumericalUI = StaticCastSharedRef<STransformGizmoNumericalUIOverlay>(AsShared());
+
+	const FUIAction NumericalUIAction(
+		FExecuteAction::CreateLambda([NumericalUI]
+		{
+			if (NumericalUI.IsValid())
+			{
+				TSharedPtr<STransformGizmoNumericalUIOverlay> NumericalUIPinned = NumericalUI.Pin();
+				NumericalUIPinned->SetEnabled(!NumericalUIPinned->IsEnabled());
+			}
+		}),
+		FCanExecuteAction::CreateLambda([NumericalUI]()
+		{
+			return NumericalUI.IsValid();
+		}),
+		FIsActionChecked::CreateLambda([NumericalUI]()
+		{
+			if (!NumericalUI.IsValid())
+			{
+				return false;
+			}
+			return NumericalUI.Pin()->IsEnabled();
+		}));
+	MenuBuilder.AddMenuEntry(LOCTEXT("NumericalUI_Enabled", "Enabled"), 
+		LOCTEXT("NumericalUI_Enabled_Tooltip", "Show the gizmo numerical UI floating panel when a gizmo is visible."),
+		FSlateIcon(), NumericalUIAction, NAME_None, EUserInterfaceActionType::ToggleButton);
+
+	const FUIAction GizmoMode_ResetNumericalUIPosition(
+		FExecuteAction::CreateLambda([NumericalUI]
+		{
+			if (NumericalUI.IsValid())
+			{
+				NumericalUI.Pin()->ResetPositionInViewport();
+			}
+		}),
+		FCanExecuteAction::CreateLambda([NumericalUI]()
+		{
+			return NumericalUI.IsValid() && NumericalUI.Pin()->IsEnabled();
+		}));
+	MenuBuilder.AddMenuEntry(LOCTEXT("NumericalUI_ResetNumericalUIPosition", "Reset UI Position"),
+		LOCTEXT("NumericalUI_ResetNumericalUIPosition_Tooltip", "Reset the numerical UI position in the viewport."),
+		FSlateIcon(), GizmoMode_ResetNumericalUIPosition, NAME_None, EUserInterfaceActionType::Button);
+}
+
+void STransformGizmoNumericalUIOverlay::SetCustomDisplayConversionFunctions(
+	TFunction<FVector3d(const FVector3d& InternalValue)> InternalToDisplayConversionIn, 
+	TFunction<FVector3d(const FVector3d& DisplayValue)> DisplayToInternalConversionIn, 
+	double TranslationScrubSensitivityIn)
+{
+	if (ensure(DataBinder.IsValid()))
+	{
+		DataBinder->SetTranslationConversionFunctions(InternalToDisplayConversionIn, DisplayToInternalConversionIn);
+
+		TranslationScrubSensitivity = TranslationScrubSensitivityIn;
+	}
+}
+
+void STransformGizmoNumericalUIOverlay::SetDefaultLocalReferenceTransform(const TOptional<FTransform>& CustomReferenceTransform)
+{
+	if (ensure(DataBinder.IsValid()))
+	{
+		DataBinder->SetDefaultLocalReferenceTransform(CustomReferenceTransform);
+	}
+}
+
+#undef LOCTEXT_NAMESPACE
\ No newline at end of file
diff --git a/Engine/Plugins/Experimental/MeshModelingToolsetExp/Source/ModelingEditorUI/Public/ModelingWidgets/SDraggableBox.h b/Engine/Plugins/Experimental/MeshModelingToolsetExp/Source/ModelingEditorUI/Public/ModelingWidgets/SDraggableBox.h
new file mode 100644
--- /dev/null
+++ b/Engine/Plugins/Experimental/MeshModelingToolsetExp/Source/ModelingEditorUI/Public/ModelingWidgets/SDraggableBox.h
@@ -0,0 +1,123 @@
+// Copyright Epic Games, Inc. All Rights Reserved.
+
+#pragma once
+
+#include "CoreMinimal.h"
+
+#include "Widgets/DeclarativeSyntaxSupport.h"
+#include "Widgets/SCompoundWidget.h"
+
+/**
+ * Class that can be used to place a draggable box into a viewport or some other large widget as an
+ * overlay. Just place the widget that you want to be draggable as the contents of SDraggableBoxOverlay.
+ */
+class MODELINGEDITORUI_API SDraggableBoxOverlay : public SCompoundWidget
+{
+public:
+
+	SLATE_BEGIN_ARGS(SDraggableBoxOverlay) {}
+
+	// When true, the positioning of the box is relative to the bottom of the widget rather than the top,
+	// so the VerticalPosition in SetBoxPosition is interpreted as distance from the bottom of the
+	// containing widget. The choice here depends on how your widget is generally positioned, for instance
+	// if it is near the bottom, you want this to be true so that making the window smaller does not
+	// immediately clip your box.
+	SLATE_ARGUMENT(bool, bPositionRelativeToBottom)
+
+	SLATE_DEFAULT_SLOT(FArguments, Content)
+	SLATE_END_ARGS()
+
+	void Construct(const FArguments& InArgs);
+
+	/**
+	 * Sets the box position in the overlay. Horizontal position is distance from left, and vertical position
+	 * is distance from bottom or from top depending on whether bPositionRelativeToBottom == true.
+	 */
+	void SetBoxPosition(float HorizontalPosition, float VerticalPosition);
+
+protected:
+
+	TSharedPtr<SWidget> DraggableBox;
+	TSharedPtr<SWidget> ContainingBox;
+
+	float DraggableBoxPaddingHorizontal = 0;
+	float DraggableBoxPaddingVertical = 0;
+};
+
+/**
+ * A widget for the draggable box itself, which requires its parent to handle its positioning in
+ * response to the drag.
+ * 
+ * Users probably shouldn't use this class directly; rather, they should use SDraggableBoxOverlay, 
+ * which will put its contents into a draggable box and properly handle the dragging without the
+ * user having to set it up.
+ */
+class MODELINGEDITORUI_API SDraggableBox : public SCompoundWidget
+{
+public:
+
+	DECLARE_DELEGATE_OneParam(FOnDragComplete, const FVector2D& /*ScreenSpacePosition*/);
+
+	SLATE_BEGIN_ARGS(SDraggableBox) {}
+	SLATE_EVENT(FOnDragComplete, OnDragComplete)
+	SLATE_DEFAULT_SLOT(FArguments, Content)
+	SLATE_END_ARGS()
+
+	void Construct(const FArguments& InArgs);
+
+	// SWidget overrides necessary for box to be draggable
+	FReply OnMouseButtonDown(const FGeometry& MyGeometry, const FPointerEvent& MouseEvent) override;
+	FReply OnDragDetected(const FGeometry& MyGeometry, const FPointerEvent& MouseEvent) override;
+
+protected:
+	TSharedPtr<SWidget> InnerWidget;
+	FOnDragComplete OnDragComplete;
+};
+
+
+/**
+ * A drag/drop operation used by SDraggableBox, largely modeled on FInViewportUIDragOperation, except
+ * instead of requiring the drop location to be able to handle that particular class to trigger the
+ * OnDragComplete delegate, it just triggers on any drop. This makes it possible to use it in any
+ * viewports, not just the level editor.
+ */
+class MODELINGEDITORUI_API FDraggableBoxUIDragOperation : public FDragDropOperation
+{
+public:
+
+	DRAG_DROP_OPERATOR_TYPE(FDraggableBoxUIDragOperation, FDragDropOperation)
+
+	// FDragDropOperation
+	virtual void OnDrop(bool bDropWasHandled, const FPointerEvent& MouseEvent) override;
+	virtual void OnDragged(const FDragDropEvent& DragDropEvent) override;
+
+	/**
+	 * Create this Drag and Drop Content
+	 *
+	 * @param InUIToBeDragged	  The UI being dragged
+	 * @param InDecoratorOffset   Where within the UI we grabbed, so we're not dragging by the upper left of the UI.
+	 * @param OwnerAreaSize       Size of the DockArea at the time when we start dragging.
+	 * @param OnDragComplete      Delegate to call when dropped.
+	 *
+	 * @return a new FDockingDragOperation
+	 */
+	static TSharedRef<FDraggableBoxUIDragOperation> New(const TSharedRef<class SWidget>& InUIToBeDragged, const FVector2D InDecoratorOffset, 
+		const FVector2D& OwnerAreaSize, SDraggableBox::FOnDragComplete& OnDragComplete);
+
+	virtual ~FDraggableBoxUIDragOperation() {}
+
+protected:
+	/** The constructor is protected, so that this class can only be instantiated as a shared pointer. */
+	FDraggableBoxUIDragOperation(const TSharedRef<class SWidget>& InUIToBeDragged, const FVector2D InDecoratorOffsetFromCursor, 
+		const FVector2D& OwnerAreaSize, SDraggableBox::FOnDragComplete& OnDragCompletee);
+
+	/**
+	 * Shared pointer to our contents, which are a weak pointer.
+	 */
+	TSharedPtr<class SWidget> UIBeingDragged;
+
+	FVector2D DecoratorOffsetFromCursor;
+	FVector2D LastContentSize;
+
+	SDraggableBox::FOnDragComplete OnDragComplete;
+};
diff --git a/Engine/Plugins/Experimental/MeshModelingToolsetExp/Source/ModelingEditorUI/Public/ModelingWidgets/STickDelegateWidget.h b/Engine/Plugins/Experimental/MeshModelingToolsetExp/Source/ModelingEditorUI/Public/ModelingWidgets/STickDelegateWidget.h
new file mode 100644
--- /dev/null
+++ b/Engine/Plugins/Experimental/MeshModelingToolsetExp/Source/ModelingEditorUI/Public/ModelingWidgets/STickDelegateWidget.h
@@ -0,0 +1,33 @@
+// Copyright Epic Games, Inc. All Rights Reserved.
+
+#pragma once
+
+#include "CoreMinimal.h"
+
+#include "Widgets/DeclarativeSyntaxSupport.h"
+#include "Widgets/SCompoundWidget.h"
+
+/**
+ * Simple widget that calls its OnTickDelegate in its Tick function, to allow for custom
+ * tick behavior via composition instead of having to create a class and override the method,
+ * for instance when creating custom slate in-line somewhere.
+ */
+class MODELINGEDITORUI_API STickDelegateWidget : public SCompoundWidget
+{
+public:
+
+	DECLARE_DELEGATE_ThreeParams(FOnTick, const FGeometry& AllottedGeometry, const double InCurrentTime, const float InDeltaTime);
+
+	SLATE_BEGIN_ARGS(STickDelegateWidget) {}
+	SLATE_EVENT(FOnTick, OnTickDelegate)
+	SLATE_DEFAULT_SLOT(FArguments, Content)
+	SLATE_END_ARGS()
+
+	void Construct(const FArguments& InArgs);
+
+	virtual void Tick(const FGeometry& AllottedGeometry, const double InCurrentTime, const float InDeltaTime) override;
+
+protected:
+
+	FOnTick OnTickDelegate;
+};
\ No newline at end of file
diff --git a/Engine/Plugins/Experimental/MeshModelingToolsetExp/Source/ModelingEditorUI/Public/STransformGizmoNumericalUIOverlay.h b/Engine/Plugins/Experimental/MeshModelingToolsetExp/Source/ModelingEditorUI/Public/STransformGizmoNumericalUIOverlay.h
new file mode 100644
--- /dev/null
+++ b/Engine/Plugins/Experimental/MeshModelingToolsetExp/Source/ModelingEditorUI/Public/STransformGizmoNumericalUIOverlay.h
@@ -0,0 +1,139 @@
+// Copyright Epic Games, Inc. All Rights Reserved.
+
+#pragma once
+
+#include "CoreMinimal.h"
+
+#include "InteractiveGizmo.h" // ETransformGizmoSubElements
+#include "Widgets/DeclarativeSyntaxSupport.h"
+#include "Widgets/Input/SNumericEntryBox.h"
+#include "Widgets/SCompoundWidget.h"
+
+class FTransformGizmoDataBinder;
+class SBorder;
+class SDraggableBoxOverlay;
+class SWidget;
+class UInteractiveToolsContext;
+
+/**
+ * An overlay that contains a draggable panel that mirrors the values of a UCombinedTransformGizmo
+ * and allows them to be settable by scrubbing the values or typing.
+ *
+ * To use it, just add the overlay to your viewport and call BindToGizmoContextObject (sometime after
+ * TransformGizmoUtil::RegisterTransformGizmoContextObject has already been called). This will bind
+ * any gizmos subsequently created through TransformGizmoUtil to the panel.
+ *
+ * MakeNumericalUISubMenu creates a submenu that can be slotted somewhere to enable/disable the UI
+ * panel or to reset its position in the viewport.
+ */
+class MODELINGEDITORUI_API STransformGizmoNumericalUIOverlay : public SCompoundWidget
+{
+public:
+
+	SLATE_BEGIN_ARGS(STransformGizmoNumericalUIOverlay) 
+		: _bPositionRelativeToBottom(true)
+		, _DefaultLeftPadding(15)
+		, _DefaultVerticalPadding(15)
+		, _TranslationScrubSensitivity(1)
+	{}
+
+	// When true (the default) the panel is positioned by measuring a distance from the bottom of
+	// the viewport, and resizing the viewport will keep that distance the same.
+	SLATE_ARGUMENT(bool, bPositionRelativeToBottom)
+	// These values determine where the panel is reset when ResetPositionInViewport() is called.
+	SLATE_ARGUMENT(float, DefaultLeftPadding)
+	SLATE_ARGUMENT(float, DefaultVerticalPadding)
+
+	// Usually do not need setting. See comment for SetCustomDisplayConversionFunctions
+	SLATE_ARGUMENT(TFunction<FVector3d(const FVector3d& InternalValue)>, InternalToDisplayFunction)
+	SLATE_ARGUMENT(TFunction<FVector3d(const FVector3d& DisplayValue)>, DisplayToInternalFunction)
+	SLATE_ARGUMENT(double, TranslationScrubSensitivity)
+	// Usually does not need setting. See comment for SetDefaultLocalReferenceTransform
+	SLATE_ARGUMENT(TOptional<FTransform>, DefaultLocalReferenceTransform)
+
+	SLATE_DEFAULT_SLOT(FArguments, Content)
+	SLATE_END_ARGS()
+
+	void Construct(const FArguments& InArgs);
+
+	virtual ~STransformGizmoNumericalUIOverlay();
+
+	/**
+	 * Binding to the context object makes it so that the overlay gets attached to any subsequently created 
+	 * transform gizmos.
+	 */
+	bool BindToGizmoContextObject(UInteractiveToolsContext* ToolsContext);
+
+	virtual bool IsEnabled() { return bIsEnabled; }
+	virtual void SetEnabled(bool bEnabledIn);
+
+	/**
+	 * Resets the position of the draggable UI box in the overlay to its default.
+	 */
+	virtual void ResetPositionInViewport();
+
+	/**
+	 * Creates a submenu to deal with the numerical UI, for placement in some toolbar by the mode/editor.
+	 */
+	virtual void MakeNumericalUISubMenu(FMenuBuilder& MenuBuilder);
+
+	/**
+	 * Clears all gizmo bindings
+	 */
+	virtual void Reset();
+
+	/**
+	 * Allows the positional (translation) copmonent to be displayed in different units than the gizmo world, for
+	 * instance in a UV editor where the world coordinates are not actually 0 to 1.
+	 */
+	virtual void SetCustomDisplayConversionFunctions(
+		TFunction<FVector3d(const FVector3d& InternalValue)> InternalToDisplayConversionIn,
+		TFunction<FVector3d(const FVector3d& DisplayValue)> DisplayToInternalConversionIn,
+		double TranslationScrubSensitivityIn = 1.0);
+
+	/**
+	 * After calling this function, any newly initialized gizmos will be given the provided reference
+	 * transform as their custom reference transform. This mainly meant to allow the UI to be used in
+	 * a 2D situation (e.g. UV editor) where the gizmos only have two axes. Normally, we only allow 
+	 * delta mode in such a case, but the 2D editor can provide the used plane's reference transform 
+	 * to allow the gizmos to be usable in destination mode as well.
+	 */
+	virtual void SetDefaultLocalReferenceTransform(const TOptional<FTransform>& CustomReferenceTransform);
+
+protected:
+
+	// Currently displayed values
+	FVector3d DisplayTranslation;
+	FVector3d DisplayEulerAngles;
+	FVector3d DisplayScale;
+
+	// Holds the functionality that binds the displayed values to the gizmo.
+	TSharedPtr<FTransformGizmoDataBinder> DataBinder;
+
+	TSharedPtr<SDraggableBoxOverlay> DraggableBoxOverlay;
+
+	// The part of the UI that changes when a different gizmo is chosen
+	TSharedPtr<SBorder> WidgetContents;
+
+	// The values to which the UI positioning is reset if asked
+	float DefaultLeftPadding = 15.0f;
+	float DefaultVerticalPadding = 75.0f;
+
+	// Helpers for building the UI when changing gizmo
+	TSharedRef<SWidget> CreateWidgetForGizmo(ETransformGizmoSubElements GizmoElements);
+	SNumericEntryBox<double>::FArguments CreateComponentBaseArgs(FVector3d& TargetVector, int ComponentIndex, TAttribute<bool> IsEnabled);
+	TSharedRef<SWidget> CreateTranslationComponent(int ComponentIndex, TAttribute<bool> IsEnabled);
+	TSharedRef<SWidget> CreateRotationComponent(int ComponentIndex, TAttribute<bool> IsEnabled);
+	TSharedRef<SWidget> CreateScaleComponent(int ComponentIndex, TAttribute<bool> IsEnabled);
+
+	// We have to respond to coordinate system changes for the gizmo, and to do that we need to be able to poll
+	// on tick, since we don't get notifications about it.
+	EToolContextCoordinateSystem LastCoordinateSystem;
+	void OnTickWidget(const FGeometry& AllottedGeometry, const double InCurrentTime, const float InDeltaTime);
+
+	// Modifier for scrubbing translation, if we're using custom display functions.
+	double TranslationScrubSensitivity = 1;
+
+	bool bIsEnabled = true;
+	bool bPositionRelativeToBottom = true;
+};
diff --git a/Engine/Plugins/Experimental/UVEditor/Source/UVEditor/Private/UVEditorMode.cpp b/Engine/Plugins/Experimental/UVEditor/Source/UVEditor/Private/UVEditorMode.cpp
--- a/Engine/Plugins/Experimental/UVEditor/Source/UVEditor/Private/UVEditorMode.cpp
+++ b/Engine/Plugins/Experimental/UVEditor/Source/UVEditor/Private/UVEditorMode.cpp
@@ -7,6 +7,7 @@
 #include "Actions/UVSeamSewAction.h"
 #include "Actions/UVToolAction.h"
 #include "AssetEditorModeManager.h"
+#include "BaseGizmos/TransformGizmoUtil.h"
 #include "ContextObjects/UVToolViewportButtonsAPI.h"
 #include "ContextObjects/UVToolAssetInputsContext.h"
 #include "ContextObjectStore.h"
@@ -454,6 +455,13 @@ void UUVEditorMode::Enter()
 	RegisterActions();
 	ActivateDefaultTool();
 
+	// do any toolkit UI initialization that depends on the mode setup above
+	if (Toolkit.IsValid())
+	{
+		FUVEditorModeToolkit* UVModeToolkit = (FUVEditorModeToolkit*)Toolkit.Get();
+		UVModeToolkit->InitializeAfterModeSetup();
+	}
+
 	if (FEngineAnalytics::IsAvailable())
 	{
 		AnalyticsLastStartTimestamp = FDateTime::UtcNow();
@@ -997,6 +1005,9 @@ void UUVEditorMode::InitializeModeContexts()
 
 	UUVEditorToolPropertiesAPI* UVEditorToolPropertiesAPI = NewObject<UUVEditorToolPropertiesAPI>();
 	AddContextObject(UVEditorToolPropertiesAPI);
+
+	// Register gizmo helper
+	UE::TransformGizmoUtil::RegisterTransformGizmoContextObject(GetInteractiveToolsContext());
 }
 
 void UUVEditorMode::InitializeTargets()
diff --git a/Engine/Plugins/Experimental/UVEditor/Source/UVEditor/Private/UVEditorModeToolkit.cpp b/Engine/Plugins/Experimental/UVEditor/Source/UVEditor/Private/UVEditorModeToolkit.cpp
--- a/Engine/Plugins/Experimental/UVEditor/Source/UVEditor/Private/UVEditorModeToolkit.cpp
+++ b/Engine/Plugins/Experimental/UVEditor/Source/UVEditor/Private/UVEditorModeToolkit.cpp
@@ -10,10 +10,12 @@
 #include "IDetailsView.h"
 #include "Modules/ModuleManager.h"
 #include "SPrimaryButton.h"
+#include "STransformGizmoNumericalUIOverlay.h"
 #include "Tools/UEdMode.h"
 #include "UVEditorBackgroundPreview.h"
 #include "UVEditorCommands.h"
 #include "UVEditorMode.h"
+#include "UVEditorUXSettings.h"
 #include "Widgets/Input/SButton.h"
 #include "Widgets/Layout/SBorder.h"
 #include "Widgets/Text/STextBlock.h"
@@ -195,6 +197,42 @@ void FUVEditorModeToolkit::Init(const TSharedPtr<IToolkitHost>& InitToolkitHost,
 		]	
 	];
 
+	// Set up the gizmo numerical UI
+	GizmoNumericalUIOverlayWidget = SNew(STransformGizmoNumericalUIOverlay)
+		.bPositionRelativeToBottom(true)
+		.DefaultLeftPadding(15.0f)
+		// Unlike the level editor, we don't have an axis indicator that we have to be above, so this value is different
+		.DefaultVerticalPadding(15.0f)
+		// Enable non-delta mode by giving a default reference transform
+		.DefaultLocalReferenceTransform(FTransform())
+		// Make our displays show 0-1 units instead of world units.
+		.InternalToDisplayFunction([](const FVector3d& WorldPosition)
+		{
+			FVector2f UV = FUVEditorUXSettings::InternalUVToExternalUV(FUVEditorUXSettings::VertPositionToUV(WorldPosition));
+			return FVector3d(UV.X, UV.Y, 0);
+		})
+		.DisplayToInternalFunction([](const FVector3d& UVPosition)
+		{
+			return FUVEditorUXSettings::UVToVertPosition(FUVEditorUXSettings::ExternalUVToInternalUV(FVector2f(UVPosition.X, UVPosition.Y)));
+		})
+		.TranslationScrubSensitivity(1 / FUVEditorUXSettings::UVMeshScalingFactor)
+		;
+
+	GetToolkitHost()->AddViewportOverlayWidget(GizmoNumericalUIOverlayWidget.ToSharedRef());
+}
+
+void FUVEditorModeToolkit::InitializeAfterModeSetup()
+{
+	if (bFirstInitializeAfterModeSetup)
+	{
+		if (ensure(GizmoNumericalUIOverlayWidget.IsValid()))
+		{
+			// This has to happen after the gizmo context object is registered
+			GizmoNumericalUIOverlayWidget->BindToGizmoContextObject(GetScriptableEditorMode()->GetInteractiveToolsContext());
+		}
+		bFirstInitializeAfterModeSetup = false;
+	}
+	
 }
 
 FName FUVEditorModeToolkit::GetToolkitFName() const
@@ -289,6 +327,14 @@ TSharedRef<SWidget> FUVEditorModeToolkit::GetToolDisplaySettingsWidget()
 	}
 }
 
+void FUVEditorModeToolkit::MakeGizmoNumericalUISubMenu(FMenuBuilder& MenuBuilder)
+{
+	if (GizmoNumericalUIOverlayWidget.IsValid())
+	{
+		GizmoNumericalUIOverlayWidget->MakeNumericalUISubMenu(MenuBuilder);
+	}
+}
+
 TSharedRef<SWidget> FUVEditorModeToolkit::CreateDisplaySettingsWidget(UObject* SettingsObject) const
 {
 	TSharedRef<SBorder> GridDetailsContainer =
diff --git a/Engine/Plugins/Experimental/UVEditor/Source/UVEditor/Private/UVEditorToolkit.cpp b/Engine/Plugins/Experimental/UVEditor/Source/UVEditor/Private/UVEditorToolkit.cpp
--- a/Engine/Plugins/Experimental/UVEditor/Source/UVEditor/Private/UVEditorToolkit.cpp
+++ b/Engine/Plugins/Experimental/UVEditor/Source/UVEditor/Private/UVEditorToolkit.cpp
@@ -515,6 +515,29 @@ void FUVEditorToolkit::PostInitAssetEditor()
 			
 			TSharedRef<SVerticalBox> Container = SNew(SVerticalBox);
 
+			FMenuBuilder MenuBuilder(true, TSharedPtr<FUICommandList>());
+			MenuBuilder.BeginSection("Section_Gizmo", LOCTEXT("Section_Gizmo", "Gizmo"));
+
+			MenuBuilder.AddSubMenu(
+				LOCTEXT("NumericalUISubMenu", "Numerical UI"), LOCTEXT("NumericalUISubMenu_ToolTip", "Configure the gizmo numerical UI."),
+				FNewMenuDelegate::CreateLambda([UVModeToolkit](FMenuBuilder& SubMenuBuilder) {
+					UVModeToolkit->MakeGizmoNumericalUISubMenu(SubMenuBuilder);
+				}));
+
+			MenuBuilder.EndSection();
+
+
+			Container->AddSlot()
+			.AutoHeight()
+			.Padding(FMargin(0.f, 0.f, 8.f, 0.f))
+			[
+				SNew(SBox)
+				.MinDesiredWidth(500)
+				[
+					MenuBuilder.MakeWidget()
+				]
+			];
+
 			Container->AddSlot()
 				.AutoHeight()
 				.Padding(FMargin(0.f, 0.f, 8.f, 0.f))
diff --git a/Engine/Plugins/Experimental/UVEditor/Source/UVEditor/Public/UVEditorModeToolkit.h b/Engine/Plugins/Experimental/UVEditor/Source/UVEditor/Public/UVEditorModeToolkit.h
--- a/Engine/Plugins/Experimental/UVEditor/Source/UVEditor/Public/UVEditorModeToolkit.h
+++ b/Engine/Plugins/Experimental/UVEditor/Source/UVEditor/Public/UVEditorModeToolkit.h
@@ -11,6 +11,7 @@
 struct FSlateBrush;
 class SBorder;
 class STextBlock;
+class STransformGizmoNumericalUIOverlay;
 class UInteractiveToolPropertySet;
 
 /**
@@ -23,6 +24,8 @@ public:
 	FUVEditorModeToolkit();
 	~FUVEditorModeToolkit();
 
+	virtual void InitializeAfterModeSetup();
+
 	/** Creates a menu where the displayed UV Channel can be changed for each asset */
 	virtual TSharedRef<SWidget> CreateChannelMenu();
 	
@@ -38,6 +41,9 @@ public:
 	/** Creates a widget where the active tool's display settings can be changed. */
 	virtual TSharedRef<SWidget> GetToolDisplaySettingsWidget();
 
+	/** Creates submenu for manipulating the gizmo numerical UI */
+	virtual void MakeGizmoNumericalUISubMenu(FMenuBuilder& MenuBuilder);
+
 	/** Enables or diables the Play in Editor mode warning for the editor */
 	virtual void EnableShowPIEWarning(bool bEnable);
 
@@ -104,6 +110,8 @@ protected:
 	/** Contains the widget container for the Accept/Cancel buttons for tools. */
 	TSharedPtr<SWidget> ViewportOverlayWidget;
 
+	TSharedPtr<STransformGizmoNumericalUIOverlay> GizmoNumericalUIOverlayWidget;
+
 	FText ActiveToolName;
 	const FSlateBrush* ActiveToolIcon;
 	FStatusBarMessageHandle ActiveToolMessageHandle;
@@ -112,4 +120,5 @@ protected:
 	void UpdateActiveToolProperties();
 	void InvalidateCachedDetailPanelState(UObject* ChangedObject);
 
+	bool bFirstInitializeAfterModeSetup = true;
 };
\ No newline at end of file
diff --git a/Engine/Plugins/Experimental/UVEditor/Source/UVEditor/UVEditor.Build.cs b/Engine/Plugins/Experimental/UVEditor/Source/UVEditor/UVEditor.Build.cs
--- a/Engine/Plugins/Experimental/UVEditor/Source/UVEditor/UVEditor.Build.cs
+++ b/Engine/Plugins/Experimental/UVEditor/Source/UVEditor/UVEditor.Build.cs
@@ -56,6 +56,7 @@ public class UVEditor : ModuleRules
 				"LevelEditor",
 				"ModelingComponentsEditorOnly", // Static/skeletal mesh tool targets
 				"ModelingComponents",
+				"ModelingEditorUI",
 				"StatusBar",
 				"ToolMenus",
 				"ToolWidgets",
diff --git a/Engine/Plugins/Experimental/UVEditor/Source/UVEditorTools/Private/UVSelectTool.cpp b/Engine/Plugins/Experimental/UVEditor/Source/UVEditorTools/Private/UVSelectTool.cpp
--- a/Engine/Plugins/Experimental/UVEditor/Source/UVEditorTools/Private/UVSelectTool.cpp
+++ b/Engine/Plugins/Experimental/UVEditor/Source/UVEditorTools/Private/UVSelectTool.cpp
@@ -4,6 +4,7 @@
 
 #include "BaseGizmos/CombinedTransformGizmo.h"
 #include "BaseGizmos/GizmoBaseComponent.h"
+#include "BaseGizmos/TransformGizmoUtil.h"
 #include "ContextObjects/UVToolViewportButtonsAPI.h"
 #include "ContextObjectStore.h"
 #include "DynamicMesh/DynamicMeshChangeTracker.h"
@@ -235,7 +236,7 @@ void UUVSelectTool::Setup()
 	// Gizmo setup
 	UInteractiveGizmoManager* GizmoManager = GetToolManager()->GetPairedGizmoManager();
 	UTransformProxy* TransformProxy = NewObject<UTransformProxy>(this);
-	TransformGizmo = GizmoManager->CreateCustomTransformGizmo(
+	TransformGizmo = UE::TransformGizmoUtil::CreateCustomTransformGizmo(GizmoManager, 
 		ETransformGizmoSubElements::TranslateAxisX | ETransformGizmoSubElements::TranslateAxisY | ETransformGizmoSubElements::TranslatePlaneXY
 		| ETransformGizmoSubElements::ScaleAxisX | ETransformGizmoSubElements::ScaleAxisY | ETransformGizmoSubElements::ScalePlaneXY
 		| ETransformGizmoSubElements::RotateAxisZ,
diff --git a/Engine/Plugins/Runtime/MeshModelingToolset/Source/MeshModelingTools/Private/EditMeshPolygonsTool.cpp b/Engine/Plugins/Runtime/MeshModelingToolset/Source/MeshModelingTools/Private/EditMeshPolygonsTool.cpp
--- a/Engine/Plugins/Runtime/MeshModelingToolset/Source/MeshModelingTools/Private/EditMeshPolygonsTool.cpp
+++ b/Engine/Plugins/Runtime/MeshModelingToolset/Source/MeshModelingTools/Private/EditMeshPolygonsTool.cpp
@@ -960,6 +960,13 @@ void UEditMeshPolygonsTool::OnGizmoTransformChanged(UTransformProxy* Proxy, FTra
 void UEditMeshPolygonsTool::OnEndGizmoTransform(UTransformProxy* Proxy)
 {
 	bInGizmoDrag = false;
+	// Sometimes we don't get a tick between OnGizmoTransformChanged and OnEndGizmoTransform. In
+	// most drag cases this is not much of a problem, but if we type values into the gizmo numerical
+	// UI, it is.
+	if (bGizmoUpdatePending)
+	{
+		ComputeUpdate_Gizmo();
+	}
 	bGizmoUpdatePending = false;
 	bSpatialDirty = true;
 	SelectionMechanic->NotifyMeshChanged(false);
diff --git a/Engine/Plugins/Runtime/MeshModelingToolset/Source/ModelingComponents/Private/Transforms/TransformGizmoDataBinder.cpp b/Engine/Plugins/Runtime/MeshModelingToolset/Source/ModelingComponents/Private/Transforms/TransformGizmoDataBinder.cpp
new file mode 100644
--- /dev/null
+++ b/Engine/Plugins/Runtime/MeshModelingToolset/Source/ModelingComponents/Private/Transforms/TransformGizmoDataBinder.cpp
@@ -0,0 +1,561 @@
+// Copyright Epic Games, Inc. All Rights Reserved.
+
+#include "Transforms/TransformGizmoDataBinder.h"
+
+#include "BaseGizmos/TransformGizmoUtil.h" // UCombinedTransformGizmoContextObject
+#include "BaseGizmos/TransformProxy.h"
+#include "VectorTypes.h"
+
+#define LOCTEXT_NAMESPACE "TransformGizmoDataBinder"
+
+// Key function for going from gizmo to values
+void FTransformGizmoDataBinder::UpdateDataValuesFromGizmo()
+{
+	if (!CurrentlyTrackedGizmo.IsValid() || !ensure(CurrentlyTrackedGizmo->ActiveTarget))
+	{
+		return;
+	}
+
+	FTransform CurrentTransform = CurrentlyTrackedGizmo->ActiveTarget->GetTransform();
+	FVector3d CurrentTranslation;
+	FRotator CurrentRotation;
+	if (bUsingDeltaMode)
+	{
+		*BoundScale = CurrentTransform.GetScale3D() / DeltaStartTransform.GetScale3D();
+
+		if (CurrentlyTrackedGizmo->CurrentCoordinateSystem == EToolContextCoordinateSystem::Local)
+		{
+			// Here we want the delta from the start transform in the coordinate system of the start transform
+			CurrentRotation = FRotator(DeltaStartTransform.GetRotation().Inverse() * CurrentTransform.GetRotation());
+			CurrentTranslation = DeltaStartTransform.GetRotation().UnrotateVector(CurrentTransform.GetTranslation() - DeltaStartTransform.GetTranslation());
+		}
+		else
+		{
+			// Get the delta from the start transform in the world coordinate system
+			CurrentRotation = FRotator(CurrentTransform.GetRotation() * DeltaStartTransform.GetRotation().Inverse());
+			CurrentTranslation = CurrentTransform.GetTranslation() - DeltaStartTransform.GetTranslation();
+		}
+	}
+	else
+	{
+		TOptional<FTransform> ReferenceTransformToUse;
+		if (CurrentlyTrackedGizmo->CurrentCoordinateSystem == EToolContextCoordinateSystem::Local && CurrentCustomLocalReferenceTransform.IsSet())
+		{
+			ReferenceTransformToUse = CurrentCustomLocalReferenceTransform;
+		}
+
+		if (ReferenceTransformToUse.IsSet())
+		{
+			const FTransform& ReferenceTransform = ReferenceTransformToUse.GetValue();
+
+			*BoundScale = CurrentTransform.GetScale3D() / ReferenceTransform.GetScale3D();
+			CurrentRotation = FRotator(ReferenceTransform.GetRotation().Inverse() * CurrentTransform.GetRotation());
+			CurrentTranslation = ReferenceTransform.GetRotation().UnrotateVector(CurrentTransform.GetTranslation() - ReferenceTransform.GetTranslation());
+		}
+		else
+		{
+			CurrentTranslation = CurrentTransform.GetTranslation();
+			CurrentRotation = FRotator(CurrentTransform.GetRotation());
+			*BoundScale = CurrentTransform.GetScale3D();
+		}
+	}
+
+	*BoundTranslation = ActualToBoundConversion ? ActualToBoundConversion(CurrentTranslation) : CurrentTranslation;
+
+	// We don't want to update the rotation if it is actually the same as what the user typed in previously.
+	// Note that FRotator::Equals doesn't actually seem to do a proper comparison of rotations within a tolerance, despite what 
+	// its header comment promises, hence us using FQuat::Equals here.
+	FQuat ExistingRotation = FQuat::MakeFromEuler(*BoundEulerAngles);
+	if (!ExistingRotation.Equals(CurrentRotation.Quaternion(), KINDA_SMALL_NUMBER))
+	{
+		*BoundEulerAngles = CurrentRotation.Euler();
+	}
+}
+
+// Key function for going from values to gizmo
+FTransform FTransformGizmoDataBinder::GetGizmoTransformFromDataValues(bool bLocalCoordinates)
+{
+	FVector3d CurrentTranslation = BoundToActualConversion ? BoundToActualConversion(*BoundTranslation) : *BoundTranslation;
+	FRotator CurrentRotation = FRotator::MakeFromEuler(*BoundEulerAngles);
+	FTransform CurrentUITransform(CurrentRotation, CurrentTranslation, *BoundScale);
+	FTransform ResultTransform;
+	if (bUsingDeltaMode)
+	{
+		ResultTransform.SetScale3D(DeltaStartTransform.GetScale3D() * *BoundScale);
+
+		if (bLocalCoordinates)
+		{
+			ResultTransform.SetRotation(DeltaStartTransform.GetRotation() * CurrentRotation.Quaternion());
+			ResultTransform.SetTranslation(DeltaStartTransform.GetTranslation()
+				+ DeltaStartTransform.GetRotation().RotateVector(CurrentTranslation));
+		}
+		else
+		{
+			ResultTransform.SetRotation(CurrentRotation.Quaternion() * DeltaStartTransform.GetRotation());
+			ResultTransform.SetTranslation(DeltaStartTransform.GetTranslation() + CurrentTranslation);
+		}
+	}
+	else
+	{
+		TOptional<FTransform> ReferenceTransformToUse;
+		if (bLocalCoordinates && CurrentCustomLocalReferenceTransform.IsSet())
+		{
+			ReferenceTransformToUse = CurrentCustomLocalReferenceTransform;
+		}
+
+		if (ReferenceTransformToUse.IsSet())
+		{
+			const FTransform& ReferenceTransform = ReferenceTransformToUse.GetValue();
+
+			ResultTransform.SetScale3D(ReferenceTransform.GetScale3D() * *BoundScale);
+			ResultTransform.SetRotation(ReferenceTransform.GetRotation() * CurrentRotation.Quaternion());
+			ResultTransform.SetTranslation(ReferenceTransform.GetTranslation()
+				+ ReferenceTransform.GetRotation().RotateVector(CurrentTranslation));
+		}
+		else
+		{
+			ResultTransform = CurrentUITransform;
+		}
+	}
+
+	return ResultTransform;
+}
+
+
+void FTransformGizmoDataBinder::SetTrackedGizmo(UCombinedTransformGizmo* Gizmo)
+{
+	if (CurrentlyTrackedGizmo != Gizmo && ensure(Gizmo && Gizmo->ActiveTarget && BoundGizmos.Contains(Gizmo)))
+	{
+		CurrentlyTrackedGizmo = Gizmo;
+
+		ETransformGizmoSubElements GizmoElements = CurrentlyTrackedGizmo->GetGizmoElements();
+
+		CurrentCustomLocalReferenceTransform = Gizmo->GetDisplaySpaceTransform();
+		if (!CurrentCustomLocalReferenceTransform.IsSet())
+		{
+			CurrentCustomLocalReferenceTransform = DefaultCustomLocalReferenceTransform;
+		}
+
+		LastCoordinateSystem = CurrentlyTrackedGizmo->CurrentCoordinateSystem;
+
+		bCurrentGizmoLacksDegreeOfFreedom = ((GizmoElements & ETransformGizmoSubElements::TranslateAllAxes) != ETransformGizmoSubElements::TranslateAllAxes)
+			|| ((GizmoElements & ETransformGizmoSubElements::RotateAllAxes) != ETransformGizmoSubElements::RotateAllAxes);
+
+		bCurrentGizmoOnlyHasUniformScale = (GizmoElements & ETransformGizmoSubElements::ScaleUniform) != ETransformGizmoSubElements::None
+			&& (GizmoElements & ETransformGizmoSubElements::ScaleAllAxes) == ETransformGizmoSubElements::None
+			&& (GizmoElements & ETransformGizmoSubElements::ScaleAllPlanes) == ETransformGizmoSubElements::None;
+
+		// See header comment for ShouldDestinationModeBeAllowed
+		if (bAvoidDestinationModeWhenUnsafe && !ShouldDestinationModeBeAllowed() && !bUsingDeltaMode)
+		{
+			ResetToDeltaMode();
+		}
+		else if (bUsingDeltaMode)
+		{
+			ResetToDeltaMode();
+		}
+		else
+		{
+			UpdateDataValuesFromGizmo();
+		}
+
+		OnTrackedGizmoChanged.Broadcast(Gizmo);
+	}
+}
+
+void FTransformGizmoDataBinder::SetToDeltaMode(const FTransform& StartTransform)
+{
+	bUsingDeltaMode = true;
+	DeltaStartTransform = StartTransform;
+
+	UpdateDataValuesFromGizmo();
+}
+
+void FTransformGizmoDataBinder::ResetToDeltaMode()
+{
+	bUsingDeltaMode = true;
+
+	FTransform StartTransform;
+	if (CurrentlyTrackedGizmo.IsValid() && ensure(CurrentlyTrackedGizmo->ActiveTarget))
+	{
+		StartTransform = CurrentlyTrackedGizmo->ActiveTarget->GetTransform();
+	}
+	DeltaStartTransform = StartTransform;
+
+	// UI values should be at zeroes
+	*BoundTranslation = ActualToBoundConversion ? ActualToBoundConversion(FVector::Zero()) : FVector::Zero();
+	*BoundEulerAngles = FVector::Zero();
+	*BoundScale = FVector::One();
+}
+
+void FTransformGizmoDataBinder::SetToDestinationMode()
+{
+	bUsingDeltaMode = false;
+	UpdateDataValuesFromGizmo();
+}
+
+bool FTransformGizmoDataBinder::HasVisibleGizmo()
+{
+	return CurrentlyTrackedGizmo.IsValid() && CurrentlyTrackedGizmo->ActiveTarget && CurrentlyTrackedGizmo->IsVisible();
+}
+
+bool FTransformGizmoDataBinder::ShouldDestinationModeBeAllowed()
+{
+	return !(bCurrentGizmoLacksDegreeOfFreedom
+		&& !CurrentCustomLocalReferenceTransform.IsSet()
+		&& CurrentlyTrackedGizmo.IsValid()
+		&& CurrentlyTrackedGizmo->CurrentCoordinateSystem == EToolContextCoordinateSystem::Local);
+}
+
+
+// Functions to listen to gizmo changes:
+
+void FTransformGizmoDataBinder::OnProxyBeginTransformEdit(UTransformProxy* TransformProxy, UCombinedTransformGizmo* Gizmo)
+{
+	if (!TransformProxy)
+	{
+		checkSlow(false);
+		return;
+	}
+
+	if (bIgnoreCallbackForDebouncing)
+	{
+		return;
+	}
+
+	bGizmoIsBeingDragged = true;
+
+	if (bChangeDisplayedGizmoOnDrag && CurrentlyTrackedGizmo != Gizmo)
+	{
+		SetTrackedGizmo(Gizmo);
+	}
+
+	if (CurrentlyTrackedGizmo == Gizmo && bUsingDeltaMode)
+	{
+		// Reset delta mode at the start of each drag
+		ResetToDeltaMode();
+	}
+}
+
+void FTransformGizmoDataBinder::OnProxyTransformChanged(UTransformProxy* TransformProxy, FTransform Transform, UCombinedTransformGizmo* Gizmo)
+{
+	if (!TransformProxy)
+	{
+		checkSlow(false);
+		return;
+	}
+
+	if (bIgnoreCallbackForDebouncing || CurrentlyTrackedGizmo != Gizmo)
+	{
+		return;
+	}
+
+	if (!bGizmoIsBeingDragged && bUsingDeltaMode)
+	{
+		// In delta mode, if we got some kind of update that didn't come from a drag (so we haven't reset it yet), we probably
+		// need to reset delta mode. However it's worth seeing whether the update actually changes our transform, because sometimes
+		// tools issue updates that keep the transform the same (for instance, to make sure that the scale is all 1's). In that case
+		// we don't want to reset anything.
+		FTransform TransformFromUI = GetGizmoTransformFromDataValues(CurrentlyTrackedGizmo->CurrentCoordinateSystem == EToolContextCoordinateSystem::Local);
+
+		if (!TransformFromUI.GetRotation().Equals(CurrentlyTrackedGizmo->ActiveTarget->GetTransform().GetRotation(), KINDA_SMALL_NUMBER)
+			|| !TransformFromUI.GetTranslation().Equals(CurrentlyTrackedGizmo->ActiveTarget->GetTransform().GetTranslation(), KINDA_SMALL_NUMBER)
+			|| !TransformFromUI.GetScale3D().Equals(CurrentlyTrackedGizmo->ActiveTarget->GetTransform().GetScale3D(), KINDA_SMALL_NUMBER))
+		{
+			ResetToDeltaMode();
+		}
+	}
+	else
+	{
+		// If we're not in delta mode or if we're dragging, no need to figure out whether the update changes anything, just do
+		// the UI update (which involves much of the same work).
+		UpdateDataValuesFromGizmo();
+	}
+}
+
+void FTransformGizmoDataBinder::OnProxyEndTransformEdit(UTransformProxy* TransformProxy, UCombinedTransformGizmo* Gizmo)
+{
+	bGizmoIsBeingDragged = false;
+}
+
+void FTransformGizmoDataBinder::OnDisplaySpaceTransformChanged(UCombinedTransformGizmo* Gizmo, TOptional<FTransform> Transform)
+{
+	if (ensure(Gizmo) && CurrentlyTrackedGizmo == Gizmo)
+	{
+		CurrentCustomLocalReferenceTransform = Transform;
+		if (!CurrentCustomLocalReferenceTransform.IsSet())
+		{
+			CurrentCustomLocalReferenceTransform = DefaultCustomLocalReferenceTransform;
+		}
+		UpdateDataValuesFromGizmo();
+	}
+}
+
+void FTransformGizmoDataBinder::OnGizmoVisibilityChanged(UCombinedTransformGizmo* Gizmo, bool bVisible)
+{
+	// Generally we don't need to worry about visibility changes for the gizmo that we are bound to,
+	// since it's one of the things we poll for in the UI panel visibility. However, for non-bound
+	// gizmos, we do need to listen in case we don't have a currently displayed gizmo, and therefore
+	// should switch to showing a newly visible different gizmo.
+	if (!ensure(Gizmo) || CurrentlyTrackedGizmo == Gizmo)
+	{
+		return;
+	}
+
+	if (!HasVisibleGizmo() && bVisible)
+	{
+		SetTrackedGizmo(Gizmo);
+	}
+}
+
+
+// Functions for listening to value changes:
+
+void FTransformGizmoDataBinder::BeginDataEditSequence()
+{
+	bInDataEditSequence = true;
+
+	// Used for proporitionally changing all results during a drag. Only relevant to scaling if
+	// we're only allowed to uniformly scale, but too messy to special case that.
+	ProportionalDragInitialVector = *BoundScale;
+
+	if (CurrentlyTrackedGizmo.IsValid() && CurrentlyTrackedGizmo->ActiveTarget)
+	{
+		TGuardValue<bool> DebounceGuard(bIgnoreCallbackForDebouncing, true);
+		CurrentlyTrackedGizmo->BeginTransformEditSequence();
+	}
+}
+
+void FTransformGizmoDataBinder::UpdateAfterDataEdit()
+{
+	if (bGizmoIsBeingDragged)
+	{
+		// Generally it's unsafe to be trying to make changes while the gizmo is being dragged. As an example,
+		// note that in a UI, if we type something and immediately start dragging the gizmo, the update from
+		// the typing focus loss may actually come after the gizmo has already issued a BeginTransformEdit call.
+		// This will cause problems since we would end up trying to nest our own Begin/EndTransformEditSequence
+		// calls inside the gizmo's.
+		return;
+	}
+
+	// See if everything actually stayed the same. This is actually an important early-out because we don't want
+	// to emit undo/redo in this situation, and it arises in UI when a user click to edit a field but then doesn't.
+	bool bScaleChanged = BoundScale && !BoundScale->Equals(LastTranslation, KINDA_SMALL_NUMBER);
+	if ((!BoundTranslation || BoundTranslation->Equals(LastTranslation, KINDA_SMALL_NUMBER))
+		&& (!BoundEulerAngles || BoundEulerAngles->Equals(LastEulerAngles, KINDA_SMALL_NUMBER))
+		&& !bScaleChanged)
+	{
+		return;
+	}
+
+	// Apply proportionality to the scale, if relevant
+	if (bScaleChanged && (bCurrentGizmoOnlyHasUniformScale ||
+		(bEnforceUniformScaleConstraintsIfPresent && CurrentlyTrackedGizmo.IsValid()
+			&& !CurrentlyTrackedGizmo->IsNonUniformScaleAllowed())))
+	{
+		FVector3d ReferenceVector = bInDataEditSequence ? ProportionalDragInitialVector : *BoundScale;
+		FVector3d Difference = *BoundScale - LastScale;
+		int32 DifferenceMaxElementIndex = UE::Geometry::MaxAbsElementIndex(Difference);
+		double ReferenceValue = ReferenceVector[DifferenceMaxElementIndex];
+
+		if (ReferenceValue == 0)
+		{
+			ReferenceValue = 1;
+		}
+
+		*BoundScale = ReferenceVector * ((*BoundScale)[DifferenceMaxElementIndex] / ReferenceValue);
+	}
+
+	LastTranslation = *BoundTranslation;
+	LastEulerAngles = *BoundEulerAngles;
+	LastScale = *BoundScale;
+
+	// Apply the bound values to the gizmo
+	if (CurrentlyTrackedGizmo.IsValid() && CurrentlyTrackedGizmo->ActiveTarget)
+	{
+		TGuardValue<bool> DebounceGuard(bIgnoreCallbackForDebouncing, true);
+
+		if (!bInDataEditSequence && bTriggerSequenceBookendsForNonSequenceUpdates)
+		{
+			CurrentlyTrackedGizmo->BeginTransformEditSequence();
+
+			CurrentlyTrackedGizmo->UpdateTransformDuringEditSequence(
+				GetGizmoTransformFromDataValues(CurrentlyTrackedGizmo->CurrentCoordinateSystem == EToolContextCoordinateSystem::Local));
+
+			CurrentlyTrackedGizmo->EndTransformEditSequence();
+			UpdateDataValuesFromGizmo();
+		}
+		else
+		{
+			CurrentlyTrackedGizmo->UpdateTransformDuringEditSequence(
+				GetGizmoTransformFromDataValues(CurrentlyTrackedGizmo->CurrentCoordinateSystem == EToolContextCoordinateSystem::Local));
+		}
+	}
+}
+
+void FTransformGizmoDataBinder::EndDataEditSequence()
+{
+	bInDataEditSequence = false;
+	if (CurrentlyTrackedGizmo.IsValid() && CurrentlyTrackedGizmo->ActiveTarget)
+	{
+		TGuardValue<bool> DebounceGuard(bIgnoreCallbackForDebouncing, true);
+		CurrentlyTrackedGizmo->EndTransformEditSequence();
+
+		// This is necesarry in case tools modify the gizmo transform after EndTransformEditSequence, for instance
+		// PolyEd resetting the scale after each drag...
+		UpdateDataValuesFromGizmo();
+	}
+}
+
+
+
+void FTransformGizmoDataBinder::BindToInitializedGizmo(UCombinedTransformGizmo* Gizmo, UTransformProxy*)
+{
+	if (!ensure(Gizmo && Gizmo->ActiveTarget && !BoundGizmos.Contains(Gizmo)))
+	{
+		return;
+	}
+
+	Gizmo->ActiveTarget->OnBeginTransformEdit.AddSP(this, &FTransformGizmoDataBinder::OnProxyBeginTransformEdit, Gizmo);
+	Gizmo->ActiveTarget->OnTransformChanged.AddSP(this, &FTransformGizmoDataBinder::OnProxyTransformChanged, Gizmo);
+	Gizmo->ActiveTarget->OnEndTransformEdit.AddSP(this, &FTransformGizmoDataBinder::OnProxyEndTransformEdit, Gizmo);
+
+	Gizmo->ActiveTarget->OnBeginPivotEdit.AddSP(this, &FTransformGizmoDataBinder::OnProxyBeginTransformEdit, Gizmo);
+	Gizmo->ActiveTarget->OnPivotChanged.AddSP(this, &FTransformGizmoDataBinder::OnProxyTransformChanged, Gizmo);
+	Gizmo->ActiveTarget->OnEndPivotEdit.AddSP(this, &FTransformGizmoDataBinder::OnProxyEndTransformEdit, Gizmo);
+
+	Gizmo->OnDisplaySpaceTransformChanged.AddSP(this, &FTransformGizmoDataBinder::OnDisplaySpaceTransformChanged);
+	Gizmo->OnVisibilityChanged.AddSP(this, &FTransformGizmoDataBinder::OnGizmoVisibilityChanged);
+
+	BoundGizmos.Add(Gizmo);
+
+	if (!HasVisibleGizmo())
+	{
+		SetTrackedGizmo(Gizmo);
+	}
+}
+
+void FTransformGizmoDataBinder::UnbindFromGizmo(UCombinedTransformGizmo* Gizmo, UTransformProxy*)
+{
+	if (!Gizmo)
+	{
+		return;
+	}
+
+	// Note about this ensure: it might seem ok to unbind once the ActiveTarget is no longer around,
+	// but it can be dangerous in case it is reused elsewhere instead of being immediately destroyed.
+	if (ensure(Gizmo->ActiveTarget))
+	{
+		Gizmo->ActiveTarget->OnBeginTransformEdit.RemoveAll(this);
+		Gizmo->ActiveTarget->OnTransformChanged.RemoveAll(this);
+		Gizmo->ActiveTarget->OnEndTransformEdit.RemoveAll(this);
+
+		Gizmo->ActiveTarget->OnBeginPivotEdit.RemoveAll(this);
+		Gizmo->ActiveTarget->OnPivotChanged.RemoveAll(this);
+		Gizmo->ActiveTarget->OnEndPivotEdit.RemoveAll(this);
+	}
+
+	Gizmo->OnDisplaySpaceTransformChanged.RemoveAll(this);
+	Gizmo->OnVisibilityChanged.RemoveAll(this);
+
+	BoundGizmos.Remove(Gizmo);
+
+	if (CurrentlyTrackedGizmo == Gizmo)
+	{
+		CurrentlyTrackedGizmo.Reset();
+		CurrentCustomLocalReferenceTransform.Reset();
+	}
+}
+
+void FTransformGizmoDataBinder::BindToGizmoContextObject(UCombinedTransformGizmoContextObject* ContextObject)
+{
+	if (!ensure(IsValid(ContextObject)))
+	{
+		return;
+	}
+
+	ContextObject->OnGizmoCreated.AddSP(this, &FTransformGizmoDataBinder::BindToUninitializedGizmo);
+
+	ContextObjectsToUnregisterWith.Add(ContextObject);
+}
+
+void FTransformGizmoDataBinder::BindToUninitializedGizmo(UCombinedTransformGizmo* Gizmo)
+{
+	Gizmo->OnSetActiveTarget.AddSP(this, &FTransformGizmoDataBinder::BindToInitializedGizmo);
+	Gizmo->OnAboutToClearActiveTarget.AddSP(this, &FTransformGizmoDataBinder::UnbindFromGizmo);
+}
+
+FTransformGizmoDataBinder::~FTransformGizmoDataBinder()
+{
+	Reset();
+}
+
+void FTransformGizmoDataBinder::Reset()
+{
+	// Deinitialize while walking through a copied list of gizmos so we don't try
+	// to delete from under ourselves.
+	TArray<TWeakObjectPtr<UCombinedTransformGizmo>> Gizmos = BoundGizmos.Array();
+	for (TWeakObjectPtr<UCombinedTransformGizmo> Gizmo : Gizmos)
+	{
+		if (!Gizmo.IsValid())
+		{
+			continue;
+		}
+
+		UnbindFromGizmo(Gizmo.Get(), Gizmo->ActiveTarget);
+	}
+	BoundGizmos.Empty();
+
+	CurrentlyTrackedGizmo.Reset();
+	CurrentCustomLocalReferenceTransform.Reset();
+	DefaultCustomLocalReferenceTransform.Reset();
+
+	BoundTranslation = VectorsToUseIfUnbound;
+	BoundEulerAngles = VectorsToUseIfUnbound + 1;
+	BoundScale = VectorsToUseIfUnbound + 2;
+
+	for (TWeakObjectPtr<UCombinedTransformGizmoContextObject>& ContextObject : ContextObjectsToUnregisterWith)
+	{
+		if (ContextObject.IsValid())
+		{
+			ContextObject->OnGizmoCreated.RemoveAll(this);
+		}
+	}
+	ContextObjectsToUnregisterWith.Empty();
+
+	OnTrackedGizmoChanged.Clear();
+}
+
+void FTransformGizmoDataBinder::InitializeBoundVectors(FVector3d* Translation, FVector3d* RotationEulerAngles, FVector3d* Scale)
+{
+	if (Translation)
+	{
+		BoundTranslation = Translation;
+	}
+	if (RotationEulerAngles)
+	{
+		BoundEulerAngles = RotationEulerAngles;
+	}
+	if (Scale)
+	{
+		BoundScale = Scale;
+	}
+}
+
+void FTransformGizmoDataBinder::UpdateCoordinateSystem()
+{
+	// When the coordinate system switches, we need to reset the delta mode because otherwise we will try to apply
+	// the existing delta in a new direction from the old reference point, and the gizmo will jump.
+	if (CurrentlyTrackedGizmo.IsValid() && CurrentlyTrackedGizmo->CurrentCoordinateSystem != LastCoordinateSystem)
+	{
+		if (bUsingDeltaMode)
+		{
+			ResetToDeltaMode();
+		}
+
+		LastCoordinateSystem = CurrentlyTrackedGizmo->CurrentCoordinateSystem;
+		UpdateDataValuesFromGizmo();
+	}
+}
+
+#undef LOCTEXT_NAMESPACE
\ No newline at end of file
diff --git a/Engine/Plugins/Runtime/MeshModelingToolset/Source/ModelingComponents/Public/Transforms/TransformGizmoDataBinder.h b/Engine/Plugins/Runtime/MeshModelingToolset/Source/ModelingComponents/Public/Transforms/TransformGizmoDataBinder.h
new file mode 100644
--- /dev/null
+++ b/Engine/Plugins/Runtime/MeshModelingToolset/Source/ModelingComponents/Public/Transforms/TransformGizmoDataBinder.h
@@ -0,0 +1,251 @@
+// Copyright Epic Games, Inc. All Rights Reserved.
+
+#pragma once
+
+#include "CoreMinimal.h"
+
+#include "BaseGizmos/TransformGizmoUtil.h" // UCombinedTransformGizmoContextObject::IGizmoBinder
+
+class UCombinedTransformGizmo;
+class UTransformProxy;
+
+/**
+ * Helper class for binding a UCombinedTransformGizmo to several FVector3d's such that they
+ * are updated when the gizmo is changed, and the gizmo is updated from the vectors if the
+ * approriate function is called. The vectors can serve as a display of the gizmo's values,
+ * and if they are made modifiable, they can serve as a numerical control of the gizmo.
+ * 
+ * To use, just call InitializeBoundVectors with 3 vectors (it can be fewer, but the binder
+ * currently does the same amount of work regardless of whether 3 or fewer are bound), and
+ * then either call BindToGizmoContextObject with the context object that will be used to
+ * create the gizmos (in which case the vectors will bind to all newly created gizmos, and 
+ * will switch bindings as different gizmos are dragged), or just bind to a particular gizmo
+ * using BindToInitializedGizmo after its SetActiveTarget call has been made..
+ */
+class MODELINGCOMPONENTS_API FTransformGizmoDataBinder : public TSharedFromThis<FTransformGizmoDataBinder, ESPMode::ThreadSafe>
+{
+public:
+
+	virtual ~FTransformGizmoDataBinder();
+
+	/**
+	 * Set the vectors that should reflect the gizmo transform. Any of these pointers can be null,
+	 * but currently the same amount of work is done by the binder (it will simply bind to some
+	 * internal vectors).
+	 */
+	virtual void InitializeBoundVectors(FVector3d* Translation, FVector3d* RotaionEulerAngles, FVector3d* Scale);
+	
+	/**
+	 * Makes it so that the gizmo binder attaches to any gizmos created by the context object
+	 * in the future. The binding is automatically removed if FTransformGizmoDataBinder is
+	 * destroyed.
+	 */
+	virtual void BindToGizmoContextObject(UCombinedTransformGizmoContextObject* ContextObject);
+
+	/**
+	 * Binds to a specific gizmo for tracking. Requires ActiveTarget to be set, and is
+	 * done automatically for gizmos when binding via the gizmo context object.
+	 */
+	virtual void BindToInitializedGizmo(UCombinedTransformGizmo* Gizmo, UTransformProxy*);
+	/**
+	 * Unbinds from a given gizmo. Done automatically for gizmos when their ActiveTarget
+	 * is cleared when binding via the gizmo context object.
+	 */
+	virtual void UnbindFromGizmo(UCombinedTransformGizmo* Gizmo, UTransformProxy*);
+
+	// These are used to push updates from the current value of vectors to the gizmo:
+
+	/**
+	 * Should be called at the beginning of a sequence of UpdateAfterDataEdit calls, for instance
+	 * if the bound values are being scrubbed in a display.
+	 */
+	void BeginDataEditSequence();
+	/**
+	 * Called to update the gizmo after updating the bound data. This acts slightly differently
+	 * depending on whether we're between begin/end sequence calls (for instance it will trigger
+	 * individual undo transactions if not between Begin/End calls).
+	 */
+	void UpdateAfterDataEdit();
+	/**
+	 * Should be called at the end of a sequence of UpdateAfterDataEdit calls, for instance
+	 * if the bound values are being scrubbed in a display.
+	 */
+	void EndDataEditSequence();
+
+	/**
+	 * Can either be called on tick or in response to some broadcast to allow the binder
+	 * to respond to coordinate system changes. If not called, the coordinate system will
+	 * still be updated the next time that the gizmo is modified.
+	 */
+	virtual void UpdateCoordinateSystem();
+
+	/**
+	 * Determines whether there is a visible currently tracked gizmo.
+	 */
+	virtual bool HasVisibleGizmo();
+
+	/**
+	 * In delta mode, the bound values hold the difference of the current gizmo transform relative
+	 * to some previous start transform. When false, the binder is in "destination mode", where the
+	 * bound values represent the final gizmo transform.
+	 */
+	virtual bool IsInDeltaMode() { return bUsingDeltaMode; }
+
+	/**
+	 * Sets to delta mode with given start transform. This call will update the bound values
+	 * appropriately from the current gizmo transform.
+	 */
+	virtual void SetToDeltaMode(const FTransform& StartTransform);
+
+	/**
+	 * Much like SetToDeltaMode, but uses the current gizmo transform as the start transform, meaning
+	 * that the bound data values will be reset to zero.
+	 */
+	virtual void ResetToDeltaMode();
+
+	/**
+	 * Sets to destination mode (as opposed to delta mode). Calling this will update the bound values
+	 * to reflect the current tracked gizmo transform. If the gizmo is in local coordinate mode and it
+	 * has a custom reference transform, that will be used as the reference.
+	 */
+	virtual void SetToDestinationMode();
+
+	/**
+	 * If the gizmo doesn't have all three axes, then it likely needs to stay constrained in some
+	 * plane/line. In this situation, allowing destination mode is only safe if the gizmo has a 
+	 * DisplaySpaceTransform that (presumably) orients the measurement axes correctly, or if the
+	 * coordinate mode is global, in which case the world axes are known to be fine.
+	 * 
+	 * Currently, the binder will automatically switch to delta mode when tracking a new gizmo
+	 * that doesn't satisfy the above definition of safety, though it won't prevent the mode from
+	 * being switched from there...
+	 */
+	virtual bool ShouldDestinationModeBeAllowed();
+
+	/**
+	 * Changes which gizmo is being tracked, out of those to which the data binder has been
+	 * bound to via BindToInitializedGizmo().
+	 *
+	 * Note that if the same data binder is bound to multiple gizmos, then it will switch tracked gizmos
+	 * automatically whenever a gizmo is dragged, or whenever a gizmo becomes visible without the previously
+	 * tracked gizmo being visible.
+	 */
+	virtual void SetTrackedGizmo(UCombinedTransformGizmo* Gizmo);
+
+	/**
+	 * Allows the data bound to the translation component to be stored in different units than
+	 * used by the gizmo. For instance, in a UV editor where the world coordinates are not actually 0 to 1.
+	 */
+	virtual void SetTranslationConversionFunctions(
+		TFunction<FVector3d(const FVector3d& InternalValue)> ActualToBoundConversionIn,
+		TFunction<FVector3d(const FVector3d& DisplayValue)> BoundToActualConversionIn)
+	{
+		ActualToBoundConversion = ActualToBoundConversionIn;
+		BoundToActualConversion = BoundToActualConversionIn;
+	}
+
+	/**
+	 * In local coordinate destination mode, when a DisplaySpaceTransform is not available on
+	 * the gizmo, this default transform can be used instead. Currently, this is used to allow
+	 * destination mode for the 2D gizmos in UV editor, but that's a fairly specific use case.
+	 */
+	void SetDefaultLocalReferenceTransform(TOptional<FTransform> DefaultTransform)
+	{
+		DefaultCustomLocalReferenceTransform = DefaultTransform;
+	}
+
+	/**
+	 * Unbinds from everything
+	 */
+	virtual void Reset();
+
+	/**
+	 * Broadcast whenever the tracked gizmo is changed (because one of multiple gizmos is
+	 * dragged, or a new gizmo appears after no others are present).
+	 */
+	DECLARE_MULTICAST_DELEGATE_OneParam(FOnTrackedGizmoChanged, UCombinedTransformGizmo*);
+	FOnTrackedGizmoChanged OnTrackedGizmoChanged;
+
+
+protected:
+	// Gizmo to which the values are currently bound.
+	TWeakObjectPtr<UCombinedTransformGizmo> CurrentlyTrackedGizmo;
+
+	// All the gizmos the numerical UI can correspond to. When one of them is dragged,
+	// CurrentlyTrackedGizmo is updated.
+	TSet<TWeakObjectPtr<UCombinedTransformGizmo>> BoundGizmos;
+
+	FVector3d VectorsToUseIfUnbound[3]{};
+
+	FVector3d* BoundTranslation = VectorsToUseIfUnbound;
+	FVector3d* BoundEulerAngles = VectorsToUseIfUnbound + 1;
+	FVector3d* BoundScale = VectorsToUseIfUnbound + 2;
+
+	// Used to detect whether anything actually changed.
+	FVector3d LastTranslation;
+	FVector3d LastEulerAngles;
+	FVector3d LastScale;
+
+	FTransform GetGizmoTransformFromDataValues(bool bLocalCoordinates);
+	void UpdateDataValuesFromGizmo();
+
+	// Controls delta mode, where the gizmo specifies a difference from some start transform.
+	bool bUsingDeltaMode = false;
+	FTransform DeltaStartTransform;
+
+	// Used for debouncing when editing the gizmo via the UI
+	bool bIgnoreCallbackForDebouncing = false;
+
+	// Additional state that lets us respond appropriately to various notifications
+	bool bInDataEditSequence = false;
+	bool bGizmoIsBeingDragged = false;
+
+	// If a gizmo doesn't have all three axes and is not aligned to world (i.e. is in local mode) then 
+	// we currently require a reference transform to allow the user to be in destination mode.
+	bool bCurrentGizmoLacksDegreeOfFreedom = false;
+
+	bool bCurrentGizmoOnlyHasUniformScale = false;
+
+	// Reference transform to use when in local gizmo mode. If it is not specified, the
+	// global reference transform is used.
+	TOptional<FTransform> CurrentCustomLocalReferenceTransform;
+
+	// Reference transform to use in local gizmo mode by default.
+	TOptional<FTransform> DefaultCustomLocalReferenceTransform;
+
+	// See comment for SetTranslationConversionFunctions
+	TFunction<FVector3d(const FVector3d& UIValues)> ActualToBoundConversion = nullptr;
+	TFunction<FVector3d(const FVector3d& UIValues)> BoundToActualConversion = nullptr;
+
+	// Used when we only allow uniform scaling, to scale everything proportionally
+	FVector3d ProportionalDragInitialVector;
+
+	// Used to determine whether the current gizmo's coordinate system has changed,
+	// when we receive an UpdateCoordinateSystem() call.
+	EToolContextCoordinateSystem LastCoordinateSystem;
+
+	// Used in binding to gizmos. Currently these need to be member variables
+	// to allow the use of AddSP, as there is not an AddLambda overload that allows
+	// for subsequent RemoveAll() calls (because AddWeakLambda only works with UObjects).
+	void OnProxyBeginTransformEdit(UTransformProxy* TransformProxy, UCombinedTransformGizmo* Gizmo);
+	void OnProxyTransformChanged(UTransformProxy* TransformProxy, FTransform Transform, UCombinedTransformGizmo* Gizmo);
+	void OnProxyEndTransformEdit(UTransformProxy* TransformProxy, UCombinedTransformGizmo* Gizmo);
+	void OnDisplaySpaceTransformChanged(UCombinedTransformGizmo* Gizmo, TOptional<FTransform> Transform);
+	void OnGizmoVisibilityChanged(UCombinedTransformGizmo* Gizmo, bool bVisible);
+
+	// Used for binding to context object, called from the OnGizmoCreated delegate.
+	void BindToUninitializedGizmo(UCombinedTransformGizmo* Gizmo);
+
+	// Used to clean up our registration with context object when being destroyed
+	TArray<TWeakObjectPtr<UCombinedTransformGizmoContextObject>> ContextObjectsToUnregisterWith;
+
+private:
+
+	// Currently we don't expose any of these because it is not clear whether we want them to be
+	// customizable. However, if we decide that we want some of our behavior choices to be configurable,
+	// we may expose them.
+	bool bTriggerSequenceBookendsForNonSequenceUpdates = true;
+	bool bEnforceUniformScaleConstraintsIfPresent = true;
+	bool bChangeDisplayedGizmoOnDrag = true;
+	bool bAvoidDestinationModeWhenUnsafe = true;
+};
diff --git a/Engine/Source/Runtime/InteractiveToolsFramework/Private/BaseGizmos/CombinedTransformGizmo.cpp b/Engine/Source/Runtime/InteractiveToolsFramework/Private/BaseGizmos/CombinedTransformGizmo.cpp
--- a/Engine/Source/Runtime/InteractiveToolsFramework/Private/BaseGizmos/CombinedTransformGizmo.cpp
+++ b/Engine/Source/Runtime/InteractiveToolsFramework/Private/BaseGizmos/CombinedTransformGizmo.cpp
@@ -32,6 +32,40 @@
 
 #define LOCTEXT_NAMESPACE "UCombinedTransformGizmo"
 
+namespace CombinedTransformGizmoLocals
+{
+	// Looks at a gizmo actor and figures out what sub element flags must have been active when creating it.
+	ETransformGizmoSubElements GetSubElementFlagsFromActor(ACombinedTransformGizmoActor* GizmoActor)
+	{
+		ETransformGizmoSubElements Elements = ETransformGizmoSubElements::None;
+		if (!GizmoActor)
+		{
+			return Elements;
+		}
+
+		if (GizmoActor->TranslateX) { Elements |= ETransformGizmoSubElements::TranslateAxisX; }
+		if (GizmoActor->TranslateY) { Elements |= ETransformGizmoSubElements::TranslateAxisY; }
+		if (GizmoActor->TranslateZ) { Elements |= ETransformGizmoSubElements::TranslateAxisZ; }
+		if (GizmoActor->TranslateXY) { Elements |= ETransformGizmoSubElements::TranslatePlaneXY; }
+		if (GizmoActor->TranslateYZ) { Elements |= ETransformGizmoSubElements::TranslatePlaneYZ; }
+		if (GizmoActor->TranslateXZ) { Elements |= ETransformGizmoSubElements::TranslatePlaneXZ; }
+
+		if (GizmoActor->RotateX) { Elements |= ETransformGizmoSubElements::RotateAxisX; }
+		if (GizmoActor->RotateY) { Elements |= ETransformGizmoSubElements::RotateAxisY; }
+		if (GizmoActor->RotateZ) { Elements |= ETransformGizmoSubElements::RotateAxisZ; }
+
+		if (GizmoActor->AxisScaleX) { Elements |= ETransformGizmoSubElements::ScaleAxisX; }
+		if (GizmoActor->AxisScaleY) { Elements |= ETransformGizmoSubElements::ScaleAxisY; }
+		if (GizmoActor->AxisScaleZ) { Elements |= ETransformGizmoSubElements::ScaleAxisZ; }
+		if (GizmoActor->PlaneScaleXY) { Elements |= ETransformGizmoSubElements::ScalePlaneXY; }
+		if (GizmoActor->PlaneScaleYZ) { Elements |= ETransformGizmoSubElements::ScalePlaneYZ; }
+		if (GizmoActor->PlaneScaleXZ) { Elements |= ETransformGizmoSubElements::ScalePlaneXZ; }
+
+		if (GizmoActor->UniformScale) { Elements |= ETransformGizmoSubElements::ScaleUniform; }
+
+		return Elements;
+	}
+}
 
 ACombinedTransformGizmoActor::ACombinedTransformGizmoActor()
 {
@@ -367,7 +401,7 @@ void UCombinedTransformGizmo::SetDisallowNegativeScaling(bool bDisallow)
 
 void UCombinedTransformGizmo::SetIsNonUniformScaleAllowedFunction(TUniqueFunction<bool()>&& IsNonUniformScaleAllowedIn)
 {
-	IsNonUniformScaleAllowed = MoveTemp(IsNonUniformScaleAllowedIn);
+	IsNonUniformScaleAllowedFunc = MoveTemp(IsNonUniformScaleAllowedIn);
 }
 
 
@@ -476,7 +510,7 @@ void UCombinedTransformGizmo::Tick(float DeltaTime)
 		(UseGizmoMode == EToolContextTransformGizmoMode::Combined || UseGizmoMode == EToolContextTransformGizmoMode::Scale);
 	bool bShouldShowNonUniformScale = 
 		(UseGizmoMode == EToolContextTransformGizmoMode::Combined || UseGizmoMode == EToolContextTransformGizmoMode::Scale)
-		&& IsNonUniformScaleAllowed();
+		&& IsNonUniformScaleAllowedFunc();
 
 	SetSubGizmoTypeVisibility(TranslationSubGizmos, bShouldShowTranslation);
 	SetSubGizmoTypeVisibility(RotationSubGizmos, bShouldShowRotation);
@@ -643,6 +677,8 @@ void UCombinedTransformGizmo::SetActiveTarget(UTransformProxy* Target, IToolCont
 		ActiveComponents.Add(GizmoActor->PlaneScaleXY);
 		NonUniformScaleSubGizmos.Add(FSubGizmoInfo{ GizmoActor->PlaneScaleXY, NewGizmo });
 	}
+
+	OnSetActiveTarget.Broadcast(this, ActiveTarget);
 }
 
 FTransform UCombinedTransformGizmo::GetGizmoTransform() const
@@ -678,7 +714,30 @@ void UCombinedTransformGizmo::SetNewGizmoTransform(const FTransform& NewTransfor
 {
 	check(ActiveTarget != nullptr);
 
-	StateTarget->BeginUpdate();
+	BeginTransformEditSequence();
+	UpdateTransformDuringEditSequence(NewTransform, bKeepGizmoUnscaled);
+	EndTransformEditSequence();
+}
+
+void UCombinedTransformGizmo::BeginTransformEditSequence()
+{
+	if (ensure(StateTarget))
+	{
+		StateTarget->BeginUpdate();
+	}
+}
+
+void UCombinedTransformGizmo::EndTransformEditSequence()
+{
+	if (ensure(StateTarget))
+	{
+		StateTarget->EndUpdate();
+	}
+}
+
+void UCombinedTransformGizmo::UpdateTransformDuringEditSequence(const FTransform& NewTransform, bool bKeepGizmoUnscaled)
+{
+	check(ActiveTarget != nullptr);
 
 	USceneComponent* GizmoComponent = GizmoActor->GetRootComponent();
 	FTransform GizmoTransform = NewTransform;
@@ -688,11 +747,8 @@ void UCombinedTransformGizmo::SetNewGizmoTransform(const FTransform& NewTransfor
 	}
 	GizmoComponent->SetWorldTransform(GizmoTransform);
 	ActiveTarget->SetTransform(NewTransform);
-
-	StateTarget->EndUpdate();
 }
 
-
 void UCombinedTransformGizmo::SetNewChildScale(const FVector& NewChildScale)
 {
 	FTransform NewTransform = ActiveTarget->GetTransform();
@@ -707,12 +763,35 @@ void UCombinedTransformGizmo::SetNewChildScale(const FVector& NewChildScale)
 
 void UCombinedTransformGizmo::SetVisibility(bool bVisible)
 {
+	bool bPreviousVisibility = !GizmoActor->IsHidden();
+
 	GizmoActor->SetActorHiddenInGame(bVisible == false);
 #if WITH_EDITOR
 	GizmoActor->SetIsTemporarilyHiddenInEditor(bVisible == false);
 #endif
+
+	if (bPreviousVisibility != bVisible)
+	{
+		OnVisibilityChanged.Broadcast(this, bVisible);
+	}
+}
+
+void UCombinedTransformGizmo::SetDisplaySpaceTransform(TOptional<FTransform> TransformIn)
+{
+	if (DisplaySpaceTransform.IsSet() != TransformIn.IsSet()
+		|| (TransformIn.IsSet() && !TransformIn.GetValue().Equals(DisplaySpaceTransform.GetValue())))
+	{
+		DisplaySpaceTransform = TransformIn;
+		OnDisplaySpaceTransformChanged.Broadcast(this, TransformIn);
+	}
 }
 
+ETransformGizmoSubElements UCombinedTransformGizmo::GetGizmoElements()
+{
+	using namespace CombinedTransformGizmoLocals;
+
+	return GetSubElementFlagsFromActor(GizmoActor);
+}
 
 UInteractiveGizmo* UCombinedTransformGizmo::AddAxisTranslationGizmo(
 	UPrimitiveComponent* AxisComponent, USceneComponent* RootComponent,
@@ -953,6 +1032,8 @@ UInteractiveGizmo* UCombinedTransformGizmo::AddUniformScaleGizmo(
 
 void UCombinedTransformGizmo::ClearActiveTarget()
 {
+	OnAboutToClearActiveTarget.Broadcast(this, ActiveTarget);
+
 	for (UInteractiveGizmo* Gizmo : ActiveGizmos)
 	{
 		GetGizmoManager()->DestroyGizmo(Gizmo);
diff --git a/Engine/Source/Runtime/InteractiveToolsFramework/Private/BaseGizmos/TransformGizmoUtil.cpp b/Engine/Source/Runtime/InteractiveToolsFramework/Private/BaseGizmos/TransformGizmoUtil.cpp
--- a/Engine/Source/Runtime/InteractiveToolsFramework/Private/BaseGizmos/TransformGizmoUtil.cpp
+++ b/Engine/Source/Runtime/InteractiveToolsFramework/Private/BaseGizmos/TransformGizmoUtil.cpp
@@ -141,8 +141,12 @@ UCombinedTransformGizmo* UCombinedTransformGizmoContextObject::Create3AxisTransf
 	if (ensure(bDefaultGizmosRegistered))
 	{
 		UInteractiveGizmo* NewGizmo = GizmoManager->CreateGizmo(DefaultThreeAxisTransformBuilderIdentifier, InstanceIdentifier, Owner);
-		ensure(NewGizmo);
-		return Cast<UCombinedTransformGizmo>(NewGizmo);
+		UCombinedTransformGizmo* CastGizmo = Cast<UCombinedTransformGizmo>(NewGizmo);
+		if (ensure(CastGizmo))
+		{
+			OnGizmoCreated.Broadcast(CastGizmo);
+		}
+		return CastGizmo;
 	}
 	return nullptr;
 }
@@ -173,8 +177,12 @@ UCombinedTransformGizmo* UCombinedTransformGizmoContextObject::CreateCustomTrans
 	{
 		GizmoActorBuilder->EnableElements = Elements;
 		UInteractiveGizmo* NewGizmo = GizmoManager->CreateGizmo(CustomThreeAxisTransformBuilderIdentifier, InstanceIdentifier, Owner);
-		ensure(NewGizmo);
-		return Cast<UCombinedTransformGizmo>(NewGizmo);
+		UCombinedTransformGizmo* CastGizmo = Cast<UCombinedTransformGizmo>(NewGizmo);
+		if (ensure(CastGizmo))
+		{
+			OnGizmoCreated.Broadcast(CastGizmo);
+		}
+		return CastGizmo;
 	}
 	return nullptr;
 }
@@ -207,7 +215,12 @@ UCombinedTransformGizmo* UCombinedTransformGizmoContextObject::CreateCustomRepos
 		GizmoActorBuilder->EnableElements = Elements;
 		UInteractiveGizmo* NewGizmo = GizmoManager->CreateGizmo(CustomRepositionableThreeAxisTransformBuilderIdentifier, InstanceIdentifier, Owner);
 		ensure(NewGizmo);
-		return Cast<UCombinedTransformGizmo>(NewGizmo);
+		UCombinedTransformGizmo* CastGizmo = Cast<UCombinedTransformGizmo>(NewGizmo);
+		if (ensure(CastGizmo))
+		{
+			OnGizmoCreated.Broadcast(CastGizmo);
+		}
+		return CastGizmo;
 	}
 	return nullptr;
 }
diff --git a/Engine/Source/Runtime/InteractiveToolsFramework/Public/BaseGizmos/CombinedTransformGizmo.h b/Engine/Source/Runtime/InteractiveToolsFramework/Public/BaseGizmos/CombinedTransformGizmo.h
--- a/Engine/Source/Runtime/InteractiveToolsFramework/Public/BaseGizmos/CombinedTransformGizmo.h
+++ b/Engine/Source/Runtime/InteractiveToolsFramework/Public/BaseGizmos/CombinedTransformGizmo.h
@@ -324,6 +324,12 @@ public:
 		TUniqueFunction<bool()>&& IsNonUniformScaleAllowed
 	);
 
+	/**
+	 * Exposes the return value of the current IsNonUniformScaleAllowed function so that, for instance,
+	 * numerical UI can react appropriately.
+	 */
+	virtual bool IsNonUniformScaleAllowed() const { return IsNonUniformScaleAllowedFunc(); }
+
 	/**
 	 * By default, the nonuniform scale components can scale negatively. However, they can be made to clamp
 	 * to zero instead by passing true here. This is useful for using the gizmo to flatten geometry.
@@ -378,6 +384,22 @@ public:
 	 */
 	virtual void SetNewGizmoTransform(const FTransform& NewTransform, bool bKeepGizmoUnscaled = true);
 
+	/**
+	 * Called at the start of a sequence of gizmo transform edits, for instance while dragging or
+	 * manipulating the gizmo numerical UI.
+	 */
+	virtual void BeginTransformEditSequence();
+
+	/**
+	 * Called at the end of a sequence of gizmo transform edits.
+	 */
+	virtual void EndTransformEditSequence();
+
+	/**
+	 * Updates the gizmo transform between Begin/EndTransformeditSequence calls.
+	 */
+	void UpdateTransformDuringEditSequence(const FTransform& NewTransform, bool bKeepGizmoUnscaled = true);
+
 	/**
 	 * Explicitly set the child scale. Mainly useful to "reset" the child scale to (1,1,1) when re-using Gizmo across multiple transform actions.
 	 * @warning does not generate change/modify events!!
@@ -463,6 +485,50 @@ public:
 	UPROPERTY()
 	EToolContextTransformGizmoMode ActiveGizmoMode = EToolContextTransformGizmoMode::Combined;
 
+	/**
+	 * Gets the elements that this gizmo was initialized with. Note that this may not account for individual
+	 * element visibility- for instance the scaling component may not be visible if IsNonUniformScaleAllowed() is false.
+	 */
+	ETransformGizmoSubElements GetGizmoElements();
+
+
+	/**
+	 * The DisplaySpaceTransform is not used by the gizmo itself, but can be used by external adapters that might
+	 * display gizmo values, to give values relative to this transform rather than relative to world origin and axes.
+	 * 
+	 * For example a numerical UI for a two-axis gizmo that is not in a world XY/YZ/XZ plane cannot use the global
+	 * axes for setting the absolute position of the plane if it wants the gizmo to remain in that plane; instead, the
+	 * DisplaySpaceTransform can give a space in which X and Y values keep the gizmo in the plane.
+	 *
+	 * Note that this is an optional feature, as it would require tools to keep this transform up to date if they
+	 * want the UI to use it, so tools could just leave it unset.
+	 */
+	void SetDisplaySpaceTransform(TOptional<FTransform> TransformIn);
+	const TOptional<FTransform>& GetDisplaySpaceTransform() { return DisplaySpaceTransform; }
+
+	/**
+	 * Broadcast at the end of a SetDisplaySpaceTransform call that changes the display space transform.
+	 */
+	DECLARE_MULTICAST_DELEGATE_TwoParams(FOnDisplaySpaceTransformChanged, UCombinedTransformGizmo*, TOptional<FTransform>);
+	FOnDisplaySpaceTransformChanged OnDisplaySpaceTransformChanged;
+
+	/** 
+	 * Broadcast at the end of a SetActiveTarget call. Using this, an adapter such as a numerical UI widget can 
+	 * bind to the gizmo at construction and still be able to initialize using the transform proxy once that is set.
+	 */
+	DECLARE_MULTICAST_DELEGATE_TwoParams(FOnSetActiveTarget, UCombinedTransformGizmo*, UTransformProxy*);
+	FOnSetActiveTarget OnSetActiveTarget;
+
+	/** 
+	 * Broadcast at the beginning of a ClearActiveTarget call, when the ActiveTarget (if present) is not yet 
+	 * disconnected. Gives things a chance to unbind from it.
+	 */
+	DECLARE_MULTICAST_DELEGATE_TwoParams(FOnClearActiveTarget, UCombinedTransformGizmo*, UTransformProxy*);
+	FOnClearActiveTarget OnAboutToClearActiveTarget;
+
+	/** Broadcast at the end of a SetVisibility call if the visibility changes. */
+	DECLARE_MULTICAST_DELEGATE_TwoParams(FOnVisibilityChanged, UCombinedTransformGizmo*, bool);
+	FOnVisibilityChanged OnVisibilityChanged;
 
 protected:
 	TSharedPtr<FCombinedTransformGizmoActorFactory> GizmoActorBuilder;
@@ -574,9 +640,12 @@ protected:
 	TUniqueFunction<bool()> ShouldAlignDestination = []() { return false; };
 	TUniqueFunction<bool(const FRay&, FVector&)> DestinationAlignmentRayCaster = [](const FRay&, FVector&) {return false; };
 
-	TUniqueFunction<bool()> IsNonUniformScaleAllowed = [this]() { return CurrentCoordinateSystem == EToolContextCoordinateSystem::Local; };
+	TUniqueFunction<bool()> IsNonUniformScaleAllowedFunc = [this]() { return CurrentCoordinateSystem == EToolContextCoordinateSystem::Local; };
 
 	bool bDisallowNegativeScaling = false;
+
+	// See comment for SetDisplaySpaceTransform;
+	TOptional<FTransform> DisplaySpaceTransform;
 protected:
 
 
diff --git a/Engine/Source/Runtime/InteractiveToolsFramework/Public/BaseGizmos/TransformGizmoUtil.h b/Engine/Source/Runtime/InteractiveToolsFramework/Public/BaseGizmos/TransformGizmoUtil.h
--- a/Engine/Source/Runtime/InteractiveToolsFramework/Public/BaseGizmos/TransformGizmoUtil.h
+++ b/Engine/Source/Runtime/InteractiveToolsFramework/Public/BaseGizmos/TransformGizmoUtil.h
@@ -125,6 +125,9 @@ public:
 	 */
 	virtual UCombinedTransformGizmo* CreateCustomRepositionableTransformGizmo(UInteractiveGizmoManager* GizmoManager, ETransformGizmoSubElements Elements, void* Owner = nullptr, const FString& InstanceIdentifier = FString());
 
+	DECLARE_MULTICAST_DELEGATE_OneParam(FOnGizmoCreated, UCombinedTransformGizmo*);
+	FOnGizmoCreated OnGizmoCreated;
+
 protected:
 	TSharedPtr<FCombinedTransformGizmoActorFactory> GizmoActorBuilder;
 	bool bDefaultGizmosRegistered = false;
