ModelingTools: Fix static analysis warning about inconsistent (unnecessary) null checks.

diff --git a/Engine/Plugins/Runtime/MeshModelingToolset/Source/ModelingComponents/Private/Transforms/TransformGizmoDataBinder.cpp b/Engine/Plugins/Runtime/MeshModelingToolset/Source/ModelingComponents/Private/Transforms/TransformGizmoDataBinder.cpp
--- a/Engine/Plugins/Runtime/MeshModelingToolset/Source/ModelingComponents/Private/Transforms/TransformGizmoDataBinder.cpp
+++ b/Engine/Plugins/Runtime/MeshModelingToolset/Source/ModelingComponents/Private/Transforms/TransformGizmoDataBinder.cpp
@@ -340,10 +340,10 @@ void FTransformGizmoDataBinder::UpdateAfterDataEdit()
 
 	// See if everything actually stayed the same. This is actually an important early-out because we don't want
 	// to emit undo/redo in this situation, and it arises in UI when a user click to edit a field but then doesn't.
-	bool bScaleChanged = BoundScale && !BoundScale->Equals(LastTranslation, KINDA_SMALL_NUMBER);
-	if ((!BoundTranslation || BoundTranslation->Equals(LastTranslation, KINDA_SMALL_NUMBER))
-		&& (!BoundEulerAngles || BoundEulerAngles->Equals(LastEulerAngles, KINDA_SMALL_NUMBER))
-		&& !bScaleChanged)
+	bool bScaleChanged = !BoundScale->Equals(LastScale, KINDA_SMALL_NUMBER);
+	if (!bScaleChanged
+		&& BoundTranslation->Equals(LastTranslation, KINDA_SMALL_NUMBER)
+		&& BoundEulerAngles->Equals(LastEulerAngles, KINDA_SMALL_NUMBER))
 	{
 		return;
 	}
