MeshModelingTools: Change numerical UI display name to transform panel. Tweak how the numerical gizmo UI deals with scrubbing zero-valued scale in uniform scaling to match what editor does. Fix a typo and updating of last values.

diff --git a/Engine/Plugins/Editor/ModelingToolsEditorMode/Source/ModelingToolsEditorMode/Private/ModelingModeToolkit_Toolbars.cpp b/Engine/Plugins/Editor/ModelingToolsEditorMode/Source/ModelingToolsEditorMode/Private/ModelingModeToolkit_Toolbars.cpp
--- a/Engine/Plugins/Editor/ModelingToolsEditorMode/Source/ModelingToolsEditorMode/Private/ModelingModeToolkit_Toolbars.cpp
+++ b/Engine/Plugins/Editor/ModelingToolsEditorMode/Source/ModelingToolsEditorMode/Private/ModelingModeToolkit_Toolbars.cpp
@@ -111,7 +111,7 @@ void MakeSubMenu_GizmoVisibilityMode(FModelingToolsEditorModeToolkit* Toolkit, F
 	
 
 	MenuBuilder.AddSubMenu(
-		LOCTEXT("NumericalUISubMenu", "Numerical UI"), LOCTEXT("NumericalUISubMenu_ToolTip", "Configure the gizmo numerical UI."),
+		LOCTEXT("TransformPanelSubMenu", "Transform Panel"), LOCTEXT("TransformPanelSubMenu_ToolTip", "Configure the Gizmo Transform Panel."),
 		FNewMenuDelegate::CreateLambda([Toolkit](FMenuBuilder& SubMenuBuilder) {
 			TSharedPtr<STransformGizmoNumericalUIOverlay> NumericalUI = Toolkit->GetGizmoNumericalUIOverlayWidget();
 			if (ensure(NumericalUI.IsValid()))
diff --git a/Engine/Plugins/Experimental/MeshModelingToolsetExp/Source/ModelingEditorUI/Private/STransformGizmoNumericalUIOverlay.cpp b/Engine/Plugins/Experimental/MeshModelingToolsetExp/Source/ModelingEditorUI/Private/STransformGizmoNumericalUIOverlay.cpp
--- a/Engine/Plugins/Experimental/MeshModelingToolsetExp/Source/ModelingEditorUI/Private/STransformGizmoNumericalUIOverlay.cpp
+++ b/Engine/Plugins/Experimental/MeshModelingToolsetExp/Source/ModelingEditorUI/Private/STransformGizmoNumericalUIOverlay.cpp
@@ -493,8 +493,8 @@ void STransformGizmoNumericalUIOverlay::MakeNumericalUISubMenu(FMenuBuilder& Men
 			}
 			return NumericalUI.Pin()->IsEnabled();
 		}));
-	MenuBuilder.AddMenuEntry(LOCTEXT("NumericalUI_Enabled", "Enabled"), 
-		LOCTEXT("NumericalUI_Enabled_Tooltip", "Show the gizmo numerical UI floating panel when a gizmo is visible."),
+	MenuBuilder.AddMenuEntry(LOCTEXT("TransformPanel_Enabled", "Enabled"), 
+		LOCTEXT("TransformPanel_Enabled_Tooltip", "Show the Gizmo Transform Panel when a gizmo is visible."),
 		FSlateIcon(), NumericalUIAction, NAME_None, EUserInterfaceActionType::ToggleButton);
 
 	const FUIAction GizmoMode_ResetNumericalUIPosition(
@@ -509,8 +509,8 @@ void STransformGizmoNumericalUIOverlay::MakeNumericalUISubMenu(FMenuBuilder& Men
 		{
 			return NumericalUI.IsValid() && NumericalUI.Pin()->IsEnabled();
 		}));
-	MenuBuilder.AddMenuEntry(LOCTEXT("NumericalUI_ResetNumericalUIPosition", "Reset UI Position"),
-		LOCTEXT("NumericalUI_ResetNumericalUIPosition_Tooltip", "Reset the numerical UI position in the viewport."),
+	MenuBuilder.AddMenuEntry(LOCTEXT("TransformPanel_ResetPosition", "Reset Panel Position"),
+		LOCTEXT("TransformPanel_ResetPosition_Tooltip", "Reset the Gizmo Transform Panel position in the viewport."),
 		FSlateIcon(), GizmoMode_ResetNumericalUIPosition, NAME_None, EUserInterfaceActionType::Button);
 }
 
diff --git a/Engine/Plugins/Runtime/MeshModelingToolset/Source/ModelingComponents/Private/Transforms/TransformGizmoDataBinder.cpp b/Engine/Plugins/Runtime/MeshModelingToolset/Source/ModelingComponents/Private/Transforms/TransformGizmoDataBinder.cpp
--- a/Engine/Plugins/Runtime/MeshModelingToolset/Source/ModelingComponents/Private/Transforms/TransformGizmoDataBinder.cpp
+++ b/Engine/Plugins/Runtime/MeshModelingToolset/Source/ModelingComponents/Private/Transforms/TransformGizmoDataBinder.cpp
@@ -70,6 +70,10 @@ void FTransformGizmoDataBinder::UpdateDataValuesFromGizmo()
 	{
 		*BoundEulerAngles = CurrentRotation.Euler();
 	}
+
+	LastTranslation = *BoundTranslation;
+	LastEulerAngles = *BoundEulerAngles;
+	LastScale = *BoundScale;
 }
 
 // Key function for going from values to gizmo
@@ -186,6 +190,10 @@ void FTransformGizmoDataBinder::ResetToDeltaMode()
 	*BoundTranslation = ActualToBoundConversion ? ActualToBoundConversion(FVector::Zero()) : FVector::Zero();
 	*BoundEulerAngles = FVector::Zero();
 	*BoundScale = FVector::One();
+
+	LastTranslation = *BoundTranslation;
+	LastEulerAngles = *BoundEulerAngles;
+	LastScale = *BoundScale;
 }
 
 void FTransformGizmoDataBinder::SetToDestinationMode()
@@ -340,7 +348,9 @@ void FTransformGizmoDataBinder::UpdateAfterDataEdit()
 
 	// See if everything actually stayed the same. This is actually an important early-out because we don't want
 	// to emit undo/redo in this situation, and it arises in UI when a user click to edit a field but then doesn't.
-	bool bScaleChanged = !BoundScale->Equals(LastScale, KINDA_SMALL_NUMBER);
+
+	bool bScaleChanged = !BoundScale->Equals(LastScale, KINDA_SMALL_NUMBER); // saved since used further too
+
 	if (!bScaleChanged
 		&& BoundTranslation->Equals(LastTranslation, KINDA_SMALL_NUMBER)
 		&& BoundEulerAngles->Equals(LastEulerAngles, KINDA_SMALL_NUMBER))
@@ -349,27 +359,29 @@ void FTransformGizmoDataBinder::UpdateAfterDataEdit()
 	}
 
 	// Apply proportionality to the scale, if relevant
-	if (bScaleChanged && (bCurrentGizmoOnlyHasUniformScale ||
-		(bEnforceUniformScaleConstraintsIfPresent && CurrentlyTrackedGizmo.IsValid()
-			&& !CurrentlyTrackedGizmo->IsNonUniformScaleAllowed())))
+
+	auto ShouldEnforceUniformScale = [this]() {
+		return bEnforceUniformScaleConstraintsIfPresent && CurrentlyTrackedGizmo.IsValid()
+			&& !CurrentlyTrackedGizmo->IsNonUniformScaleAllowed();
+	};
+
+	if (bScaleChanged && (bCurrentGizmoOnlyHasUniformScale || ShouldEnforceUniformScale()))
 	{
-		FVector3d ReferenceVector = bInDataEditSequence ? ProportionalDragInitialVector : *BoundScale;
+		// Figure out which component changed.
 		FVector3d Difference = *BoundScale - LastScale;
 		int32 DifferenceMaxElementIndex = UE::Geometry::MaxAbsElementIndex(Difference);
-		double ReferenceValue = ReferenceVector[DifferenceMaxElementIndex];
 
-		if (ReferenceValue == 0)
+		FVector3d ReferenceVectorToUse = bInDataEditSequence ? ProportionalDragInitialVector : LastScale;
+		if (ReferenceVectorToUse[DifferenceMaxElementIndex] == 0)
 		{
-			ReferenceValue = 1;
+			// It's not clear how to apply proporitional scale if the scrubbed/changed value was zero, but the editor 
+			// goes the route of changing it to be 1 beforehand and scaling the result. So, we do the same.
+			ReferenceVectorToUse[DifferenceMaxElementIndex] = 1;
 		}
 
-		*BoundScale = ReferenceVector * ((*BoundScale)[DifferenceMaxElementIndex] / ReferenceValue);
+		*BoundScale = ReferenceVectorToUse * ((*BoundScale)[DifferenceMaxElementIndex] / ReferenceVectorToUse[DifferenceMaxElementIndex]);
 	}
 
-	LastTranslation = *BoundTranslation;
-	LastEulerAngles = *BoundEulerAngles;
-	LastScale = *BoundScale;
-
 	// Apply the bound values to the gizmo
 	if (CurrentlyTrackedGizmo.IsValid() && CurrentlyTrackedGizmo->ActiveTarget)
 	{
@@ -391,6 +403,10 @@ void FTransformGizmoDataBinder::UpdateAfterDataEdit()
 				GetGizmoTransformFromDataValues(CurrentlyTrackedGizmo->CurrentCoordinateSystem == EToolContextCoordinateSystem::Local));
 		}
 	}
+
+	LastTranslation = *BoundTranslation;
+	LastEulerAngles = *BoundEulerAngles;
+	LastScale = *BoundScale;
 }
 
 void FTransformGizmoDataBinder::EndDataEditSequence()
