MeshModelingTools: Fix PolyEd group edge insertion crashing if an undo transaction is received while waiting for an insertion to complete, and harden some of the state management.

diff --git a/Engine/Plugins/Runtime/MeshModelingToolset/Source/MeshModelingTools/Private/ToolActivities/PolyEditInsertEdgeActivity.cpp b/Engine/Plugins/Runtime/MeshModelingToolset/Source/MeshModelingTools/Private/ToolActivities/PolyEditInsertEdgeActivity.cpp
--- a/Engine/Plugins/Runtime/MeshModelingToolset/Source/MeshModelingTools/Private/ToolActivities/PolyEditInsertEdgeActivity.cpp
+++ b/Engine/Plugins/Runtime/MeshModelingToolset/Source/MeshModelingTools/Private/ToolActivities/PolyEditInsertEdgeActivity.cpp
@@ -96,7 +96,8 @@ void UPolyEditInsertEdgeActivity::Setup(UInteractiveTool* ParentToolIn)
 	ActivityContext->OnUndoRedo.AddWeakLambda(this, [this](bool bGroupTopologyModified)
 	{
 		UpdateComputeInputs();
-		ClearPreview();
+		ToolState = EState::GettingStart;
+		ClearPreview(true);
 	});
 }
 
@@ -275,7 +276,7 @@ void UPolyEditInsertEdgeActivity::Tick(float DeltaTime)
 		}
 		else
 		{
-			ToolState = EState::GettingEnd;
+			ToolState = PreviewPoints.Num() > 0 ? EState::GettingEnd : EState::GettingStart;
 		}
 
 		PreviewEdges.Reset();
@@ -341,6 +342,13 @@ void UPolyEditInsertEdgeActivity::ClearPreview(bool bClearDrawnElements)
 	{
 		PreviewEdges.Reset();
 		PreviewPoints.Reset();
+
+		// If we're removing the start point, we shouldn't be in a state that requires
+		// it, i.e. getting the second point.
+		if (!ensure(ToolState != EState::GettingEnd))
+		{
+			ToolState = EState::GettingStart;
+		}
 	}
 }
 
@@ -464,7 +472,15 @@ bool UPolyEditInsertEdgeActivity::OnUpdateHover(const FInputDeviceRay& DevicePos
 	}
 	case EState::GettingEnd:
 	{
-		check(PreviewPoints.Num() > 0);
+		// This shouldn't happen- if it does, we messed up with our state handling somewhere. But don't crash,
+		// just reset.
+		if (!ensure(PreviewPoints.Num() > 0))
+		{
+			ToolState = EState::GettingStart;
+			ClearPreview(true);
+			return false;
+		}
+
 		PreviewPoints.SetNum(1); // Keep the first element, which is the start point
 
 		// Don't update the end variables right away so that we can check if they actually changed (they
@@ -492,7 +508,8 @@ bool UPolyEditInsertEdgeActivity::OnUpdateHover(const FInputDeviceRay& DevicePos
 			return true;
 		}
 
-		// If we don't have a valid endpoint, draw a line to the current hit location.
+		// If we're here, then we don't have a valid endpoint. Make sure our preview is reset and draw a line
+		// to the current hit location.
 		if (!bShowingBaseMesh)
 		{
 			ClearPreview(false);
@@ -900,10 +917,11 @@ void FGroupEdgeInsertionFirstPointChange::Revert(UObject* Object)
 {
 	UPolyEditInsertEdgeActivity* Activity = Cast<UPolyEditInsertEdgeActivity>(Object);
 
-	check(Activity->ToolState == UPolyEditInsertEdgeActivity::EState::GettingEnd);
+	check(Activity->ToolState == UPolyEditInsertEdgeActivity::EState::GettingEnd
+		|| Activity->ToolState == UPolyEditInsertEdgeActivity::EState::WaitingForInsertComplete);
 	Activity->ToolState = UPolyEditInsertEdgeActivity::EState::GettingStart;
 
-	Activity->ClearPreview();
+	Activity->ClearPreview(true);
 
 	bHaveDoneUndo = true;
 }
diff --git a/Engine/Plugins/Runtime/MeshModelingToolset/Source/MeshModelingTools/Public/ToolActivities/PolyEditInsertEdgeActivity.h b/Engine/Plugins/Runtime/MeshModelingToolset/Source/MeshModelingTools/Public/ToolActivities/PolyEditInsertEdgeActivity.h
--- a/Engine/Plugins/Runtime/MeshModelingToolset/Source/MeshModelingTools/Public/ToolActivities/PolyEditInsertEdgeActivity.h
+++ b/Engine/Plugins/Runtime/MeshModelingToolset/Source/MeshModelingTools/Public/ToolActivities/PolyEditInsertEdgeActivity.h
@@ -161,7 +161,7 @@ protected:
 	void ConditionallyUpdatePreview(const FGroupEdgeInserter::FGroupEdgeSplitPoint& NewEndPoint, 
 		int32 NewEndTopologyID, bool bNewEndIsCorner, int32 NewCommonGroupID, int32 NewBoundaryIndex);
 
-	void ClearPreview(bool bClearDrawnElements = true);
+	void ClearPreview(bool bClearDrawnElements);
 
 	void GetCornerTangent(int32 CornerID, int32 GroupID, int32 BoundaryIndex, FVector3d& TangentOut);
 
@@ -191,10 +191,9 @@ public:
 	{
 		UPolyEditInsertEdgeActivity* Activity = Cast<UPolyEditInsertEdgeActivity>(Object);
 		return bHaveDoneUndo || Activity->CurrentChangeStamp != ChangeStamp
-			|| Activity->ToolState != UPolyEditInsertEdgeActivity::EState::GettingEnd;
-		// TODO: this is a bit of a hack in that we should probably have a separate stamp
-		// for expiring these instead of letting the tool state help (they expire after
-		// each new insertion unlike the other changes, which expire on tool close).
+			// We only allow undo if we're looking for the next point or waiting for completion, i.e. when
+			// we have a start point to undo.
+			|| Activity->ToolState == UPolyEditInsertEdgeActivity::EState::GettingStart;
 	}
 	virtual FString ToString() const override
 	{
