MeshModelingTools: Retarget gizmo numerical UI on different bound gizmo if current one is hidden by tool.

diff --git a/Engine/Plugins/Runtime/MeshModelingToolset/Source/ModelingComponents/Private/Transforms/TransformGizmoDataBinder.cpp b/Engine/Plugins/Runtime/MeshModelingToolset/Source/ModelingComponents/Private/Transforms/TransformGizmoDataBinder.cpp
--- a/Engine/Plugins/Runtime/MeshModelingToolset/Source/ModelingComponents/Private/Transforms/TransformGizmoDataBinder.cpp
+++ b/Engine/Plugins/Runtime/MeshModelingToolset/Source/ModelingComponents/Private/Transforms/TransformGizmoDataBinder.cpp
@@ -299,20 +299,30 @@ void FTransformGizmoDataBinder::OnDisplaySpaceTransformChanged(UCombinedTransfor
 	}
 }
 
-void FTransformGizmoDataBinder::OnGizmoVisibilityChanged(UCombinedTransformGizmo* Gizmo, bool bVisible)
+void FTransformGizmoDataBinder::OnGizmoVisibilityChanged(UCombinedTransformGizmo* GizmoIn, bool bVisible)
 {
-	// Generally we don't need to worry about visibility changes for the gizmo that we are bound to,
-	// since it's one of the things we poll for in the UI panel visibility. However, for non-bound
-	// gizmos, we do need to listen in case we don't have a currently displayed gizmo, and therefore
-	// should switch to showing a newly visible different gizmo.
-	if (!ensure(Gizmo) || CurrentlyTrackedGizmo == Gizmo)
+	if (!ensure(GizmoIn))
 	{
 		return;
 	}
 
+	// If we don't have a gizmo to track, and this one is becoming visible, track it.
 	if (!HasVisibleGizmo() && bVisible)
 	{
-		SetTrackedGizmo(Gizmo);
+		SetTrackedGizmo(GizmoIn);
+	}
+	// If we currently track this gizmo, and it's becoming invisible, see if there's
+	// another one we can track instead.
+	else if (CurrentlyTrackedGizmo == GizmoIn && !bVisible)
+	{
+		for (TWeakObjectPtr<UCombinedTransformGizmo> Gizmo : BoundGizmos)
+		{
+			if (Gizmo.IsValid() && Gizmo->ActiveTarget && Gizmo->IsVisible())
+			{
+				SetTrackedGizmo(Gizmo.Get());
+				break;
+			}
+		}
 	}
 }
 
