MeshModelingTools: Fix undo in PolyEd group edge "continuous insertion" mode.

diff --git a/Engine/Plugins/Runtime/MeshModelingToolset/Source/MeshModelingTools/Private/ToolActivities/PolyEditInsertEdgeActivity.cpp b/Engine/Plugins/Runtime/MeshModelingToolset/Source/MeshModelingTools/Private/ToolActivities/PolyEditInsertEdgeActivity.cpp
--- a/Engine/Plugins/Runtime/MeshModelingToolset/Source/MeshModelingTools/Private/ToolActivities/PolyEditInsertEdgeActivity.cpp
+++ b/Engine/Plugins/Runtime/MeshModelingToolset/Source/MeshModelingTools/Private/ToolActivities/PolyEditInsertEdgeActivity.cpp
@@ -26,6 +26,8 @@ namespace PolyEditInsertEdgeActivityLocals
 		int32& GroupIDOut, int32& BoundaryIndexOut);
 	bool DoesBoundaryContainPoint(const FGroupTopology& Topology,
 		const FGroupTopology::FGroupBoundary& Boundary, int32 PointTopologyID, bool bPointIsCorner);
+
+	FText GroupEdgeStartTransactionName = LOCTEXT("GroupEdgeStartTransactionName", "Group Edge Start");
 }
 
 TUniquePtr<FDynamicMeshOperator> UPolyEditInsertEdgeActivity::MakeNewOperator()
@@ -271,6 +273,10 @@ void UPolyEditInsertEdgeActivity::Tick(float DeltaTime)
 					PreviewPoints.Reset();
 					PreviewPoints.Add(PreviewPoint);
 					ToolState = EState::GettingEnd;
+
+					ParentTool->GetToolManager()->EmitObjectChange(this,
+						MakeUnique<FGroupEdgeInsertionFirstPointChange>(CurrentChangeStamp),
+						PolyEditInsertEdgeActivityLocals::GroupEdgeStartTransactionName);
 				}
 			}
 		}
@@ -599,9 +605,9 @@ void UPolyEditInsertEdgeActivity::OnClicked(const FInputDeviceRay& ClickPos)
 			PreviewPoints.Add(PreviewPoint);
 			ToolState = EState::GettingEnd;
 
-			ParentTool->GetToolManager()->BeginUndoTransaction(LOCTEXT("GroupEdgeStartTransactionName", "Group Edge Start"));
+			ParentTool->GetToolManager()->BeginUndoTransaction(PolyEditInsertEdgeActivityLocals::GroupEdgeStartTransactionName);
 			ParentTool->GetToolManager()->EmitObjectChange(this, MakeUnique<FGroupEdgeInsertionFirstPointChange>(CurrentChangeStamp),
-				LOCTEXT("GroupEdgeStart", "Group Edge Start"));
+				PolyEditInsertEdgeActivityLocals::GroupEdgeStartTransactionName);
 			ParentTool->GetToolManager()->EndUndoTransaction();
 		}
 		break;
