ModelingMode: Fix the static mesh tool target to not be so aggressive in checking for transaction success when modifying mesh descriptions to prevent a crash from occurring when asset editors are closed with local changes during an asset reload.

diff --git a/Engine/Plugins/Runtime/MeshModelingToolset/Source/ModelingComponentsEditorOnly/Private/ToolTargets/StaticMeshToolTarget.cpp b/Engine/Plugins/Runtime/MeshModelingToolset/Source/ModelingComponentsEditorOnly/Private/ToolTargets/StaticMeshToolTarget.cpp
--- a/Engine/Plugins/Runtime/MeshModelingToolset/Source/ModelingComponentsEditorOnly/Private/ToolTargets/StaticMeshToolTarget.cpp
+++ b/Engine/Plugins/Runtime/MeshModelingToolset/Source/ModelingComponentsEditorOnly/Private/ToolTargets/StaticMeshToolTarget.cpp
@@ -290,11 +290,11 @@ void UStaticMeshToolTarget::CommitMeshDescription(UStaticMesh* StaticMeshIn, con
 
 	if (EditingLODIn == EMeshLODIdentifier::HiResSource)
 	{
-		verify(StaticMeshIn->ModifyHiResMeshDescription());
+		StaticMeshIn->ModifyHiResMeshDescription();
 	}
 	else
 	{
-		verify(StaticMeshIn->ModifyMeshDescription((int32)EditingLODIn));
+		StaticMeshIn->ModifyMeshDescription((int32)EditingLODIn);
 	}
 
 	FCommitterParams CommitterParams;
