Render Capture Baking: Fix a crash which occurred because the SceneCapture state was not properly updated when its computation was cancelled

diff --git a/Engine/Plugins/Experimental/MeshModelingToolsetExp/Source/MeshModelingToolsExp/Private/BakeRenderCaptureTool.cpp b/Engine/Plugins/Experimental/MeshModelingToolsetExp/Source/MeshModelingToolsExp/Private/BakeRenderCaptureTool.cpp
--- a/Engine/Plugins/Experimental/MeshModelingToolsetExp/Source/MeshModelingToolsExp/Private/BakeRenderCaptureTool.cpp
+++ b/Engine/Plugins/Experimental/MeshModelingToolsetExp/Source/MeshModelingToolsExp/Private/BakeRenderCaptureTool.cpp
@@ -22,7 +22,6 @@
 #include "Baking/BakingTypes.h"
 #include "Sampling/MeshImageBakingCache.h"
 #include "Sampling/MeshMapBaker.h"
-#include "Misc/ScopedSlowTask.h"
 #include "Sampling/RenderCaptureMapEvaluator.h"
 
 #include UE_INLINE_GENERATED_CPP_BY_NAME(BakeRenderCaptureTool)
@@ -839,32 +838,25 @@ FRenderCaptureUpdate UBakeRenderCaptureTool::UpdateSceneCapture()
 
 	if (!SceneCapture)
 	{
-		constexpr bool bAllowCancel = false;
-
-		FScopedSlowTask Progress(1.f, LOCTEXT("CapturingScene", "Capturing Scene..."));
-		Progress.EnterProgressFrame(1.f);
-		Progress.MakeDialog(bAllowCancel);
-
 		const FRenderCaptureOptions Options = MakeRenderCaptureOptions(*RenderCaptureProperties);
+
+		constexpr bool bAllowCancel = false;
 		SceneCapture = CapturePhotoSet(Actors, Options, Update, bAllowCancel);
 	}
 	else
 	{
-		constexpr bool bAllowCancel = true;
-
-		FScopedSlowTask Progress(1.f, LOCTEXT("CapturingScene", "Capturing Scene..."));
-		Progress.EnterProgressFrame(1.f);
-		Progress.MakeDialog(bAllowCancel);
-
 		// Get already computed options from the existing SceneCapture so we can restore them if the capture is cancelled
 		const FRenderCaptureOptions ComputedOptions = GetComputedPhotoSetOptions(SceneCapture);
 
 		const FRenderCaptureOptions Options = MakeRenderCaptureOptions(*RenderCaptureProperties);
+
+		constexpr bool bAllowCancel = true;
 		Update = UpdatePhotoSets(SceneCapture, Actors, Options, bAllowCancel);
 
 		if (SceneCapture->Cancelled())
 		{
 			// Restore the settings present before the change that invoked the scene capture recompute
+			// Note UpdatePhotoSets will have restored the SceneCapture settings to their values prior to the cancel
 			RenderCaptureProperties->bBaseColorMap      = ComputedOptions.bBakeBaseColor;
 			RenderCaptureProperties->bNormalMap         = ComputedOptions.bBakeNormalMap;
 			RenderCaptureProperties->bMetallicMap       = ComputedOptions.bBakeMetallic;
diff --git a/Engine/Plugins/Runtime/MeshModelingToolset/Source/ModelingComponents/Private/Baking/RenderCaptureFunctions.cpp b/Engine/Plugins/Runtime/MeshModelingToolset/Source/ModelingComponents/Private/Baking/RenderCaptureFunctions.cpp
--- a/Engine/Plugins/Runtime/MeshModelingToolset/Source/ModelingComponents/Private/Baking/RenderCaptureFunctions.cpp
+++ b/Engine/Plugins/Runtime/MeshModelingToolset/Source/ModelingComponents/Private/Baking/RenderCaptureFunctions.cpp
@@ -10,10 +10,12 @@
 #include "GameFramework/Actor.h"
 #include "Algo/NoneOf.h"
 #include "Algo/AnyOf.h"
+#include "Misc/ScopedSlowTask.h"
 
 
 using namespace UE::Geometry;
 
+#define LOCTEXT_NAMESPACE "RenderCaptureFunctions"
 
 
 
@@ -203,30 +205,15 @@ bool FSceneCapturePhotoSetSampler::IsValidCorrespondence(const FMeshMapEvaluator
 
 
 
-
-
-
-
-TUniquePtr<FSceneCapturePhotoSet> UE::Geometry::CapturePhotoSet(
-	const TArray<TObjectPtr<AActor>>& Actors,
-	const FRenderCaptureOptions& Options,
-	FRenderCaptureUpdate& Update,
-	bool bAllowCancel)
+namespace UE::Geometry
+{
+namespace
 {
-	TRACE_CPUPROFILER_EVENT_SCOPE(CapturePhotoSet);
-
-	TUniquePtr<FSceneCapturePhotoSet> SceneCapture = MakeUnique<FSceneCapturePhotoSet>();
-
-	Update = UpdatePhotoSets(SceneCapture, Actors, Options, bAllowCancel);
-
-	return SceneCapture;
-}
 
-FRenderCaptureUpdate UE::Geometry::UpdatePhotoSets(
+void UpdateSceneCaptureSettings(
 	const TUniquePtr<FSceneCapturePhotoSet>& SceneCapture,
-	const TArray<TObjectPtr<AActor>>& Actors,
-	const FRenderCaptureOptions& Options,
-	bool bAllowCancel)
+	const TArray<AActor*>& Actors,
+	const FRenderCaptureOptions& Options)
 {
 	SceneCapture->SetCaptureTypeEnabled(ERenderCaptureType::DeviceDepth, Options.bBakeDeviceDepth);
 	SceneCapture->SetCaptureTypeEnabled(ERenderCaptureType::BaseColor,   Options.bBakeBaseColor);
@@ -264,8 +251,42 @@ FRenderCaptureUpdate UE::Geometry::UpdatePhotoSets(
 		true, true, true, true, true);
 
 	SceneCapture->SetSpatialPhotoParams(SpatialParams);
+}
+
+} // namespace
+} // namespace UE::Geometry
+
 
-	SceneCapture->SetAllowCancel(bAllowCancel);
+
+TUniquePtr<FSceneCapturePhotoSet> UE::Geometry::CapturePhotoSet(
+	const TArray<TObjectPtr<AActor>>& Actors,
+	const FRenderCaptureOptions& Options,
+	FRenderCaptureUpdate& Update,
+	bool bAllowCancel)
+{
+	TRACE_CPUPROFILER_EVENT_SCOPE(CapturePhotoSet);
+
+	TUniquePtr<FSceneCapturePhotoSet> SceneCapture = MakeUnique<FSceneCapturePhotoSet>();
+
+	Update = UpdatePhotoSets(SceneCapture, Actors, Options, bAllowCancel);
+
+	return SceneCapture;
+}
+
+FRenderCaptureUpdate UE::Geometry::UpdatePhotoSets(
+	const TUniquePtr<FSceneCapturePhotoSet>& SceneCapture,
+	const TArray<TObjectPtr<AActor>>& Actors,
+	const FRenderCaptureOptions& Options,
+	bool bAllowCancel)
+{
+	FScopedSlowTask Progress(0.f, LOCTEXT("CapturingScene", "Capturing Scene..."));
+	Progress.MakeDialog(bAllowCancel);
+
+	// Cache previous SceneCapture settings so these can be restored if the computation is cancelled
+	const FRenderCaptureOptions ComputedOptions = GetComputedPhotoSetOptions(SceneCapture);
+	const TArray<AActor*> ComputedActors = SceneCapture->GetCaptureSceneActors();
+
+	UpdateSceneCaptureSettings(SceneCapture, Actors, Options);
 
 	// If the provided Options required some photo sets to recompute they will have been cleared by this point.
 	// The updated photo sets are the ones where the photo sets are empty here but are not empty after the Compute call
@@ -282,8 +303,15 @@ FRenderCaptureUpdate UE::Geometry::UpdatePhotoSets(
 	Update.bUpdatedSubsurfaceColor = (SceneCapture->GetSubsurfaceColorPhotoSet().Num() == 0);
 	Update.bUpdatedDeviceDepth     = (SceneCapture->GetDeviceDepthPhotoSet().Num()     == 0);
 
+	SceneCapture->SetAllowCancel(bAllowCancel);
+
 	SceneCapture->Compute();
 
+	if (SceneCapture->Cancelled())
+	{
+		UpdateSceneCaptureSettings(SceneCapture, ComputedActors, ComputedOptions);
+	}
+
 	Update.bUpdatedBaseColor       = (SceneCapture->GetBaseColorPhotoSet().Num()       == 0) != Update.bUpdatedBaseColor;
 	Update.bUpdatedRoughness       = (SceneCapture->GetRoughnessPhotoSet().Num()       == 0) != Update.bUpdatedRoughness;
 	Update.bUpdatedSpecular        = (SceneCapture->GetSpecularPhotoSet().Num()        == 0) != Update.bUpdatedSpecular;
@@ -293,7 +321,7 @@ FRenderCaptureUpdate UE::Geometry::UpdatePhotoSets(
 	Update.bUpdatedEmissive        = (SceneCapture->GetEmissivePhotoSet().Num()        == 0) != Update.bUpdatedEmissive;
 	Update.bUpdatedOpacity         = (SceneCapture->GetOpacityPhotoSet().Num()         == 0) != Update.bUpdatedOpacity;
 	Update.bUpdatedSubsurfaceColor = (SceneCapture->GetSubsurfaceColorPhotoSet().Num() == 0) != Update.bUpdatedSubsurfaceColor;
-	Update.bUpdatedDeviceDepth     = (SceneCapture->GetDeviceDepthPhotoSet().Num()     == 0) != Update.bUpdatedDeviceDepth ;
+	Update.bUpdatedDeviceDepth     = (SceneCapture->GetDeviceDepthPhotoSet().Num()     == 0) != Update.bUpdatedDeviceDepth;
 
 	return Update;
 }
@@ -704,3 +732,6 @@ void UE::Geometry::GetTexturesFromRenderCaptureBaker(const TUniquePtr<FMeshMapBa
 		}
 	}
 }
+
+
+#undef LOCTEXT_NAMESPACE
diff --git a/Engine/Plugins/Runtime/MeshModelingToolset/Source/ModelingComponents/Private/Scene/SceneCapturePhotoSet.cpp b/Engine/Plugins/Runtime/MeshModelingToolset/Source/ModelingComponents/Private/Scene/SceneCapturePhotoSet.cpp
--- a/Engine/Plugins/Runtime/MeshModelingToolset/Source/ModelingComponents/Private/Scene/SceneCapturePhotoSet.cpp
+++ b/Engine/Plugins/Runtime/MeshModelingToolset/Source/ModelingComponents/Private/Scene/SceneCapturePhotoSet.cpp
@@ -273,6 +273,8 @@ void FSceneCapturePhotoSet::Compute()
 {
 	check(this->TargetWorld != nullptr);
 
+	bWasCancelled = false;
+
 	FScopedSlowTask Progress(PhotoSetParams.Num(), LOCTEXT("ComputingViewpoints", "Computing Viewpoints..."));
 	Progress.MakeDialog(bAllowCancel);
 
@@ -472,7 +474,9 @@ void FSceneCapturePhotoSet::AddExteriorCaptures(
 	const TArray<FVector3d>& Directions)
 {
 	check(this->TargetWorld != nullptr);
-	
+
+	bWasCancelled = false;
+
 	FScopedSlowTask Progress(Directions.Num(), LOCTEXT("ComputingViewpoints", "Computing Viewpoints..."));
 	Progress.MakeDialog(bAllowCancel);
 
diff --git a/Engine/Plugins/Runtime/MeshModelingToolset/Source/ModelingComponents/Public/Baking/RenderCaptureFunctions.h b/Engine/Plugins/Runtime/MeshModelingToolset/Source/ModelingComponents/Public/Baking/RenderCaptureFunctions.h
--- a/Engine/Plugins/Runtime/MeshModelingToolset/Source/ModelingComponents/Public/Baking/RenderCaptureFunctions.h
+++ b/Engine/Plugins/Runtime/MeshModelingToolset/Source/ModelingComponents/Public/Baking/RenderCaptureFunctions.h
@@ -139,6 +139,9 @@ struct MODELINGCOMPONENTS_API FRenderCaptureUpdate
 };
 
 
+/**
+ * This function computes the SceneCapture for the first time
+ */
 MODELINGCOMPONENTS_API
 TUniquePtr<FSceneCapturePhotoSet> CapturePhotoSet(
 	const TArray<TObjectPtr<AActor>>& Actors,
@@ -154,6 +157,10 @@ TUniquePtr<FSceneCapturePhotoSet> CapturePhotoSet(
 * - If the given Options enable a new capture channel then the new photos set are captured and added to the SceneCapture
 * - If the given Actors are different the ones cached in the SceneCapture, or if the given Options changes a parameter
 *   affecting all photo sets (e.g., photo resolution), then all the photo sets are cleared and recomputed
+*   
+* If the SceneCapture computation is cancelled then its state will be identical to how it looked prior to calling this
+* function, which means that queries of the SceneCapture settings (e.g., via the GetXXX() API) will return results which
+* are consistent with the internal photo sets.
 */
 MODELINGCOMPONENTS_API
 FRenderCaptureUpdate UpdatePhotoSets(
