MeshModelingTools: Fix crash in Merge tool when mesh uses a material ID that is not in its material list.

diff --git a/Engine/Plugins/Runtime/MeshModelingToolset/Source/MeshModelingTools/Private/CombineMeshesTool.cpp b/Engine/Plugins/Runtime/MeshModelingToolset/Source/MeshModelingTools/Private/CombineMeshesTool.cpp
--- a/Engine/Plugins/Runtime/MeshModelingToolset/Source/MeshModelingTools/Private/CombineMeshesTool.cpp
+++ b/Engine/Plugins/Runtime/MeshModelingToolset/Source/MeshModelingTools/Private/CombineMeshesTool.cpp
@@ -30,6 +30,42 @@ using namespace UE::Geometry;
 
 #define LOCTEXT_NAMESPACE "UCombineMeshesTool"
 
+namespace CombineMeshesToolLocals
+{
+	void SetNewMaterialID(int32 ComponentIdx, FDynamicMeshMaterialAttribute* MatAttrib, int32 TID, 
+		TArray<TArray<int32>>& MaterialIDRemaps, TArray<UMaterialInterface*>& AllMaterials)
+	{
+		int MatID = MatAttrib->GetValue(TID);
+		if (!ensure(MatID >= 0))
+		{
+			return;
+		}
+		if (MatID >= MaterialIDRemaps[ComponentIdx].Num())
+		{
+			UE_LOG(LogGeometry, Warning, TEXT("UCombineMeshesTool: Component %d had at least one material ID (%d) "
+				"that was not in its material list."), ComponentIdx, MatID);
+			
+			// There are different things we could do here. It's worth noting that out of bounds material indices
+			// are handled differently in static meshes and dynamic mesh components, and depend in part on how we
+			// got to that state. So trying to preserve a specific behavior is not practical, and probably not 
+			// expected if the user is not giving us valid data to begin with.
+			// The route we go is to give a separate nullptr material slot to each out of bounds ID. This will give
+			// the user a chance to preserve their material assignments while fixing the issue by assigning materials to
+			// the slots created in the output (at least, unless they pass through this tool again, at which point any
+			// nullptr-pointing IDs will be collapsed to point to the same nullptr slot, due to the way we create the
+			// combined material list for in-bounds IDs).
+			int32 NumElementsToAdd = MatID - MaterialIDRemaps[ComponentIdx].Num() + 1;
+			for (int32 i = 0; i < NumElementsToAdd; ++i)
+			{
+				MaterialIDRemaps[ComponentIdx].Add(AllMaterials.Num());
+				AllMaterials.Add(nullptr);
+			}
+			checkSlow(MaterialIDRemaps[ComponentIdx].Num() == MatID + 1);
+		}
+		MatAttrib->SetValue(TID, MaterialIDRemaps[ComponentIdx][MatID]);
+	}
+}
+
 /*
  * ToolBuilder
  */
@@ -146,6 +182,8 @@ void UCombineMeshesTool::OnShutdown(EToolShutdownType ShutdownType)
 
 void UCombineMeshesTool::CreateNewAsset()
 {
+	using namespace CombineMeshesToolLocals;
+
 	// Make sure meshes are available before we open transaction. This is to avoid potential stability issues related 
 	// to creation/load of meshes inside a transaction, for assets that possibly do not have bulk data currently loaded.
 	TArray<FDynamicMesh3> InputMeshes;
@@ -226,8 +264,7 @@ void UCombineMeshesTool::CreateNewAsset()
 			FDynamicMeshMaterialAttribute* MatAttrib = ComponentDMesh.Attributes()->GetMaterialID();
 			for (int TID : ComponentDMesh.TriangleIndicesItr())
 			{
-				int MatID = MatAttrib->GetValue(TID);
-				MatAttrib->SetValue(TID, MaterialIDRemaps[ComponentIdx][MatID]);
+				SetNewMaterialID(ComponentIdx, MatAttrib, TID, MaterialIDRemaps, AllMaterials);
 			}
 
 			FDynamicMeshEditor Editor(&AccumulateDMesh);
@@ -328,6 +365,8 @@ void UCombineMeshesTool::CreateNewAsset()
 
 void UCombineMeshesTool::UpdateExistingAsset()
 {
+	using namespace CombineMeshesToolLocals;
+
 	// Make sure meshes are available before we open transaction. This is to avoid potential stability issues related 
 	// to creation/load of meshes inside a transaction, for assets that possibly do not have bulk data currently loaded.
 	TArray<FDynamicMesh3> InputMeshes;
@@ -392,8 +431,7 @@ void UCombineMeshesTool::UpdateExistingAsset()
 			FDynamicMeshMaterialAttribute* MatAttrib = ComponentDMesh.Attributes()->GetMaterialID();
 			for (int TID : ComponentDMesh.TriangleIndicesItr())
 			{
-				int MatID = MatAttrib->GetValue(TID);
-				MatAttrib->SetValue(TID, MaterialIDRemaps[ComponentIdx][MatID]);
+				SetNewMaterialID(ComponentIdx, MatAttrib, TID, MaterialIDRemaps, AllMaterials);
 			}
 
 			if (ComponentIdx != SkipIndex)
