MeshModelingTools: Fix static mesh mode-level edits (via selector) not being properly undoable until passing through some other tool.

diff --git a/Engine/Plugins/Runtime/MeshModelingToolset/Source/ModelingComponentsEditorOnly/Private/Selection/StaticMeshSelector.cpp b/Engine/Plugins/Runtime/MeshModelingToolset/Source/ModelingComponentsEditorOnly/Private/Selection/StaticMeshSelector.cpp
--- a/Engine/Plugins/Runtime/MeshModelingToolset/Source/ModelingComponentsEditorOnly/Private/Selection/StaticMeshSelector.cpp
+++ b/Engine/Plugins/Runtime/MeshModelingToolset/Source/ModelingComponentsEditorOnly/Private/Selection/StaticMeshSelector.cpp
@@ -120,6 +120,8 @@ void FStaticMeshSelector::CommitMeshTransform()
 
 	// emit transaction here??
 
+	// make sure transactional flag is on for this asset
+	StaticMesh->SetFlags(RF_Transactional);
 	// mark as modified
 	StaticMesh->Modify();
 
