MeshModelingTools: Fix ensure being fired when exiting group edge activity with an edge started.

diff --git a/Engine/Plugins/Runtime/MeshModelingToolset/Source/MeshModelingTools/Private/ToolActivities/PolyEditInsertEdgeActivity.cpp b/Engine/Plugins/Runtime/MeshModelingToolset/Source/MeshModelingTools/Private/ToolActivities/PolyEditInsertEdgeActivity.cpp
--- a/Engine/Plugins/Runtime/MeshModelingToolset/Source/MeshModelingTools/Private/ToolActivities/PolyEditInsertEdgeActivity.cpp
+++ b/Engine/Plugins/Runtime/MeshModelingToolset/Source/MeshModelingTools/Private/ToolActivities/PolyEditInsertEdgeActivity.cpp
@@ -178,7 +178,10 @@ EToolActivityEndResult UPolyEditInsertEdgeActivity::End(EToolShutdownType Shutdo
 
 	ActivityContext->Preview->OnOpCompleted.RemoveAll(this);
 	ActivityContext->Preview->OnMeshUpdated.RemoveAll(this);
-	ClearPreview(true);
+	ToolState = EState::GettingStart;
+	// Note that this does an ensure on us not being in a ToolState that requires the preview points, hence
+	// we do this after resetting ToolState.
+	ClearPreview(true); 
 	ActivityContext->Preview->ClearOpFactory();
 
 	LatestOpTopologyResult.Reset();
