Mesh Modeling / UV Projection Tool: allow the user to change the Initialization mode and have it affect the current instance of the tool, as long as they do it before modifying the transform

diff --git a/Engine/Plugins/Runtime/MeshModelingToolset/Source/MeshModelingTools/Private/UVProjectionTool.cpp b/Engine/Plugins/Runtime/MeshModelingToolset/Source/MeshModelingTools/Private/UVProjectionTool.cpp
--- a/Engine/Plugins/Runtime/MeshModelingToolset/Source/MeshModelingTools/Private/UVProjectionTool.cpp
+++ b/Engine/Plugins/Runtime/MeshModelingToolset/Source/MeshModelingTools/Private/UVProjectionTool.cpp
@@ -167,11 +167,14 @@ void UUVProjectionTool::Setup()
 	});
 	DimensionsModeWatcher = BasicProperties->WatchProperty(BasicProperties->bProportionalDimensions, OnProportionalDimensionsChanged);
 	BasicProperties->WatchProperty(BasicProperties->Initialization, [this](EUVProjectionToolInitializationMode NewMode) { OnInitializationModeChanged(); });
-	bTransformModified = false;
 	BasicProperties->SilentUpdateWatched();
 
 	OnProportionalDimensionsChanged(true);	// Initialize CachedDimensions
-	
+
+	// Allow the user to change the Initialization mode and have it affect the current run of the tool, as long as they do it before modifying the transform
+	bTransformModified = false;
+
+
 	// click to set plane behavior
 	SetPlaneCtrlClickBehaviorTarget = MakeUnique<FSelectClickedAction>();
 	SetPlaneCtrlClickBehaviorTarget->SnapManager = USceneSnappingManager::Find(GetToolManager());
