Fix incorrect handling of texture source format BGRE8 in UE::AssetUtils::ReadTexture function

diff --git a/Engine/Plugins/Runtime/MeshModelingToolset/Source/ModelingComponents/Private/AssetUtils/Texture2DUtil.cpp b/Engine/Plugins/Runtime/MeshModelingToolset/Source/ModelingComponents/Private/AssetUtils/Texture2DUtil.cpp
--- a/Engine/Plugins/Runtime/MeshModelingToolset/Source/ModelingComponents/Private/AssetUtils/Texture2DUtil.cpp
+++ b/Engine/Plugins/Runtime/MeshModelingToolset/Source/ModelingComponents/Private/AssetUtils/Texture2DUtil.cpp
@@ -2,6 +2,7 @@
 
 #include "AssetUtils/Texture2DUtil.h"
 #include "EngineModule.h"
+#include "ImageUtils.h"
 #include "RendererInterface.h"
 #include "RenderUtils.h"
 #include "TextureResource.h"
@@ -97,6 +98,36 @@ static bool ReadTexture_SourceData(
 	UTexture2D* TextureMap,
 	TImageBuilder<FVector4f>& DestImage)
 {
+	// TODO #jira UE-176086 Enable this code when it is less risky to do so, and after it has been tested. This code
+	// also fixes an issue where not all texture source formats are supported
+	if constexpr (false)
+	{
+		const FTextureSource& TextureSource = TextureMap->Source;
+		const int32 Width = TextureSource.GetSizeX();
+		const int32 Height = TextureSource.GetSizeY();
+
+		FImage Image;
+		FImageUtils::GetTexture2DSourceImage(TextureMap, Image);
+
+		FImage LinearImage;
+		Image.Linearize(LinearImage);
+
+		if (ensure(Width == LinearImage.SizeX) && ensure(Height == LinearImage.SizeY))
+		{
+			const FImageDimensions Dimensions = FImageDimensions(Width, Height);
+			DestImage.SetDimensions(Dimensions);
+			for (int64 LinearIndex = 0; LinearIndex < Dimensions.Num(); ++LinearIndex)
+			{
+				const FLinearColor& LinearColor = Image.AsRGBA32F()[LinearIndex];
+				DestImage.SetPixel(Dimensions.GetCoords(LinearIndex), ToVector4<float>(LinearColor));
+			}
+
+			return true;
+		}
+
+		return false;
+	}
+
 	FTextureSource& TextureSource = TextureMap->Source;
 
 	const int32 Width = TextureSource.GetSizeX();
@@ -112,7 +143,7 @@ static bool ReadTexture_SourceData(
 	const uint8* SourceDataPtr = SourceData.GetData();
 
 	// code below is derived from UBlueprintMaterialTextureNodesBPLibrary::Texture2D_SampleUV_EditorOnly()
-	if ((SourceFormat == TSF_BGRA8 || SourceFormat == TSF_BGRE8))
+	if (SourceFormat == TSF_BGRA8)
 	{
 		check(BytesPerPixel == sizeof(FColor));
 		for (int64 i = 0; i < Num; ++i)
@@ -127,6 +158,19 @@ static bool ReadTexture_SourceData(
 			DestImage.SetPixel(i, ToVector4<float>(FloatColor));
 		}
 	}
+	else if (SourceFormat == TSF_BGRE8)
+	{
+		check(BytesPerPixel == sizeof(FColor));
+		for (int64 i = 0; i < Num; ++i)
+		{
+			const uint8* PixelPtr = SourceDataPtr + (i * BytesPerPixel);
+			FColor PixelColor = *((FColor*)PixelPtr);
+
+			FLinearColor FloatColor = PixelColor.FromRGBE();
+
+			DestImage.SetPixel(i, ToVector4<float>(FloatColor));
+		}
+	}
 	// code below is also derived from FImage::CopyTo (CopyImage) - invoked during BuildTexture on import.
 	else if (SourceFormat == TSF_RGBA16)
 	{
