Fixed a bug where the trim tool starts a potentially expensive boolean calculation when only the opacity color alpha changes.

diff --git a/Engine/Plugins/Runtime/MeshModelingToolset/Source/MeshModelingTools/Private/CSGMeshesTool.cpp b/Engine/Plugins/Runtime/MeshModelingToolset/Source/MeshModelingTools/Private/CSGMeshesTool.cpp
--- a/Engine/Plugins/Runtime/MeshModelingToolset/Source/MeshModelingTools/Private/CSGMeshesTool.cpp
+++ b/Engine/Plugins/Runtime/MeshModelingToolset/Source/MeshModelingTools/Private/CSGMeshesTool.cpp
@@ -357,7 +357,8 @@ void UCSGMeshesTool::OnPropertyModified(UObject* PropertySet, FProperty* Propert
 										 Property->GetFName() == GET_MEMBER_NAME_CHECKED(UTrimMeshesToolProperties, ColorOfTrimmingMesh) ||
 										 Property->GetFName() == FName("R") ||
 										 Property->GetFName() == FName("G") ||
-										 Property->GetFName() == FName("B"));
+										 Property->GetFName() == FName("B") ||
+										 Property->GetFName() == FName("A"));
 		}
 		else 
 		{	
@@ -366,7 +367,8 @@ void UCSGMeshesTool::OnPropertyModified(UObject* PropertySet, FProperty* Propert
 										 Property->GetFName() == GET_MEMBER_NAME_CHECKED(UCSGMeshesToolProperties, SubtractedMeshOpacity) ||
 										 Property->GetFName() == FName("R") ||
 										 Property->GetFName() == FName("G") ||
-										 Property->GetFName() == FName("B"));
+										 Property->GetFName() == FName("B") ||
+										 Property->GetFName() == FName("A"));
 		}
 
 		// If the property that was changed only affects the visuals, we do not need to recalculate the output 
