Fix crash in BakeRC tool caused by not checking for validity of SceneCapture

diff --git a/Engine/Plugins/Experimental/MeshModelingToolsetExp/Source/MeshModelingToolsExp/Private/BakeRenderCaptureTool.cpp b/Engine/Plugins/Experimental/MeshModelingToolsetExp/Source/MeshModelingToolsExp/Private/BakeRenderCaptureTool.cpp
--- a/Engine/Plugins/Experimental/MeshModelingToolsetExp/Source/MeshModelingToolsExp/Private/BakeRenderCaptureTool.cpp
+++ b/Engine/Plugins/Experimental/MeshModelingToolsetExp/Source/MeshModelingToolsExp/Private/BakeRenderCaptureTool.cpp
@@ -569,6 +569,9 @@ bool UBakeRenderCaptureTool::CanAccept() const
 
 TUniquePtr<TGenericDataOperator<FMeshMapBaker>> UBakeRenderCaptureTool::MakeNewOperator()
 {
+	// We should not have requested a bake if we don't have a SceneCapture
+	check(SceneCapture.IsValid());
+
 	TUniquePtr<FRenderCaptureMapBakerOp> Op = MakeUnique<FRenderCaptureMapBakerOp>();
 	Op->BaseMesh = &TargetMesh;
 	Op->BaseMeshSpatial = &TargetMeshSpatial;
@@ -1086,39 +1089,39 @@ void UBakeRenderCaptureTool::InvalidateResults(FRenderCaptureUpdate Update)
 	// Note that the bake operation, Compute, updates ResultSettings when results are available via the
 	// Compute->OnResultUpdated delegate.
 
-	if (!SceneCapture->GetCaptureTypeEnabled(ERenderCaptureType::BaseColor) || Update.bUpdatedBaseColor)
+	if (Update.bUpdatedBaseColor)
 	{
 		ResultSettings->BaseColorMap = nullptr;
 	}
-	if (!SceneCapture->GetCaptureTypeEnabled(ERenderCaptureType::Roughness) || Update.bUpdatedRoughness)
+	if (Update.bUpdatedRoughness)
 	{
 		ResultSettings->RoughnessMap = nullptr;
 	}
-	if (!SceneCapture->GetCaptureTypeEnabled(ERenderCaptureType::Metallic) || Update.bUpdatedMetallic)
+	if (Update.bUpdatedMetallic)
 	{
 		ResultSettings->MetallicMap = nullptr;
 	}
-	if (!SceneCapture->GetCaptureTypeEnabled(ERenderCaptureType::Specular) || Update.bUpdatedSpecular)
+	if (Update.bUpdatedSpecular)
 	{
 		ResultSettings->SpecularMap = nullptr;
 	}
-	if (!SceneCapture->GetCaptureTypeEnabled(ERenderCaptureType::Emissive) || Update.bUpdatedEmissive)
+	if (Update.bUpdatedEmissive)
 	{
 		ResultSettings->EmissiveMap = nullptr;
 	}
-	if (!SceneCapture->GetCaptureTypeEnabled(ERenderCaptureType::WorldNormal) || Update.bUpdatedNormalMap)
+	if (Update.bUpdatedNormalMap)
 	{
 		ResultSettings->NormalMap = nullptr;
 	}
-	if (!SceneCapture->GetCaptureTypeEnabled(ERenderCaptureType::Opacity) || Update.bUpdatedOpacity)
+	if (Update.bUpdatedOpacity)
 	{
 		ResultSettings->OpacityMap = nullptr;
 	}
-	if (!SceneCapture->GetCaptureTypeEnabled(ERenderCaptureType::SubsurfaceColor) || Update.bUpdatedSubsurfaceColor)
+	if (Update.bUpdatedSubsurfaceColor)
 	{
 		ResultSettings->SubsurfaceColorMap = nullptr;
 	}
-	if (!SceneCapture->GetCaptureTypeEnabled(ERenderCaptureType::CombinedMRS) || Update.bUpdatedPackedMRS)
+	if (Update.bUpdatedPackedMRS)
 	{
 		ResultSettings->PackedMRSMap = nullptr;
 	}
diff --git a/Engine/Plugins/Experimental/MeshModelingToolsetExp/Source/MeshModelingToolsExp/Public/BakeRenderCaptureTool.h b/Engine/Plugins/Experimental/MeshModelingToolsetExp/Source/MeshModelingToolsExp/Public/BakeRenderCaptureTool.h
--- a/Engine/Plugins/Experimental/MeshModelingToolsetExp/Source/MeshModelingToolsExp/Public/BakeRenderCaptureTool.h
+++ b/Engine/Plugins/Experimental/MeshModelingToolsetExp/Source/MeshModelingToolsExp/Public/BakeRenderCaptureTool.h
@@ -408,8 +408,8 @@ protected:
 	UPROPERTY()
 	TObjectPtr<UPreviewMesh> PreviewMesh;
 
-	// TODO We currently need to compute this on the game thread because the implementation has checks for this
-	TUniquePtr<UE::Geometry::FSceneCapturePhotoSet> SceneCapture;
+	// Note: We need to compute this on the game thread because the implementation has checks for this
+	TUniquePtr<UE::Geometry::FSceneCapturePhotoSet> SceneCapture = nullptr;
 
 	// These are used to determine if we need to re-bake results
 	float ComputedValidDepthThreshold;
diff --git a/Engine/Plugins/Runtime/MeshModelingToolset/Source/ModelingComponents/Private/Baking/RenderCaptureFunctions.cpp b/Engine/Plugins/Runtime/MeshModelingToolset/Source/ModelingComponents/Private/Baking/RenderCaptureFunctions.cpp
--- a/Engine/Plugins/Runtime/MeshModelingToolset/Source/ModelingComponents/Private/Baking/RenderCaptureFunctions.cpp
+++ b/Engine/Plugins/Runtime/MeshModelingToolset/Source/ModelingComponents/Private/Baking/RenderCaptureFunctions.cpp
@@ -286,42 +286,67 @@ FRenderCaptureUpdate UE::Geometry::UpdatePhotoSets(
 	const FRenderCaptureOptions ComputedOptions = GetComputedPhotoSetOptions(SceneCapture);
 	const TArray<AActor*> ComputedActors = SceneCapture->GetCaptureSceneActors();
 
-	UpdateSceneCaptureSettings(SceneCapture, Actors, Options);
+	// We track the photo set counts in order to determine which photosets were updated
+	struct FPhotoSetNums
+	{
+		FPhotoSetNums(const TUniquePtr<FSceneCapturePhotoSet>& SceneCapture)
+		{
+			NumBaseColor       = SceneCapture->GetBaseColorPhotoSet().Num();
+			NumRoughness       = SceneCapture->GetRoughnessPhotoSet().Num();
+			NumSpecular        = SceneCapture->GetSpecularPhotoSet().Num();
+			NumMetallic        = SceneCapture->GetMetallicPhotoSet().Num();
+			NumPackedMRS       = SceneCapture->GetPackedMRSPhotoSet().Num();
+			NumNormalMap       = SceneCapture->GetWorldNormalPhotoSet().Num();
+			NumEmissive        = SceneCapture->GetEmissivePhotoSet().Num();
+			NumOpacity         = SceneCapture->GetOpacityPhotoSet().Num();
+			NumSubsurfaceColor = SceneCapture->GetSubsurfaceColorPhotoSet().Num();
+			NumDeviceDepth     = SceneCapture->GetDeviceDepthPhotoSet().Num();
+		}
+		
+		int32 NumBaseColor;
+		int32 NumRoughness;
+		int32 NumSpecular;
+		int32 NumMetallic;
+		int32 NumPackedMRS;
+		int32 NumNormalMap;
+		int32 NumEmissive;
+		int32 NumOpacity;
+		int32 NumSubsurfaceColor;
+		int32 NumDeviceDepth;
+	};
 
-	// If the provided Options required some photo sets to recompute they will have been cleared by this point.
-	// The updated photo sets are the ones where the photo sets are empty here but are not empty after the Compute call
-	FRenderCaptureUpdate Update;
+	const FPhotoSetNums Step1(SceneCapture);
+
+	// This will clear any photosets that are disabled or need recomputing
+	UpdateSceneCaptureSettings(SceneCapture, Actors, Options);
 
-	Update.bUpdatedBaseColor       = (SceneCapture->GetBaseColorPhotoSet().Num()       == 0);
-	Update.bUpdatedRoughness       = (SceneCapture->GetRoughnessPhotoSet().Num()       == 0);
-	Update.bUpdatedSpecular        = (SceneCapture->GetSpecularPhotoSet().Num()        == 0);
-	Update.bUpdatedMetallic        = (SceneCapture->GetMetallicPhotoSet().Num()        == 0);
-	Update.bUpdatedPackedMRS       = (SceneCapture->GetPackedMRSPhotoSet().Num()       == 0);
-	Update.bUpdatedNormalMap       = (SceneCapture->GetWorldNormalPhotoSet().Num()     == 0);
-	Update.bUpdatedEmissive        = (SceneCapture->GetEmissivePhotoSet().Num()        == 0);
-	Update.bUpdatedOpacity         = (SceneCapture->GetOpacityPhotoSet().Num()         == 0);
-	Update.bUpdatedSubsurfaceColor = (SceneCapture->GetSubsurfaceColorPhotoSet().Num() == 0);
-	Update.bUpdatedDeviceDepth     = (SceneCapture->GetDeviceDepthPhotoSet().Num()     == 0);
+	const FPhotoSetNums Step2(SceneCapture);
 
 	SceneCapture->SetAllowCancel(bAllowCancel);
 
+	// This will compute newly requested photosets
 	SceneCapture->Compute();
 
 	if (SceneCapture->Cancelled())
 	{
+		// This will clear any newly requested photosets
 		UpdateSceneCaptureSettings(SceneCapture, ComputedActors, ComputedOptions);
 	}
 
-	Update.bUpdatedBaseColor       = (SceneCapture->GetBaseColorPhotoSet().Num()       == 0) != Update.bUpdatedBaseColor;
-	Update.bUpdatedRoughness       = (SceneCapture->GetRoughnessPhotoSet().Num()       == 0) != Update.bUpdatedRoughness;
-	Update.bUpdatedSpecular        = (SceneCapture->GetSpecularPhotoSet().Num()        == 0) != Update.bUpdatedSpecular;
-	Update.bUpdatedMetallic        = (SceneCapture->GetMetallicPhotoSet().Num()        == 0) != Update.bUpdatedMetallic;
-	Update.bUpdatedPackedMRS       = (SceneCapture->GetPackedMRSPhotoSet().Num()       == 0) != Update.bUpdatedPackedMRS;
-	Update.bUpdatedNormalMap       = (SceneCapture->GetWorldNormalPhotoSet().Num()     == 0) != Update.bUpdatedNormalMap;
-	Update.bUpdatedEmissive        = (SceneCapture->GetEmissivePhotoSet().Num()        == 0) != Update.bUpdatedEmissive;
-	Update.bUpdatedOpacity         = (SceneCapture->GetOpacityPhotoSet().Num()         == 0) != Update.bUpdatedOpacity;
-	Update.bUpdatedSubsurfaceColor = (SceneCapture->GetSubsurfaceColorPhotoSet().Num() == 0) != Update.bUpdatedSubsurfaceColor;
-	Update.bUpdatedDeviceDepth     = (SceneCapture->GetDeviceDepthPhotoSet().Num()     == 0) != Update.bUpdatedDeviceDepth;
+	const FPhotoSetNums Step3(SceneCapture);
+
+	// If the photo sets were cleared/computed/recomputed we consider them updated, use the counts to figure this out
+	FRenderCaptureUpdate Update;
+	Update.bUpdatedBaseColor   = !(Step1.NumBaseColor   == Step2.NumBaseColor   && Step2.NumBaseColor   == Step3.NumBaseColor);
+	Update.bUpdatedRoughness   = !(Step1.NumRoughness   == Step2.NumRoughness   && Step2.NumRoughness   == Step3.NumRoughness);
+	Update.bUpdatedSpecular    = !(Step1.NumSpecular    == Step2.NumSpecular    && Step2.NumSpecular    == Step3.NumSpecular);
+	Update.bUpdatedMetallic    = !(Step1.NumMetallic    == Step2.NumMetallic    && Step2.NumMetallic    == Step3.NumMetallic);
+	Update.bUpdatedPackedMRS   = !(Step1.NumPackedMRS   == Step2.NumPackedMRS   && Step2.NumPackedMRS   == Step3.NumPackedMRS);
+	Update.bUpdatedNormalMap   = !(Step1.NumNormalMap   == Step2.NumNormalMap   && Step2.NumNormalMap   == Step3.NumNormalMap);
+	Update.bUpdatedEmissive    = !(Step1.NumEmissive    == Step2.NumEmissive    && Step2.NumEmissive    == Step3.NumEmissive);
+	Update.bUpdatedOpacity     = !(Step1.NumOpacity     == Step2.NumOpacity     && Step2.NumOpacity     == Step3.NumOpacity);
+	Update.bUpdatedDeviceDepth = !(Step1.NumDeviceDepth == Step2.NumDeviceDepth && Step2.NumDeviceDepth == Step3.NumDeviceDepth);
+	Update.bUpdatedSubsurfaceColor = !(Step1.NumSubsurfaceColor == Step2.NumSubsurfaceColor && Step2.NumSubsurfaceColor == Step3.NumSubsurfaceColor);
 
 	return Update;
 }
