Fix horde issue 259463

Note this issue is occurring in a `if constexpr (false)` block. This problem previously occurred in horde issue 259313 with some other things that may have since been fixed

diff --git a/Engine/Plugins/Runtime/MeshModelingToolset/Source/ModelingComponents/Private/AssetUtils/Texture2DUtil.cpp b/Engine/Plugins/Runtime/MeshModelingToolset/Source/ModelingComponents/Private/AssetUtils/Texture2DUtil.cpp
--- a/Engine/Plugins/Runtime/MeshModelingToolset/Source/ModelingComponents/Private/AssetUtils/Texture2DUtil.cpp
+++ b/Engine/Plugins/Runtime/MeshModelingToolset/Source/ModelingComponents/Private/AssetUtils/Texture2DUtil.cpp
@@ -118,7 +118,7 @@ static bool ReadTexture_SourceData(
 			DestImage.SetDimensions(Dimensions);
 			for (int64 LinearIndex = 0; LinearIndex < Dimensions.Num(); ++LinearIndex)
 			{
-				const FLinearColor& LinearColor = Image.AsRGBA32F()[LinearIndex];
+				const FLinearColor LinearColor = Image.AsRGBA32F()[LinearIndex];
 				DestImage.SetPixel(Dimensions.GetCoords(LinearIndex), ToVector4<float>(LinearColor));
 			}
 
