Fix horde issue 259463

This is the second attempt to fix this problem, sorry! It's a PVS static analysis problem.

diff --git a/Engine/Plugins/Runtime/MeshModelingToolset/Source/ModelingComponents/Private/AssetUtils/Texture2DUtil.cpp b/Engine/Plugins/Runtime/MeshModelingToolset/Source/ModelingComponents/Private/AssetUtils/Texture2DUtil.cpp
--- a/Engine/Plugins/Runtime/MeshModelingToolset/Source/ModelingComponents/Private/AssetUtils/Texture2DUtil.cpp
+++ b/Engine/Plugins/Runtime/MeshModelingToolset/Source/ModelingComponents/Private/AssetUtils/Texture2DUtil.cpp
@@ -116,10 +116,11 @@ static bool ReadTexture_SourceData(
 		{
 			const FImageDimensions Dimensions = FImageDimensions(Width, Height);
 			DestImage.SetDimensions(Dimensions);
+
+			TArrayView64<FLinearColor> ImageData = Image.AsRGBA32F();
 			for (int64 LinearIndex = 0; LinearIndex < Dimensions.Num(); ++LinearIndex)
 			{
-				const FLinearColor LinearColor = Image.AsRGBA32F()[LinearIndex];
-				DestImage.SetPixel(Dimensions.GetCoords(LinearIndex), ToVector4<float>(LinearColor));
+				DestImage.SetPixel(Dimensions.GetCoords(LinearIndex), ToVector4<float>(ImageData[LinearIndex]));
 			}
 
 			return true;
