MeshModelingTools: unbind the "toggle gizmo" action in some tools from a default hotkey to prevent users from accidentally pressing it and thinking that the tools are broken. Users can still set their own hotkey if their workflow benefits from it.

diff --git a/Engine/Plugins/Runtime/MeshModelingToolset/Source/MeshModelingTools/Private/DrawPolygonTool.cpp b/Engine/Plugins/Runtime/MeshModelingToolset/Source/MeshModelingTools/Private/DrawPolygonTool.cpp
--- a/Engine/Plugins/Runtime/MeshModelingToolset/Source/MeshModelingTools/Private/DrawPolygonTool.cpp
+++ b/Engine/Plugins/Runtime/MeshModelingToolset/Source/MeshModelingTools/Private/DrawPolygonTool.cpp
@@ -186,11 +186,14 @@ void UDrawPolygonTool::Shutdown(EToolShutdownType ShutdownType)
 
 void UDrawPolygonTool::RegisterActions(FInteractiveToolActionSet& ActionSet)
 {
+	// This is an uncommon hotkey to want, and can make tool seem broken if accidentally pressed. The only
+	// reason a user might want it is if the gizmo is in the way of their drawing, in which case they could
+	// presumably translate it in the plane... Still, allow it to be settable if the user wants.
 	ActionSet.RegisterAction(this, (int32)EStandardToolActions::BaseClientDefinedActionID + 2,
-		TEXT("ToggleGizmo"),
-		LOCTEXT("ToggleGizmo", "Toggle Gizmo"),
-		LOCTEXT("ToggleGizmoTooltip", "Toggle visibility of the transformation gizmo"),
-		EModifierKey::None, EKeys::A,
+		TEXT("ToggleGizmoVisibility"),
+		LOCTEXT("ToggleGizmoVisibilityUIName", "Toggle Gizmo"),
+		LOCTEXT("ToggleGizmoVisibilityTooltip", "Toggle visibility of the transformation gizmo"),
+		EModifierKey::None, EKeys::Invalid,
 		[this]() { PolygonProperties->bShowGridGizmo = !PolygonProperties->bShowGridGizmo; });
 }
 
diff --git a/Engine/Plugins/Runtime/MeshModelingToolset/Source/MeshModelingTools/Private/EditMeshPolygonsTool.cpp b/Engine/Plugins/Runtime/MeshModelingToolset/Source/MeshModelingTools/Private/EditMeshPolygonsTool.cpp
--- a/Engine/Plugins/Runtime/MeshModelingToolset/Source/MeshModelingTools/Private/EditMeshPolygonsTool.cpp
+++ b/Engine/Plugins/Runtime/MeshModelingToolset/Source/MeshModelingTools/Private/EditMeshPolygonsTool.cpp
@@ -613,12 +613,11 @@ void UEditMeshPolygonsTool::Setup()
 
 void UEditMeshPolygonsTool::ResetUserMessage()
 {
-	// When the gizmo is hidden, notify the user and
-	// specify the toggle hotkey to prevent panic
+	// When the gizmo is hidden, notify the user and tell them how to fix it.
 	if (!TransformGizmo->IsVisible())
 	{
 		GetToolManager()->DisplayMessage(
-			LOCTEXT("ToggleTransformGizmoNotify", "Transform Gizmo Hidden, Press 'R' to Unhide"),
+			LOCTEXT("ToggleTransformGizmoNotify", "Transform gizmo hidden, unhide by toggling \"Gizmo Visible\" (or using hotkey, if set)"),
 			EToolMessageLevel::UserNotification);
 	}
 	else
@@ -794,11 +793,15 @@ void UEditMeshPolygonsTool::RegisterActions(FInteractiveToolActionSet& ActionSet
 		LOCTEXT("DeleteSelectionUIName", "Delete Selection"),
 		LOCTEXT("DeleteSelectionTooltip", "Delete Selection"),
 		EModifierKey::None, EKeys::Delete, OnDeletionKeyPress);
+
+	// This hotkey can make the tool seem broken if it is accidentally pressed, so don't set a default.
+	// However we still register it because setting a hotkey can be useful in some workflows (when the
+	// gizmo gets in the way of shift-selecting multiple things).
 	ActionSet.RegisterAction(this, (int32)EStandardToolActions::BaseClientDefinedActionID + 5,
-		TEXT("ToggleTransformGizmoAKey"),
-		LOCTEXT("ToggleTransformGizmoUIName", "Toggle Transform Gizmo Visibility"),
-		LOCTEXT("ToggleTransformGizmoTooltip", "Toggle Transform Gizmo Visibility"),
-		EModifierKey::None, EKeys::R,
+		TEXT("ToggleGizmoVisibilityKey"),
+		LOCTEXT("ToggleGizmoVisibilityUIName", "Toggle Transform Gizmo Visibility"),
+		LOCTEXT("ToggleGizmoVisibilityTooltip", "Toggle the visibility of the transform gizmo"),
+		EModifierKey::None, EKeys::Invalid,
 		[this]() 
 		{
 			if (!CurrentActivity)
