ModelingTools: improve handling of invalid input selections to PolyEd Tool
Add UGeometrySelectionManager::GetActiveSelectionInfo, returns info about current selection w/o having to copy selection
Override UEditMeshPolygonsActionModeToolBuilder::CanBuildTool, use this new function to do specific per-action filtering to limit to supported input selection types
UEditMeshPolygonsSelectionModeToolBuilder::InitializeNewTool() now filters UI panels for current input selection if available, update Selection-palette PolyEd button to use this builder instead

diff --git a/Engine/Plugins/Editor/ModelingToolsEditorMode/Source/ModelingToolsEditorMode/Private/ModelingToolsEditorMode.cpp b/Engine/Plugins/Editor/ModelingToolsEditorMode/Source/ModelingToolsEditorMode/Private/ModelingToolsEditorMode.cpp
--- a/Engine/Plugins/Editor/ModelingToolsEditorMode/Source/ModelingToolsEditorMode/Private/ModelingToolsEditorMode.cpp
+++ b/Engine/Plugins/Editor/ModelingToolsEditorMode/Source/ModelingToolsEditorMode/Private/ModelingToolsEditorMode.cpp
@@ -771,7 +771,7 @@ void UModelingToolsEditorMode::Enter()
 	RegisterSelectionTool(ToolManagerCommands.BeginSelectionAction_Extrude, TEXT("BeginSelectionExtrudeTool"), NewObject<UExtrudeMeshSelectionToolBuilder>(), true, true);
 	RegisterSelectionTool(ToolManagerCommands.BeginSelectionAction_Offset, TEXT("BeginSelectionOffsetTool"), NewObject<UOffsetMeshSelectionToolBuilder>(), true, true);
 
-	RegisterSelectionTool(ToolManagerCommands.BeginPolyModelTool_PolyEd, TEXT("BeginSelectionPolyEdTool"), NewObject<UEditMeshPolygonsToolBuilder>(), true, false);
+	RegisterSelectionTool(ToolManagerCommands.BeginPolyModelTool_PolyEd, TEXT("BeginSelectionPolyEdTool"), NewObject<UEditMeshPolygonsSelectionModeToolBuilder>(), true, false);
 	RegisterSelectionTool(ToolManagerCommands.BeginPolyModelTool_TriSel, TEXT("BeginSelectionTriEdTool"), NewObject<UMeshSelectionToolBuilder>(), true, false);
 
 
diff --git a/Engine/Plugins/Runtime/MeshModelingToolset/Source/MeshModelingTools/Private/EditMeshPolygonsTool.cpp b/Engine/Plugins/Runtime/MeshModelingToolset/Source/MeshModelingTools/Private/EditMeshPolygonsTool.cpp
--- a/Engine/Plugins/Runtime/MeshModelingToolset/Source/MeshModelingTools/Private/EditMeshPolygonsTool.cpp
+++ b/Engine/Plugins/Runtime/MeshModelingToolset/Source/MeshModelingTools/Private/EditMeshPolygonsTool.cpp
@@ -50,6 +50,7 @@
 #include "TransformTypes.h"
 #include "Util/CompactMaps.h"
 #include "Selections/GeometrySelection.h"
+#include "Selection/GeometrySelectionManager.h"
 
 #include UE_INLINE_GENERATED_CPP_BY_NAME(EditMeshPolygonsTool)
 
@@ -133,6 +134,42 @@ void UEditMeshPolygonsToolBuilder::InitializeNewTool(USingleTargetWithSelectionT
 }
 
 
+bool UEditMeshPolygonsActionModeToolBuilder::CanBuildTool(const FToolBuilderState& SceneState) const
+{
+	if (UEditMeshPolygonsToolBuilder::CanBuildTool(SceneState))
+	{
+		if ( StartupAction == EEditMeshPolygonsToolActions::SimplifyByGroups
+			 || StartupAction == EEditMeshPolygonsToolActions::InsertEdge
+			 || StartupAction == EEditMeshPolygonsToolActions::InsertEdgeLoop )
+		{
+			return true;
+		}
+
+		if ( UGeometrySelectionManager* SelectionManager = SceneState.ToolManager->GetContextObjectStore()->FindContext<UGeometrySelectionManager>() )
+		{
+			EGeometryTopologyType TopologyType = EGeometryTopologyType::Triangle;
+			EGeometryElementType ElementType = EGeometryElementType::Face;
+			int NumTargets;
+			bool bIsEmpty = false;
+			SelectionManager->GetActiveSelectionInfo(TopologyType, ElementType, NumTargets, bIsEmpty);
+			if (StartupAction == EEditMeshPolygonsToolActions::Extrude
+				|| StartupAction == EEditMeshPolygonsToolActions::PushPull
+				|| StartupAction == EEditMeshPolygonsToolActions::Offset
+				|| StartupAction == EEditMeshPolygonsToolActions::Inset
+				|| StartupAction == EEditMeshPolygonsToolActions::Outset
+				|| StartupAction == EEditMeshPolygonsToolActions::CutFaces)
+			{
+				return (TopologyType == EGeometryTopologyType::Polygroup && ElementType == EGeometryElementType::Face && bIsEmpty == false);
+			}
+			else if (StartupAction == EEditMeshPolygonsToolActions::BevelAuto)
+			{
+				return (TopologyType == EGeometryTopologyType::Polygroup && ElementType != EGeometryElementType::Vertex && bIsEmpty == false);
+			}
+		}
+	}
+	return false;
+}
+
 
 void UEditMeshPolygonsActionModeToolBuilder::InitializeNewTool(USingleTargetWithSelectionTool* Tool, const FToolBuilderState& SceneState) const
 {
@@ -149,46 +186,101 @@ void UEditMeshPolygonsActionModeToolBuilder::InitializeNewTool(USingleTargetWith
 
 
 
+bool UEditMeshPolygonsSelectionModeToolBuilder::CanBuildTool(const FToolBuilderState& SceneState) const
+{
+	if (UEditMeshPolygonsToolBuilder::CanBuildTool(SceneState))
+	{
+		if ( UGeometrySelectionManager* SelectionManager = SceneState.ToolManager->GetContextObjectStore()->FindContext<UGeometrySelectionManager>() )
+		{
+			// if not actively selecting mesh components, tool can be started in 'standard' full-PolyEd mode
+			if (SelectionManager->GetMeshTopologyMode() == UGeometrySelectionManager::EMeshTopologyMode::None)
+			{
+				return true;
+			}
+			// otherwise can only start tool in sub-modes
+			EGeometryTopologyType TopologyType = EGeometryTopologyType::Triangle;
+			EGeometryElementType ElementType = EGeometryElementType::Face;
+			int NumTargets;
+			bool bIsEmpty = false;
+			SelectionManager->GetActiveSelectionInfo(TopologyType, ElementType, NumTargets, bIsEmpty);
+			if (TopologyType == EGeometryTopologyType::Polygroup)
+			{
+				return ElementType != EGeometryElementType::Vertex;
+			}
+		}
+	}
+	return false;
+}
+
 void UEditMeshPolygonsSelectionModeToolBuilder::InitializeNewTool(USingleTargetWithSelectionTool* Tool, const FToolBuilderState& SceneState) const
 {
 	UEditMeshPolygonsToolBuilder::InitializeNewTool(Tool, SceneState);
 	UEditMeshPolygonsTool* EditPolygonsTool = CastChecked<UEditMeshPolygonsTool>(Tool);
 
-	EEditMeshPolygonsToolSelectionMode UseMode = SelectionMode;
-	EditPolygonsTool->PostSetupFunction = [UseMode](UEditMeshPolygonsTool* PolyTool)
+	// if not actively selecting mesh components, start in full-PolyEd mode
+	if (UGeometrySelectionManager* SelectionManager = SceneState.ToolManager->GetContextObjectStore()->FindContext<UGeometrySelectionManager>())
 	{
-		PolyTool->SetToSelectionModeInterface();
+		if (SelectionManager->GetMeshTopologyMode() == UGeometrySelectionManager::EMeshTopologyMode::None)
+		{
+			return;
+		}
 
-		UPolygonSelectionMechanic* SelectionMechanic = PolyTool->SelectionMechanic;
-		UMeshTopologySelectionMechanicProperties* SelectionProps = SelectionMechanic->Properties;
-		SelectionProps->bSelectFaces = SelectionProps->bSelectEdges = SelectionProps->bSelectVertices = false;
-		SelectionProps->bSelectEdgeLoops = SelectionProps->bSelectEdgeRings = false;
+		// otherwise can only start tool in sub-modes
+		EGeometryTopologyType TopologyType = EGeometryTopologyType::Triangle;
+		EGeometryElementType ElementType = EGeometryElementType::Face;
+		int NumTargets;
+		bool bIsEmpty = false;
+		SelectionManager->GetActiveSelectionInfo(TopologyType, ElementType, NumTargets, bIsEmpty);
+		if (TopologyType != EGeometryTopologyType::Polygroup)
+		{
+			return;		// should not happen...
+		}
 
-		switch (UseMode)
+		EEditMeshPolygonsToolSelectionMode UseMode = EEditMeshPolygonsToolSelectionMode::Faces;
+		bool bIsEdgeSelection = false;
+		if (ElementType == EGeometryElementType::Edge)
 		{
-		default:
-		case EEditMeshPolygonsToolSelectionMode::Faces:
-			SelectionProps->bSelectFaces = true;
-			break;
-		case EEditMeshPolygonsToolSelectionMode::Edges:
-			SelectionProps->bSelectEdges = true;
-			break;
-		case EEditMeshPolygonsToolSelectionMode::Vertices:
-			SelectionProps->bSelectVertices = true;
-			break;
-		case EEditMeshPolygonsToolSelectionMode::Loops:
-			SelectionProps->bSelectEdges = true;
-			SelectionProps->bSelectEdgeLoops = true;
-			break;
-		case EEditMeshPolygonsToolSelectionMode::Rings:
-			SelectionProps->bSelectEdges = true;
-			SelectionProps->bSelectEdgeRings = true;
-			break;
-		case EEditMeshPolygonsToolSelectionMode::FacesEdgesVertices:
-			SelectionProps->bSelectFaces = SelectionProps->bSelectEdges = SelectionProps->bSelectVertices = true;
-			break;
+			bIsEdgeSelection = true;
+			UseMode = EEditMeshPolygonsToolSelectionMode::Edges;
 		}
-	};
+		else if (ElementType == EGeometryElementType::Vertex)
+		{
+			UseMode = EEditMeshPolygonsToolSelectionMode::Vertices;
+		}
+
+		EditPolygonsTool->PostSetupFunction = [UseMode, bIsEdgeSelection](UEditMeshPolygonsTool* PolyTool)
+		{
+			PolyTool->SetToolPropertySourceEnabled(PolyTool->EditActions, !bIsEdgeSelection);
+			PolyTool->SetToolPropertySourceEnabled(PolyTool->EditEdgeActions, bIsEdgeSelection);
+			PolyTool->SetToolPropertySourceEnabled(PolyTool->EditUVActions, !bIsEdgeSelection);
+			
+			PolyTool->SetToolPropertySourceEnabled(PolyTool->TopologyProperties, false);
+
+			UPolygonSelectionMechanic* SelectionMechanic = PolyTool->SelectionMechanic;
+			UMeshTopologySelectionMechanicProperties* SelectionProps = SelectionMechanic->Properties;
+			SelectionProps->bSelectFaces = SelectionProps->bSelectEdges = SelectionProps->bSelectVertices = false;
+			SelectionProps->bSelectEdgeLoops = SelectionProps->bSelectEdgeRings = false;
+
+			switch (UseMode)
+			{
+			default:
+			case EEditMeshPolygonsToolSelectionMode::Faces:
+				SelectionProps->bSelectFaces = true;
+				break;
+			case EEditMeshPolygonsToolSelectionMode::Edges:
+				SelectionProps->bSelectEdges = true;
+				break;
+			case EEditMeshPolygonsToolSelectionMode::Vertices:
+				SelectionProps->bSelectVertices = true;
+				break;
+			}
+
+			PolyTool->SetToolPropertySourceEnabled(SelectionProps, false);
+		};
+	}
+
+
+
 }
 
 
diff --git a/Engine/Plugins/Runtime/MeshModelingToolset/Source/MeshModelingTools/Public/EditMeshPolygonsTool.h b/Engine/Plugins/Runtime/MeshModelingToolset/Source/MeshModelingTools/Public/EditMeshPolygonsTool.h
--- a/Engine/Plugins/Runtime/MeshModelingToolset/Source/MeshModelingTools/Public/EditMeshPolygonsTool.h
+++ b/Engine/Plugins/Runtime/MeshModelingToolset/Source/MeshModelingTools/Public/EditMeshPolygonsTool.h
@@ -153,6 +153,7 @@ class MESHMODELINGTOOLS_API UEditMeshPolygonsActionModeToolBuilder : public UEdi
 public:
 	EEditMeshPolygonsToolActions StartupAction = EEditMeshPolygonsToolActions::Extrude;
 
+	virtual bool CanBuildTool(const FToolBuilderState& SceneState) const override;
 	virtual void InitializeNewTool(USingleTargetWithSelectionTool* Tool, const FToolBuilderState& SceneState) const override;
 };
 
@@ -174,6 +175,7 @@ class MESHMODELINGTOOLS_API UEditMeshPolygonsSelectionModeToolBuilder : public U
 public:
 	EEditMeshPolygonsToolSelectionMode SelectionMode = EEditMeshPolygonsToolSelectionMode::Faces;
 
+	virtual bool CanBuildTool(const FToolBuilderState& SceneState) const override;
 	virtual void InitializeNewTool(USingleTargetWithSelectionTool* Tool, const FToolBuilderState& SceneState) const override;
 };
 
diff --git a/Engine/Plugins/Runtime/MeshModelingToolset/Source/ModelingComponents/Private/Selection/DynamicMeshSelector.cpp b/Engine/Plugins/Runtime/MeshModelingToolset/Source/ModelingComponents/Private/Selection/DynamicMeshSelector.cpp
--- a/Engine/Plugins/Runtime/MeshModelingToolset/Source/ModelingComponents/Private/Selection/DynamicMeshSelector.cpp
+++ b/Engine/Plugins/Runtime/MeshModelingToolset/Source/ModelingComponents/Private/Selection/DynamicMeshSelector.cpp
@@ -185,7 +185,8 @@ void FBaseDynamicMeshSelector::UpdateSelectionFromSelection(
 		return;
 	}
 
-	ensure(false);		// todo: support this type of convert selection
+	// todo: support additional conversion types here
+	//ensure(false);
 }
 
 
diff --git a/Engine/Plugins/Runtime/MeshModelingToolset/Source/ModelingComponents/Private/Selection/GeometrySelectionManager.cpp b/Engine/Plugins/Runtime/MeshModelingToolset/Source/ModelingComponents/Private/Selection/GeometrySelectionManager.cpp
--- a/Engine/Plugins/Runtime/MeshModelingToolset/Source/ModelingComponents/Private/Selection/GeometrySelectionManager.cpp
+++ b/Engine/Plugins/Runtime/MeshModelingToolset/Source/ModelingComponents/Private/Selection/GeometrySelectionManager.cpp
@@ -240,7 +240,14 @@ bool UGeometrySelectionManager::HasSelection() const
 }
 
 
-
+void UGeometrySelectionManager::GetActiveSelectionInfo(EGeometryTopologyType& TopologyTypeOut, EGeometryElementType& ElementTypeOut, int& NumTargetsOut, bool& bIsEmpty) const
+{
+	FGeometrySelectionHitQueryConfig Config = GetCurrentSelectionQueryConfig();
+	TopologyTypeOut = Config.TopologyType;
+	ElementTypeOut = Config.ElementType;
+	NumTargetsOut = ActiveTargetReferences.Num();
+	bIsEmpty = (NumTargetsOut == 0) || ActiveTargetReferences[0]->Selection.IsEmpty();
+}
 
 
 class FGeometrySelectionManager_ActiveTargetsChange : public FToolCommandChange
diff --git a/Engine/Plugins/Runtime/MeshModelingToolset/Source/ModelingComponents/Public/Selection/GeometrySelectionManager.h b/Engine/Plugins/Runtime/MeshModelingToolset/Source/ModelingComponents/Public/Selection/GeometrySelectionManager.h
--- a/Engine/Plugins/Runtime/MeshModelingToolset/Source/ModelingComponents/Public/Selection/GeometrySelectionManager.h
+++ b/Engine/Plugins/Runtime/MeshModelingToolset/Source/ModelingComponents/Public/Selection/GeometrySelectionManager.h
@@ -257,6 +257,11 @@ public:
 	/** @return true if there is an active element selection */
 	virtual bool HasSelection() const;
 
+	/** 
+	 * Get avaialble information about the active selection/state
+	 */
+	virtual void GetActiveSelectionInfo(EGeometryTopologyType& TopologyTypeOut, EGeometryElementType& ElementTypeOut, int& NumTargetsOut, bool& bIsEmpty) const;
+
 	/** @return a world-space bounding box for the active element selection */
 	virtual bool GetSelectionBounds(FGeometrySelectionBounds& BoundsOut) const;
 	/** @return a 3D transformation frame suitable for use with the active element selection */
