ModelingTools: Adjust the behavior of the Mesh Visualizer settings to properly disable when the Visible property is checked off.

diff --git a/Engine/Plugins/Runtime/MeshModelingToolset/Source/ModelingComponents/Public/Drawing/MeshElementsVisualizer.h b/Engine/Plugins/Runtime/MeshModelingToolset/Source/ModelingComponents/Public/Drawing/MeshElementsVisualizer.h
--- a/Engine/Plugins/Runtime/MeshModelingToolset/Source/ModelingComponents/Public/Drawing/MeshElementsVisualizer.h
+++ b/Engine/Plugins/Runtime/MeshModelingToolset/Source/ModelingComponents/Public/Drawing/MeshElementsVisualizer.h
@@ -24,55 +24,55 @@ public:
 	bool bVisible = true;
 
 	/** Should mesh wireframe be shown */
-	UPROPERTY(EditAnywhere, Category = MeshElementVisualization)
+	UPROPERTY(EditAnywhere, Category = MeshElementVisualization, meta = (EditCondition = "bVisible" ))
 	bool bShowWireframe = false;
 
 	/** Should mesh boundary edges be shown */
-	UPROPERTY(EditAnywhere, Category = MeshElementVisualization)
+	UPROPERTY(EditAnywhere, Category = MeshElementVisualization, meta = (EditCondition = "bVisible"))
 	bool bShowBorders = true;
 
 	/** Should mesh uv seam edges be shown */
-	UPROPERTY(EditAnywhere, Category = MeshElementVisualization)
+	UPROPERTY(EditAnywhere, Category = MeshElementVisualization, meta = (EditCondition = "bVisible"))
 	bool bShowUVSeams = true;
 
 	/** Should mesh normal seam edges be shown */
-	UPROPERTY(EditAnywhere, Category = MeshElementVisualization)
+	UPROPERTY(EditAnywhere, Category = MeshElementVisualization, meta = (EditCondition = "bVisible"))
 	bool bShowNormalSeams = true;
 
 	/** Should mesh color seam edges be shown */
-	UPROPERTY(EditAnywhere, Category = MeshElementVisualization)
+	UPROPERTY(EditAnywhere, Category = MeshElementVisualization, meta = (EditCondition = "bVisible"))
 	bool bShowColorSeams = true;
 
 	/** multiplier on edge thicknesses */
-	UPROPERTY(EditAnywhere, Category = MeshElementVisualization, AdvancedDisplay, meta = (UIMin = 0.1, UIMax = 10.0))
+	UPROPERTY(EditAnywhere, Category = MeshElementVisualization, AdvancedDisplay, meta = (UIMin = 0.1, UIMax = 10.0, EditCondition = "bVisible"))
 	float ThicknessScale = 1.0;
 
 	/** Color of mesh wireframe */
-	UPROPERTY(EditAnywhere, Category = MeshElementVisualization, AdvancedDisplay)
+	UPROPERTY(EditAnywhere, Category = MeshElementVisualization, AdvancedDisplay, meta = (EditCondition = "bVisible"))
 	FColor WireframeColor = FColor(128, 128, 128);
 
 	/** Color of mesh boundary edges */
-	UPROPERTY(EditAnywhere, Category = MeshElementVisualization, AdvancedDisplay)
+	UPROPERTY(EditAnywhere, Category = MeshElementVisualization, AdvancedDisplay, meta = (EditCondition = "bVisible"))
 	FColor BoundaryEdgeColor = FColor(245, 15, 15);
 
 	/** Color of mesh UV seam edges */
-	UPROPERTY(EditAnywhere, Category = MeshElementVisualization, AdvancedDisplay)
+	UPROPERTY(EditAnywhere, Category = MeshElementVisualization, AdvancedDisplay, meta = (EditCondition = "bVisible"))
 	FColor UVSeamColor = FColor(240, 160, 15);
 
 	/** Color of mesh normal seam edges */
-	UPROPERTY(EditAnywhere, Category = MeshElementVisualization, AdvancedDisplay)
+	UPROPERTY(EditAnywhere, Category = MeshElementVisualization, AdvancedDisplay, meta = (EditCondition = "bVisible"))
 	FColor NormalSeamColor = FColor(128, 128, 240);
 
 	/** Color of mesh color seam edges */
-	UPROPERTY(EditAnywhere, Category = MeshElementVisualization, AdvancedDisplay)
+	UPROPERTY(EditAnywhere, Category = MeshElementVisualization, AdvancedDisplay, meta = (EditCondition = "bVisible"))
 	FColor ColorSeamColor = FColor(46, 204, 113);
 
 	/** depth bias used to slightly shift depth of lines */
-	UPROPERTY(EditAnywhere, Category = MeshElementVisualization, AdvancedDisplay, meta = (UIMin = -2.0, UIMax = 2.0))
+	UPROPERTY(EditAnywhere, Category = MeshElementVisualization, AdvancedDisplay, meta = (UIMin = -2.0, UIMax = 2.0, EditCondition = "bVisible"))
 	float DepthBias = 0.2;
 
 	/** If true, the depth bias will be adjusted depending on the diagonal length of the mesh bounding box (smaller for smaller meshes). */
-	UPROPERTY(EditAnywhere, Category = MeshElementVisualization, AdvancedDisplay)
+	UPROPERTY(EditAnywhere, Category = MeshElementVisualization, AdvancedDisplay, meta = (EditCondition = "bVisible"))
 	bool bAdjustDepthBiasUsingMeshSize = true;
 };
 
