More parent include path fixes

#rnx

diff --git a/Engine/Plugins/Runtime/HairStrands/Source/HairStrandsEditor/HairStrandsEditor.Build.cs b/Engine/Plugins/Runtime/HairStrands/Source/HairStrandsEditor/HairStrandsEditor.Build.cs
--- a/Engine/Plugins/Runtime/HairStrands/Source/HairStrandsEditor/HairStrandsEditor.Build.cs
+++ b/Engine/Plugins/Runtime/HairStrands/Source/HairStrandsEditor/HairStrandsEditor.Build.cs
@@ -37,8 +37,15 @@ namespace UnrealBuildTool.Rules
 					"HairCardGeneratorFramework",
 					"CommonMenuExtensions",
 				});
+
+			PrivateIncludePathModuleNames.AddRange(
+				new string[]
+				{
+					"Analytics"
+				});
+
 			AddEngineThirdPartyPrivateStaticDependencies(Target,
-			 "FBX"
+				"FBX"
 			);
 		}
 	}
diff --git a/Engine/Plugins/Runtime/MeshModelingToolset/Source/ModelingComponents/Private/ConversionUtils/SkinnedMeshToDynamicMesh.cpp b/Engine/Plugins/Runtime/MeshModelingToolset/Source/ModelingComponents/Private/ConversionUtils/SkinnedMeshToDynamicMesh.cpp
--- a/Engine/Plugins/Runtime/MeshModelingToolset/Source/ModelingComponents/Private/ConversionUtils/SkinnedMeshToDynamicMesh.cpp
+++ b/Engine/Plugins/Runtime/MeshModelingToolset/Source/ModelingComponents/Private/ConversionUtils/SkinnedMeshToDynamicMesh.cpp
@@ -10,7 +10,7 @@
 #include "Engine/SkinnedAssetCommon.h"
 #include "Engine/SkinnedAsset.h"
 #include "SkeletalRenderPublic.h"
-#include "MeshConversion/Public/ToDynamicMesh.h"
+#include "ToDynamicMesh.h"
 #include "Rendering/SkeletalMeshRenderData.h"
 
 
