ModelingTools: Updated the Simplify tool's inherited properties to conditionally disable the PreventNormalFlips and PreventTinyTriangles (depending on the algorithm used). This also enables support for the PreventTinyTriangles and PreventNormalFlips under certain conditions, the latter of which had been passed into the underlying operation but never used before.

diff --git a/Engine/Plugins/Experimental/MeshModelingToolsetExp/Source/MeshModelingToolsEditorOnlyExp/Private/SimplifyMeshTool.cpp b/Engine/Plugins/Experimental/MeshModelingToolsetExp/Source/MeshModelingToolsEditorOnlyExp/Private/SimplifyMeshTool.cpp
--- a/Engine/Plugins/Experimental/MeshModelingToolsetExp/Source/MeshModelingToolsEditorOnlyExp/Private/SimplifyMeshTool.cpp
+++ b/Engine/Plugins/Experimental/MeshModelingToolsetExp/Source/MeshModelingToolsEditorOnlyExp/Private/SimplifyMeshTool.cpp
@@ -198,6 +198,7 @@ TUniquePtr<FDynamicMeshOperator> USimplifyMeshTool::MakeNewOperator()
 	Op->bPreventNormalFlips = SimplifyProperties->bPreventNormalFlips;
 	Op->bPreserveSharpEdges = SimplifyProperties->bPreserveSharpEdges;
 	Op->bAllowSeamCollapse = !SimplifyProperties->bPreserveSharpEdges;
+	Op->bPreventTinyTriangles = SimplifyProperties->bPreventTinyTriangles;
 	Op->bReproject = SimplifyProperties->bReproject;
 	Op->SimplifierType = SimplifyProperties->SimplifierType;
 	Op->TargetCount = ( SimplifyProperties->TargetMode == ESimplifyTargetType::VertexCount) ?  SimplifyProperties->TargetVertexCount : SimplifyProperties->TargetTriangleCount;
diff --git a/Engine/Plugins/Experimental/MeshModelingToolsetExp/Source/MeshModelingToolsEditorOnlyExp/Public/SimplifyMeshTool.h b/Engine/Plugins/Experimental/MeshModelingToolsetExp/Source/MeshModelingToolsEditorOnlyExp/Public/SimplifyMeshTool.h
--- a/Engine/Plugins/Experimental/MeshModelingToolsetExp/Source/MeshModelingToolsEditorOnlyExp/Public/SimplifyMeshTool.h
+++ b/Engine/Plugins/Experimental/MeshModelingToolsetExp/Source/MeshModelingToolsEditorOnlyExp/Public/SimplifyMeshTool.h
@@ -96,6 +96,25 @@ public:
 	/** Enable projection back to input mesh */
 	UPROPERTY(EditAnywhere, Category = Options, AdvancedDisplay)
 	bool bReproject;
+
+protected:
+	
+	virtual bool IsPreventNormalFlipsEnabled() const override
+	{ 
+		return SimplifierType == ESimplifyType::QEM ||
+		       SimplifierType == ESimplifyType::MinimalExistingVertex || 
+		       SimplifierType == ESimplifyType::Attribute || 
+		       SimplifierType == ESimplifyType::MinimalPlanar;
+	}
+
+	virtual bool IsPreventTinyTrianglesEnabled() const override
+	{
+		return SimplifierType == ESimplifyType::QEM ||
+			   SimplifierType == ESimplifyType::MinimalExistingVertex ||
+			   SimplifierType == ESimplifyType::Attribute ||
+			   SimplifierType == ESimplifyType::MinimalPlanar;
+	}
+
 };
 
 
diff --git a/Engine/Plugins/Experimental/MeshModelingToolsetExp/Source/MeshModelingToolsExp/Public/Properties/RemeshProperties.h b/Engine/Plugins/Experimental/MeshModelingToolsetExp/Source/MeshModelingToolsExp/Public/Properties/RemeshProperties.h
--- a/Engine/Plugins/Experimental/MeshModelingToolsetExp/Source/MeshModelingToolsExp/Public/Properties/RemeshProperties.h
+++ b/Engine/Plugins/Experimental/MeshModelingToolsetExp/Source/MeshModelingToolsExp/Public/Properties/RemeshProperties.h
@@ -60,12 +60,20 @@ public:
 	EMaterialBoundaryConstraint MaterialBoundaryConstraint = EMaterialBoundaryConstraint::Free;
 
 	/** Prevent normal flips */
-	UPROPERTY(EditAnywhere, Category = "Remeshing|Edge Constraints")
+	UPROPERTY(EditAnywhere, Category = "Remeshing|Edge Constraints", meta=(EditCondition = "IsPreventNormalFlipsEnabled") )
 	bool bPreventNormalFlips = false;
 
 	/** Prevent introduction of tiny triangles or slivers */
-	UPROPERTY(EditAnywhere, Category = "Remeshing|Edge Constraints")
+	UPROPERTY(EditAnywhere, Category = "Remeshing|Edge Constraints", meta=(EditCondition = "IsPreventTinyTrianglesEnabled") )
 	bool bPreventTinyTriangles = false;
+
+protected:
+
+	UFUNCTION()
+	virtual bool IsPreventNormalFlipsEnabled() const { return true; }
+
+	UFUNCTION()
+	virtual bool IsPreventTinyTrianglesEnabled() const { return true; }
 };
 
 UCLASS()
diff --git a/Engine/Plugins/Runtime/GeometryProcessing/Source/DynamicMesh/Private/MeshSimplification.cpp b/Engine/Plugins/Runtime/GeometryProcessing/Source/DynamicMesh/Private/MeshSimplification.cpp
--- a/Engine/Plugins/Runtime/GeometryProcessing/Source/DynamicMesh/Private/MeshSimplification.cpp
+++ b/Engine/Plugins/Runtime/GeometryProcessing/Source/DynamicMesh/Private/MeshSimplification.cpp
@@ -1382,6 +1382,12 @@ ESimplificationResult TMeshSimplification<QuadricErrorType>::CollapseEdge(int ed
 		return ESimplificationResult::Ignored_CreatesFlip;
 	}
 
+	if (bPreventTinyTriangles && (CheckIfCollapseCreatesTinyTriangle(a, b, vNewPos, t0, t1) || CheckIfCollapseCreatesTinyTriangle(b, a, vNewPos, t0, t1)))
+	{
+		ProfileEndCollapse();
+		return ESimplificationResult::Ignored_CreatesTinyTriangle;
+	}
+
 	// lots of cases where we cannot collapse, but we should just let
 	// Mesh sort that out, right?
 	COUNT_COLLAPSES++;
diff --git a/Engine/Plugins/Runtime/GeometryProcessing/Source/DynamicMesh/Public/MeshSimplification.h b/Engine/Plugins/Runtime/GeometryProcessing/Source/DynamicMesh/Public/MeshSimplification.h
--- a/Engine/Plugins/Runtime/GeometryProcessing/Source/DynamicMesh/Public/MeshSimplification.h
+++ b/Engine/Plugins/Runtime/GeometryProcessing/Source/DynamicMesh/Public/MeshSimplification.h
@@ -27,7 +27,8 @@ enum class ESimplificationResult
 	Failed_OpNotSuccessful = 6,
 	Failed_NotAnEdge = 7,
 	Failed_IsolatedTriangle = 8,
-	Failed_GeometricDeviation = 9
+	Failed_GeometricDeviation = 9,
+	Ignored_CreatesTinyTriangle = 10
 };
 
 /**
@@ -76,6 +77,11 @@ public:
 	/** if true, we allow UV and Normal seams to collapse during simplification.*/
 	bool bAllowSeamCollapse = true;
 
+	/** Controls whether we disallow creation of triangles with small areas inside edge operations.
+	 * This is moderately expensive and in some cases can result in lower-quality meshes. Disabled by default.
+	 */
+	bool bPreventTinyTriangles = false;
+
 	/** When using the constraint system, these options will apply to the appropriate boundaries. */
 	EEdgeRefineFlags MeshBoundaryConstraint, GroupBoundaryConstraint, MaterialBoundaryConstraint;
 
diff --git a/Engine/Plugins/Runtime/MeshModelingToolset/Source/ModelingOperatorsEditorOnly/Private/CleaningOps/SimplifyMeshOp.cpp b/Engine/Plugins/Runtime/MeshModelingToolset/Source/ModelingOperatorsEditorOnly/Private/CleaningOps/SimplifyMeshOp.cpp
--- a/Engine/Plugins/Runtime/MeshModelingToolset/Source/ModelingOperatorsEditorOnly/Private/CleaningOps/SimplifyMeshOp.cpp
+++ b/Engine/Plugins/Runtime/MeshModelingToolset/Source/ModelingOperatorsEditorOnly/Private/CleaningOps/SimplifyMeshOp.cpp
@@ -32,7 +32,7 @@ void ComputeSimplify(FDynamicMesh3* TargetMesh, const bool bReproject,
 					 EEdgeRefineFlags MeshBoundaryConstraint,
 					 EEdgeRefineFlags GroupBoundaryConstraint,
 					 EEdgeRefineFlags MaterialBoundaryConstraint,
-					 bool bPreserveSharpEdges, bool bAllowSeamCollapse,
+					 bool bPreserveSharpEdges, bool bAllowSeamCollapse, bool bPreventNormalFlips, bool bPreventTinyTriangles,
 					 const ESimplifyTargetType TargetMode,
 					 const float TargetPercentage, const int TargetCount, const float TargetEdgeLength,
 					 const float AngleThreshold,
@@ -121,6 +121,11 @@ void ComputeSimplify(FDynamicMesh3* TargetMesh, const bool bReproject,
 	{
 		Reducer.SimplifyToMinimalPlanar(AngleThreshold);
 	}
+
+	if (!bPreventNormalFlips)
+	{
+		Reducer.SetEdgeFlipTolerance(-1.1);
+	}
 }
 
 
@@ -150,7 +155,7 @@ void FSimplifyMeshOp::CalculateResult(FProgressCancel* Progress)
 											MeshBoundaryConstraint,
 											GroupBoundaryConstraint,
 											MaterialBoundaryConstraint,
-											bPreserveSharpEdges, bAllowSeamCollapse,
+											bPreserveSharpEdges, bAllowSeamCollapse, bPreventNormalFlips, bPreventTinyTriangles,
 											TargetMode, TargetPercentage, TargetCount, TargetEdgeLength, MinimalPlanarAngleThresh,
 											FQEMSimplification::ESimplificationCollapseModes::MinimalQuadricPositionError, bUseQuadricMemory, 
 											UseGeometricTolerance);
@@ -167,7 +172,7 @@ void FSimplifyMeshOp::CalculateResult(FProgressCancel* Progress)
 													MeshBoundaryConstraint,
 													GroupBoundaryConstraint,
 													MaterialBoundaryConstraint,
-													bPreserveSharpEdges, bAllowSeamCollapse,
+													bPreserveSharpEdges, bAllowSeamCollapse, bPreventNormalFlips, bPreventTinyTriangles,
 													TargetMode, TargetPercentage, TargetCount, TargetEdgeLength, MinimalPlanarAngleThresh,
 													FAttrMeshSimplification::ESimplificationCollapseModes::MinimalQuadricPositionError, bUseQuadricMemory,
 													UseGeometricTolerance);
@@ -184,7 +189,7 @@ void FSimplifyMeshOp::CalculateResult(FProgressCancel* Progress)
 			MeshBoundaryConstraint,
 			GroupBoundaryConstraint,
 			MaterialBoundaryConstraint,
-			bPreserveSharpEdges, bAllowSeamCollapse,
+			bPreserveSharpEdges, bAllowSeamCollapse, bPreventNormalFlips, bPreventTinyTriangles,
 			ESimplifyTargetType::MinimalPlanar, TargetPercentage, TargetCount, TargetEdgeLength, MinimalPlanarAngleThresh,
 			FQEMSimplification::ESimplificationCollapseModes::MinimalQuadricPositionError, bUseQuadricMemory,
 			UseGeometricTolerance);
@@ -197,7 +202,7 @@ void FSimplifyMeshOp::CalculateResult(FProgressCancel* Progress)
 			MeshBoundaryConstraint,
 			GroupBoundaryConstraint,
 			MaterialBoundaryConstraint,
-			bPreserveSharpEdges, bAllowSeamCollapse,
+			bPreserveSharpEdges, bAllowSeamCollapse, bPreventNormalFlips, bPreventTinyTriangles,
 			TargetMode, TargetPercentage, TargetCount, TargetEdgeLength, MinimalPlanarAngleThresh, 
 			FQEMSimplification::ESimplificationCollapseModes::MinimalExistingVertexError, bUseQuadricMemory,
 			UseGeometricTolerance);
diff --git a/Engine/Plugins/Runtime/MeshModelingToolset/Source/ModelingOperatorsEditorOnly/Public/CleaningOps/SimplifyMeshOp.h b/Engine/Plugins/Runtime/MeshModelingToolset/Source/ModelingOperatorsEditorOnly/Public/CleaningOps/SimplifyMeshOp.h
--- a/Engine/Plugins/Runtime/MeshModelingToolset/Source/ModelingOperatorsEditorOnly/Public/CleaningOps/SimplifyMeshOp.h
+++ b/Engine/Plugins/Runtime/MeshModelingToolset/Source/ModelingOperatorsEditorOnly/Public/CleaningOps/SimplifyMeshOp.h
@@ -74,7 +74,7 @@ public:
 	ESimplifyType SimplifierType;
 	int TargetPercentage, TargetCount;
 	float TargetEdgeLength;
-	bool bDiscardAttributes, bReproject, bPreventNormalFlips, bPreserveSharpEdges, bAllowSeamCollapse;
+	bool bDiscardAttributes, bReproject, bPreventNormalFlips, bPreserveSharpEdges, bAllowSeamCollapse, bPreventTinyTriangles;
 	// When true, result will have attributes object regardless of whether attributes 
 	// were discarded or present initially.
 	bool bResultMustHaveAttributesEnabled = false;
