Make mesh elements visualizer distinguish between normal and tangent seams like the mesh inspector does

diff --git a/Engine/Plugins/Runtime/MeshModelingToolset/Source/ModelingComponents/Private/Drawing/MeshElementsVisualizer.cpp b/Engine/Plugins/Runtime/MeshModelingToolset/Source/ModelingComponents/Private/Drawing/MeshElementsVisualizer.cpp
--- a/Engine/Plugins/Runtime/MeshModelingToolset/Source/ModelingComponents/Private/Drawing/MeshElementsVisualizer.cpp
+++ b/Engine/Plugins/Runtime/MeshModelingToolset/Source/ModelingComponents/Private/Drawing/MeshElementsVisualizer.cpp
@@ -69,8 +69,8 @@ public:
 		}
 		if (Mesh.HasAttributes())
 		{
-			bool bIsUVSeam = false, bIsNormalSeam = false, bIsColorSeam = false;
-			if (Mesh.Attributes()->IsSeamEdge(EdgeIndex, bIsUVSeam, bIsNormalSeam, bIsColorSeam))
+			bool bIsUVSeam = false, bIsNormalSeam = false, bIsColorSeam = false, bIsTangentSeam = false;
+			if (Mesh.Attributes()->IsSeamEdge(EdgeIndex, bIsUVSeam, bIsNormalSeam, bIsColorSeam, bIsTangentSeam))
 			{
 				if (bIsUVSeam)
 				{
@@ -80,6 +80,10 @@ public:
 				{
 					EdgeType |= (int32)EMeshEdgeType::NormalSeam;
 				}
+				if (bIsTangentSeam)
+				{
+					EdgeType |= (int32)EMeshEdgeType::TangentSeam;
+				}
 				if (bIsColorSeam)
 				{
 					EdgeType |= (int32)EMeshEdgeType::ColorSeam;
@@ -122,6 +126,7 @@ void UMeshElementsVisualizer::OnCreated()
 	Settings->WatchProperty(Settings->bShowBorders, [this](bool) { bSettingsModified = true; });
 	Settings->WatchProperty(Settings->bShowUVSeams, [this](bool) { bSettingsModified = true; });
 	Settings->WatchProperty(Settings->bShowNormalSeams, [this](bool) { bSettingsModified = true; });
+	Settings->WatchProperty(Settings->bShowTangentSeams, [this](bool) { bSettingsModified = true; });
 	Settings->WatchProperty(Settings->bShowColorSeams, [this](bool){ bSettingsModified = true; });
 	Settings->WatchProperty(Settings->ThicknessScale, [this](float) { bSettingsModified = true; });
 	Settings->WatchProperty(Settings->DepthBias, [this](float) { bSettingsModified = true; });
@@ -132,6 +137,7 @@ void UMeshElementsVisualizer::OnCreated()
 	Settings->WatchProperty(Settings->BoundaryEdgeColor, [this](FColor) { bSettingsModified = true; });
 	Settings->WatchProperty(Settings->UVSeamColor, [this](FColor) { bSettingsModified = true; });
 	Settings->WatchProperty(Settings->NormalSeamColor, [this](FColor) { bSettingsModified = true; });
+	Settings->WatchProperty(Settings->TangentSeamColor, [this](FColor) { bSettingsModified = true; });
 	Settings->WatchProperty(Settings->ColorSeamColor, [this](FColor) { bSettingsModified = true; });
 	bSettingsModified = false;
 
@@ -195,12 +201,14 @@ void UMeshElementsVisualizer::UpdateVisibility()
 	WireframeComponent->bEnableBoundaryEdges = Settings->bShowBorders;
 	WireframeComponent->bEnableUVSeams = Settings->bShowUVSeams;
 	WireframeComponent->bEnableNormalSeams = Settings->bShowNormalSeams;
+	WireframeComponent->bEnableTangentSeams = Settings->bShowTangentSeams;
 	WireframeComponent->bEnableColorSeams = Settings->bShowColorSeams;
 
 	WireframeComponent->WireframeColor = Settings->WireframeColor;
 	WireframeComponent->BoundaryEdgeColor = Settings->BoundaryEdgeColor;
 	WireframeComponent->UVSeamColor = Settings->UVSeamColor;
 	WireframeComponent->NormalSeamColor = Settings->NormalSeamColor;
+	WireframeComponent->TangentSeamColor = Settings->TangentSeamColor;
 	WireframeComponent->ColorSeamColor = Settings->ColorSeamColor;
 
 	WireframeComponent->UpdateWireframe();
diff --git a/Engine/Plugins/Runtime/MeshModelingToolset/Source/ModelingComponents/Private/Drawing/MeshWireframeComponent.cpp b/Engine/Plugins/Runtime/MeshModelingToolset/Source/ModelingComponents/Private/Drawing/MeshWireframeComponent.cpp
--- a/Engine/Plugins/Runtime/MeshModelingToolset/Source/ModelingComponents/Private/Drawing/MeshWireframeComponent.cpp
+++ b/Engine/Plugins/Runtime/MeshModelingToolset/Source/ModelingComponents/Private/Drawing/MeshWireframeComponent.cpp
@@ -58,6 +58,7 @@ public:
 		const bool bEnableBoundaryEdges = Component->bEnableBoundaryEdges;
 		const bool bEnableUVSeams = Component->bEnableUVSeams;
 		const bool bEnableNormalSeams = Component->bEnableNormalSeams;
+		const bool bEnableTangentSeams = Component->bEnableTangentSeams;
 		const bool bEnableColorSeams = Component->bEnableColorSeams;
 
 		// Setup edge batches to process in parallel.
@@ -94,6 +95,7 @@ public:
 					(bEnableBoundaryEdges && (static_cast<int>(EdgeType) & static_cast<int>(IMeshWireframeSource::EMeshEdgeType::MeshBoundary)) != 0) ||
 					(bEnableUVSeams       && (static_cast<int>(EdgeType) & static_cast<int>(IMeshWireframeSource::EMeshEdgeType::UVSeam      )) != 0) ||
 					(bEnableNormalSeams   && (static_cast<int>(EdgeType) & static_cast<int>(IMeshWireframeSource::EMeshEdgeType::NormalSeam  )) != 0) ||
+					(bEnableTangentSeams  && (static_cast<int>(EdgeType) & static_cast<int>(IMeshWireframeSource::EMeshEdgeType::TangentSeam )) != 0) ||
 					(bEnableColorSeams    && (static_cast<int>(EdgeType) & static_cast<int>(IMeshWireframeSource::EMeshEdgeType::ColorSeam   )) != 0))
 				{
 					VisibleEdgeBatch.Add(UE::Geometry::FIndex3i(VertIndexA, VertIndexB, static_cast<int>(EdgeType)));
@@ -144,6 +146,8 @@ public:
 		const float UVSeamThickness = Component->ThicknessScale * Component->UVSeamThickness;
 		const FColor NormalSeamColor = FLinearColor::FromSRGBColor(Component->NormalSeamColor).ToFColor(false);
 		const float NormalSeamThickness = Component->ThicknessScale * Component->NormalSeamThickness;
+		const FColor TangentSeamColor = FLinearColor::FromSRGBColor(Component->TangentSeamColor).ToFColor(false);
+		const float TangentSeamThickness = Component->ThicknessScale * Component->TangentSeamThickness;
 		const FColor ColorSeamColor = FLinearColor::FromSRGBColor(Component->ColorSeamColor).ToFColor(false);
 		const float ColorSeamThickness = Component->ThicknessScale * Component->ColorSeamThickness;
 
@@ -187,6 +191,11 @@ public:
 						UseThickness = (bIsBoundaryEdge) ? BoundaryEdgeThickness : NormalSeamThickness;
 						UseColor = NormalSeamColor;
 					}
+					else if (bEnableTangentSeams && (static_cast<int>(EdgeType) & static_cast<int>(IMeshWireframeSource::EMeshEdgeType::TangentSeam)) != 0)
+					{
+						UseThickness = (bIsBoundaryEdge) ? BoundaryEdgeThickness : TangentSeamThickness;
+						UseColor = TangentSeamColor;
+					}
 					else if (bEnableColorSeams && (static_cast<int>(EdgeType) & static_cast<int>(IMeshWireframeSource::EMeshEdgeType::ColorSeam)) != 0)
 					{
 						UseThickness = (bIsBoundaryEdge) ? BoundaryEdgeThickness : ColorSeamThickness;
diff --git a/Engine/Plugins/Runtime/MeshModelingToolset/Source/ModelingComponents/Public/Drawing/MeshElementsVisualizer.h b/Engine/Plugins/Runtime/MeshModelingToolset/Source/ModelingComponents/Public/Drawing/MeshElementsVisualizer.h
--- a/Engine/Plugins/Runtime/MeshModelingToolset/Source/ModelingComponents/Public/Drawing/MeshElementsVisualizer.h
+++ b/Engine/Plugins/Runtime/MeshModelingToolset/Source/ModelingComponents/Public/Drawing/MeshElementsVisualizer.h
@@ -39,6 +39,10 @@ public:
 	UPROPERTY(EditAnywhere, Category = MeshElementVisualization, meta = (EditCondition = "bVisible"))
 	bool bShowNormalSeams = true;
 
+	/** Should mesh tangent seam edges be shown */
+	UPROPERTY(EditAnywhere, Category = MeshElementVisualization, meta = (EditCondition = "bVisible"))
+	bool bShowTangentSeams = true;
+
 	/** Should mesh color seam edges be shown */
 	UPROPERTY(EditAnywhere, Category = MeshElementVisualization, meta = (EditCondition = "bVisible"))
 	bool bShowColorSeams = true;
@@ -63,6 +67,10 @@ public:
 	UPROPERTY(EditAnywhere, Category = MeshElementVisualization, AdvancedDisplay, meta = (EditCondition = "bVisible"))
 	FColor NormalSeamColor = FColor(128, 128, 240);
 
+	/** Color of mesh tangent seam edges */
+	UPROPERTY(EditAnywhere, Category = MeshElementVisualization, AdvancedDisplay, meta = (EditCondition = "bVisible"))
+	FColor TangentSeamColor = FColor(64, 240, 240);
+
 	/** Color of mesh color seam edges */
 	UPROPERTY(EditAnywhere, Category = MeshElementVisualization, AdvancedDisplay, meta = (EditCondition = "bVisible"))
 	FColor ColorSeamColor = FColor(46, 204, 113);
diff --git a/Engine/Plugins/Runtime/MeshModelingToolset/Source/ModelingComponents/Public/Drawing/MeshWireframeComponent.h b/Engine/Plugins/Runtime/MeshModelingToolset/Source/ModelingComponents/Public/Drawing/MeshWireframeComponent.h
--- a/Engine/Plugins/Runtime/MeshModelingToolset/Source/ModelingComponents/Public/Drawing/MeshWireframeComponent.h
+++ b/Engine/Plugins/Runtime/MeshModelingToolset/Source/ModelingComponents/Public/Drawing/MeshWireframeComponent.h
@@ -29,7 +29,8 @@ public:
 		UVSeam = 1 << 3,
 		NormalSeam = 1 << 4,
 		GroupBoundary = 1 << 5,
-		ColorSeam = 1 << 6
+		ColorSeam = 1 << 6,
+		TangentSeam = 1 << 7
 	};
 
 public:
@@ -173,6 +174,18 @@ public:
 	float NormalSeamThickness = 2.0f;
 
 
+	// tangent seam properties
+
+	UPROPERTY(EditAnywhere, Category = MeshWireframe)
+	bool bEnableTangentSeams = true;
+
+	UPROPERTY(EditAnywhere, Category = MeshWireframe)
+	FColor TangentSeamColor = FColor(64, 240, 240);
+
+	UPROPERTY(EditAnywhere, Category = MeshWireframe)
+	float TangentSeamThickness = 2.0f;
+
+
 	// color seam properties
 
 	UPROPERTY(EditAnywhere, Category = MeshWireframe)
diff --git a/Engine/Source/Runtime/GeometryCore/Private/DynamicMesh/DynamicMeshAttributeSet.cpp b/Engine/Source/Runtime/GeometryCore/Private/DynamicMesh/DynamicMeshAttributeSet.cpp
--- a/Engine/Source/Runtime/GeometryCore/Private/DynamicMesh/DynamicMeshAttributeSet.cpp
+++ b/Engine/Source/Runtime/GeometryCore/Private/DynamicMesh/DynamicMeshAttributeSet.cpp
@@ -694,6 +694,14 @@ bool FDynamicMeshAttributeSet::IsSeamEndEdge(int eid) const
 }
 
 bool FDynamicMeshAttributeSet::IsSeamEdge(int EdgeID, bool& bIsUVSeamOut, bool& bIsNormalSeamOut, bool& bIsColorSeamOut) const
+{
+	bool bIsTangentSeam;
+	bool IsSeam = IsSeamEdge(EdgeID, bIsUVSeamOut, bIsNormalSeamOut, bIsColorSeamOut, bIsTangentSeam);
+	bIsNormalSeamOut = bIsNormalSeamOut || bIsTangentSeam;
+	return IsSeam;
+}
+
+bool FDynamicMeshAttributeSet::IsSeamEdge(int EdgeID, bool& bIsUVSeamOut, bool& bIsNormalSeamOut, bool& bIsColorSeamOut, bool& bIsTangentSeamOut) const
 {
 	bIsUVSeamOut = false;
 	for (const FDynamicMeshUVOverlay& UVLayer : UVLayers)
@@ -704,12 +712,14 @@ bool FDynamicMeshAttributeSet::IsSeamEdge(int EdgeID, bool& bIsUVSeamOut, bool&
 		}
 	}
 
-	bIsNormalSeamOut = false;
-	for (const FDynamicMeshNormalOverlay& NormalLayer : NormalLayers)
+	bIsNormalSeamOut = !NormalLayers.IsEmpty() && NormalLayers[0].IsSeamEdge(EdgeID);
+	bIsTangentSeamOut = false;
+	for (int32 LayerIdx = 1; LayerIdx < NormalLayers.Num(); ++LayerIdx)
 	{
+		const FDynamicMeshNormalOverlay& NormalLayer = NormalLayers[LayerIdx];
 		if (NormalLayer.IsSeamEdge(EdgeID))
 		{
-			bIsNormalSeamOut = true;
+			bIsTangentSeamOut = true;
 		}
 	}
 
@@ -718,7 +728,7 @@ bool FDynamicMeshAttributeSet::IsSeamEdge(int EdgeID, bool& bIsUVSeamOut, bool&
 	{
 		bIsColorSeamOut = true;
 	}
-	return (bIsUVSeamOut || bIsNormalSeamOut || bIsColorSeamOut);
+	return (bIsUVSeamOut || bIsNormalSeamOut || bIsColorSeamOut || bIsTangentSeamOut);
 }
 
 
diff --git a/Engine/Source/Runtime/GeometryCore/Public/DynamicMesh/DynamicMeshAttributeSet.h b/Engine/Source/Runtime/GeometryCore/Public/DynamicMesh/DynamicMeshAttributeSet.h
--- a/Engine/Source/Runtime/GeometryCore/Public/DynamicMesh/DynamicMeshAttributeSet.h
+++ b/Engine/Source/Runtime/GeometryCore/Public/DynamicMesh/DynamicMeshAttributeSet.h
@@ -131,9 +131,13 @@ public:
 	/** @return true if the given edge is the termination of a seam in any overlay*/
 	virtual bool IsSeamEndEdge(int EdgeID) const;
 
-	/** @return true if the given edge is a seam edge in any overlay */
+	/** @return true if the given edge is a seam edge in any overlay. This version considers tangent seams as a type of normal seam; to consider tangent seams separately, call the 4-argument IsSeamEdge instead.  */
+	UE_DEPRECATED(5.3, "Please instead call the 4 argument version of IsSeamEdge, which distinguishes between tangent and normal seam edges.")
 	virtual bool IsSeamEdge(int EdgeID, bool& bIsUVSeamOut, bool& bIsNormalSeamOut, bool& bIsColorSeamOut) const;
 
+	/** @return true if the given edge is a seam edge in any overlay */
+	virtual bool IsSeamEdge(int EdgeID, bool& bIsUVSeamOut, bool& bIsNormalSeamOut, bool& bIsColorSeamOut, bool& bIsTangentSeamOut) const;
+
 	/** @return true if the given vertex is a seam vertex in any overlay */
 	virtual bool IsSeamVertex(int VertexID, bool bBoundaryIsSeam = true) const;
 
