ModelingTools: Replaces the wireframe visuals in the UVLayoutPreview with an instance of MeshElementsVisualizer, instead of the built in wireframe provided by PreviewMesh. This helps address some of the visual z-fighting issues that were being seen and provides more control over the wireframe visuals.

diff --git a/Engine/Plugins/Runtime/MeshModelingToolset/Source/ModelingComponents/Private/Drawing/UVLayoutPreview.cpp b/Engine/Plugins/Runtime/MeshModelingToolset/Source/ModelingComponents/Private/Drawing/UVLayoutPreview.cpp
--- a/Engine/Plugins/Runtime/MeshModelingToolset/Source/ModelingComponents/Private/Drawing/UVLayoutPreview.cpp
+++ b/Engine/Plugins/Runtime/MeshModelingToolset/Source/ModelingComponents/Private/Drawing/UVLayoutPreview.cpp
@@ -3,6 +3,7 @@
 #include "Drawing/UVLayoutPreview.h"
 #include "SceneManagement.h"
 #include "ToolSetupUtil.h"
+#include "Drawing/MeshElementsVisualizer.h"
 
 #include UE_INLINE_GENERATED_CPP_BY_NAME(UVLayoutPreview)
 
@@ -21,6 +22,12 @@ void UUVLayoutPreview::CreateInWorld(UWorld* World)
 	PreviewMesh->CreateInWorld(World, FTransform::Identity);
 	PreviewMesh->SetShadowsEnabled(false);
 
+	MeshElementsVisualizer = NewObject<UMeshElementsVisualizer>(this);
+	MeshElementsVisualizer->CreateInWorld(World, FTransform::Identity);
+	MeshElementsVisualizer->SetMeshAccessFunction([this](UMeshElementsVisualizer::ProcessDynamicMeshFunc ProcessFunc) {
+		PreviewMesh->ProcessMesh(ProcessFunc);
+	});	
+
 	TriangleComponent = NewObject<UTriangleSetComponent>(PreviewMesh->GetActor());
 	TriangleComponent->SetupAttachment(PreviewMesh->GetRootComponent());
 	TriangleComponent->RegisterComponent();
@@ -31,6 +38,17 @@ void UUVLayoutPreview::CreateInWorld(UWorld* World)
 	//Settings->WatchProperty(Settings->bWireframe, [this](bool) { bSettingsModified = true; });
 	bSettingsModified = true;
 
+	MeshElementsVisualizer->Settings->bVisible = Settings->bEnabled;
+	MeshElementsVisualizer->Settings->bShowWireframe = Settings->bShowWireframe;
+	MeshElementsVisualizer->Settings->bShowBorders = Settings->bShowWireframe;
+	MeshElementsVisualizer->Settings->bShowUVSeams = false;
+	MeshElementsVisualizer->Settings->bShowNormalSeams = false;
+	MeshElementsVisualizer->Settings->bShowColorSeams = false;
+	MeshElementsVisualizer->Settings->ThicknessScale = 0.5;
+	MeshElementsVisualizer->Settings->DepthBias = 0.5;
+	MeshElementsVisualizer->Settings->bAdjustDepthBiasUsingMeshSize = false;
+	MeshElementsVisualizer->Settings->CheckAndUpdateWatched();
+
 	BackingRectangleMaterial = ToolSetupUtil::GetSelectionMaterial(FLinearColor::White, nullptr);
 	if (BackingRectangleMaterial == nullptr)
 	{
@@ -43,6 +61,9 @@ void UUVLayoutPreview::Disconnect()
 {
 	PreviewMesh->Disconnect();
 	PreviewMesh = nullptr;
+
+	MeshElementsVisualizer->Disconnect();
+	MeshElementsVisualizer = nullptr;
 }
 
 
@@ -71,12 +92,15 @@ void UUVLayoutPreview::SetCurrentCameraState(const FViewCameraState& CameraState
 void UUVLayoutPreview::SetTransform(const FTransform& UseTransform)
 {
 	PreviewMesh->SetTransform(UseTransform);
+	MeshElementsVisualizer->SetTransform(UseTransform);
 }
 
 
 void UUVLayoutPreview::SetVisible(bool bVisible)
 {
 	PreviewMesh->SetVisible(bVisible);
+	MeshElementsVisualizer->Settings->bVisible = bVisible;
+	MeshElementsVisualizer->Settings->SilentUpdateWatched();
 }
 
 
@@ -112,10 +136,14 @@ void UUVLayoutPreview::OnTick(float DeltaTime)
 	{
 		SetVisible(Settings->bEnabled);
 
-		PreviewMesh->EnableWireframe(Settings->bShowWireframe);
+		MeshElementsVisualizer->Settings->bShowWireframe = Settings->bShowWireframe;
+		MeshElementsVisualizer->Settings->bShowBorders = Settings->bShowWireframe;
+		MeshElementsVisualizer->Settings->CheckAndUpdateWatched();
 
 		bSettingsModified = false;
 	}
+
+	MeshElementsVisualizer->OnTick(DeltaTime);
 }
 
 
@@ -187,6 +215,7 @@ void UUVLayoutPreview::UpdateUVMesh(const FDynamicMesh3* SourceMesh, int32 Sourc
 	}
 
 	PreviewMesh->UpdatePreview(MoveTemp(UVMesh));
+	MeshElementsVisualizer->NotifyMeshChanged();
 }
 
 
diff --git a/Engine/Plugins/Runtime/MeshModelingToolset/Source/ModelingComponents/Public/Drawing/UVLayoutPreview.h b/Engine/Plugins/Runtime/MeshModelingToolset/Source/ModelingComponents/Public/Drawing/UVLayoutPreview.h
--- a/Engine/Plugins/Runtime/MeshModelingToolset/Source/ModelingComponents/Public/Drawing/UVLayoutPreview.h
+++ b/Engine/Plugins/Runtime/MeshModelingToolset/Source/ModelingComponents/Public/Drawing/UVLayoutPreview.h
@@ -10,6 +10,7 @@
 #include "UVLayoutPreview.generated.h"
 
 PREDECLARE_USE_GEOMETRY_CLASS(FDynamicMesh3);
+class UMeshElementsVisualizer;
 
 /**
  * Where should in-viewport UVLayoutPreview be shown, relative to target object
@@ -132,6 +133,10 @@ public:
 	UPROPERTY()
 	TObjectPtr<UPreviewMesh> PreviewMesh;
 
+	/** MeshElementsVisualizer references the PreviewMesh to handle drawing the wireframes for UVs */
+	UPROPERTY()
+	TObjectPtr<UMeshElementsVisualizer> MeshElementsVisualizer;
+
 	/** Set of additional triangles to draw, eg for backing rectangle, etc */
 	UPROPERTY()
 	TObjectPtr<UTriangleSetComponent> TriangleComponent;
