MeshModelingTools: Fix depth bias issues that come about from MeshElementsVisualizer adjusting depth bias to be too high due to a setting that works poorly with percentage-wise depth offset. For now, just making it not user-visible and set to false. Will be deprecated and removed more thoroughly in a later CL.

diff --git a/Engine/Plugins/Runtime/MeshModelingToolset/Source/ModelingComponents/Public/Drawing/MeshElementsVisualizer.h b/Engine/Plugins/Runtime/MeshModelingToolset/Source/ModelingComponents/Public/Drawing/MeshElementsVisualizer.h
--- a/Engine/Plugins/Runtime/MeshModelingToolset/Source/ModelingComponents/Public/Drawing/MeshElementsVisualizer.h
+++ b/Engine/Plugins/Runtime/MeshModelingToolset/Source/ModelingComponents/Public/Drawing/MeshElementsVisualizer.h
@@ -79,9 +79,9 @@ public:
 	UPROPERTY(EditAnywhere, Category = MeshElementVisualization, AdvancedDisplay, meta = (UIMin = -2.0, UIMax = 2.0, EditCondition = "bVisible"))
 	float DepthBias = 0.2;
 
-	/** If true, the depth bias will be adjusted depending on the diagonal length of the mesh bounding box (smaller for smaller meshes). */
-	UPROPERTY(EditAnywhere, Category = MeshElementVisualization, AdvancedDisplay, meta = (EditCondition = "bVisible"))
-	bool bAdjustDepthBiasUsingMeshSize = true;
+	//~ Will be removed- generally not desirable with percentage-wise depth offset
+	UPROPERTY()
+	bool bAdjustDepthBiasUsingMeshSize = false;
 };
 
 
