add 'clean polygon' step to EmbedPolygonsOp to resolve self-intersections and generate less glitchy output in such cases

diff --git a/Engine/Plugins/Runtime/MeshModelingToolset/Source/ModelingOperatorsEditorOnly/Private/CuttingOps/EmbedPolygonsOp.cpp b/Engine/Plugins/Runtime/MeshModelingToolset/Source/ModelingOperatorsEditorOnly/Private/CuttingOps/EmbedPolygonsOp.cpp
--- a/Engine/Plugins/Runtime/MeshModelingToolset/Source/ModelingOperatorsEditorOnly/Private/CuttingOps/EmbedPolygonsOp.cpp
+++ b/Engine/Plugins/Runtime/MeshModelingToolset/Source/ModelingOperatorsEditorOnly/Private/CuttingOps/EmbedPolygonsOp.cpp
@@ -13,7 +13,8 @@
 #include "MeshSimplification.h"
 #include "Operations/MeshPlaneCut.h"
 #include "Operations/MeshBoolean.h"
-#include "ConstrainedDelaunay2.h"
+#include "Curve/GeneralPolygon2.h"
+#include "Curve/PolygonIntersectionUtils.h"
 
 #include UE_INLINE_GENERATED_CPP_BY_NAME(EmbedPolygonsOp)
 
@@ -112,6 +113,30 @@ void CollapseDegenerateEdgesOnVertexPath(FDynamicMesh3& Mesh, TArray<int>& Verte
 	}
 }
 
+TArray<FPolygon2d> FEmbedPolygonsOp::CleanPolygon(const FPolygon2d& Input)
+{
+	TArray<FPolygon2d> Result;
+	// Attempt to clean self-intersections by applying a 'self-union' operation
+	FGeneralPolygon2d GeneralPolygon(Input);
+	FGeneralPolygon2d Empty;
+	TUnionPolygon2Polygon2<FGeneralPolygon2d, double> Union(GeneralPolygon, Empty);
+	bool bResolveSuccess = Union.ComputeResult();
+	if (bResolveSuccess && !Union.Result.IsEmpty())
+	{
+		// Just take the outer polygons for now, as the embed code can't handle holes
+		for (int32 Idx = 0; Idx < Union.Result.Num(); ++Idx)
+		{
+			Result.Add(Union.Result[Idx].GetOuter());
+		}
+	}
+	else
+	{
+		// If we didn't get a result, just fall back to using the input polygon
+		Result.Add(Input);
+	}
+	return Result;
+}
+
 void FEmbedPolygonsOp::BooleanPath(FProgressCancel* Progress)
 {
 	if (ResultMesh->TriangleCount() == 0 || ResultMesh->VertexCount() == 0)
@@ -135,22 +160,39 @@ void FEmbedPolygonsOp::BooleanPath(FProgressCancel* Progress)
 	Range.Expand(1.0);
 
 	Frame.Origin = Frame.Origin + Range.Min * Dir;
-	FGeneralizedCylinderGenerator ExtrudePolyGen;
-	ExtrudePolyGen.CrossSection = EmbedPolygon;
-	ExtrudePolyGen.InitialFrame = Frame;
-	ExtrudePolyGen.Path.Add(Frame.Origin);
-	ExtrudePolyGen.Path.Add(PolygonFrame.Origin + Range.Max * Dir);
-	ExtrudePolyGen.bCapped = true;
-	ExtrudePolyGen.bUVScaleRelativeWorld = true;
-	ExtrudePolyGen.UnitUVInWorldCoordinates = MeshDiameter;
-	FDynamicMesh3 ExtrudePoly;
-	ExtrudePoly.EnableMeshComponents(ResultMesh->GetComponentsFlags());
+	
+	TArray<FPolygon2d> CutPolygons = CleanPolygon(GetPolygon());
+	
+	FDynamicMesh3 ExtrudedPolygons;
+	ExtrudedPolygons.EnableMeshComponents(ResultMesh->GetComponentsFlags());
 	if (ResultMesh->HasAttributes())
 	{
-		ExtrudePoly.EnableAttributes();
-		ExtrudePoly.Attributes()->EnableMatchingAttributes(*ResultMesh->Attributes());
+		ExtrudedPolygons.EnableAttributes();
+		ExtrudedPolygons.Attributes()->EnableMatchingAttributes(*ResultMesh->Attributes());
+	}
+	FDynamicMeshEditor AppendEditor(&ExtrudedPolygons);
+	FMeshIndexMappings IndexMaps_Unused;
+	
+	for (const FPolygon2d& CutPolygon : CutPolygons)
+	{
+		FGeneralizedCylinderGenerator ExtrudePolyGen;
+		ExtrudePolyGen.CrossSection = CutPolygon;
+		ExtrudePolyGen.InitialFrame = Frame;
+		ExtrudePolyGen.Path.Add(Frame.Origin);
+		ExtrudePolyGen.Path.Add(PolygonFrame.Origin + Range.Max * Dir);
+		ExtrudePolyGen.bCapped = true;
+		ExtrudePolyGen.bUVScaleRelativeWorld = true;
+		ExtrudePolyGen.UnitUVInWorldCoordinates = MeshDiameter;
+		FDynamicMesh3 ExtrudePoly;
+		ExtrudePoly.EnableMeshComponents(ResultMesh->GetComponentsFlags());
+		if (ResultMesh->HasAttributes())
+		{
+			ExtrudePoly.EnableAttributes();
+			ExtrudePoly.Attributes()->EnableMatchingAttributes(*ResultMesh->Attributes());
+		}
+		ExtrudePoly.Copy(&ExtrudePolyGen.Generate());
+		AppendEditor.AppendMesh(&ExtrudePoly, IndexMaps_Unused);
 	}
-	ExtrudePoly.Copy(&ExtrudePolyGen.Generate());
 	FMeshBoolean::EBooleanOp BoolOp = FMeshBoolean::EBooleanOp::Difference;
 	switch (Operation)
 	{
@@ -176,7 +218,7 @@ void FEmbedPolygonsOp::BooleanPath(FProgressCancel* Progress)
 	}
 
 	int MaxGroupID = ResultMesh->MaxGroupID();
-	FMeshBoolean Boolean(ResultMesh.Get(), &ExtrudePoly, ResultMesh.Get(), BoolOp);
+	FMeshBoolean Boolean(ResultMesh.Get(), &ExtrudedPolygons, ResultMesh.Get(), BoolOp);
 	Boolean.Progress = Progress;
 	Boolean.bSimplifyAlongNewEdges = true;
 	Boolean.bPutResultInInputSpace = true;
@@ -295,7 +337,15 @@ void FEmbedPolygonsOp::CalculateResult(FProgressCancel* Progress)
 	FFrame3d Frame = PolygonFrame;
 	Frame.Origin = Frame.Origin + (2*MeshDiameter*Frame.Z());
 
-	FPolygon2d Polygon = GetPolygon();
+	TArray<FPolygon2d> CutPolygons = CleanPolygon(GetPolygon());
+
+	if (!ensure(!CutPolygons.IsEmpty()))
+	{
+		return;
+	}
+	
+	// Note: This (non-boolean) path cannot handle multiple cutting polygons (it will almost always fail to cut) so we just take the first polygon here
+	FPolygon2d Polygon = CutPolygons[0];
 	
 	TArray<TPair<float, int>> SortedHitTriangles;
 	TMeshQueries<FDynamicMesh3>::FindHitTriangles_LinearSearch(*ResultMesh, FRay3d(Frame.FromPlaneUV(Polygon[0]), -Frame.Z()), SortedHitTriangles);
diff --git a/Engine/Plugins/Runtime/MeshModelingToolset/Source/ModelingOperatorsEditorOnly/Public/CuttingOps/EmbedPolygonsOp.h b/Engine/Plugins/Runtime/MeshModelingToolset/Source/ModelingOperatorsEditorOnly/Public/CuttingOps/EmbedPolygonsOp.h
--- a/Engine/Plugins/Runtime/MeshModelingToolset/Source/ModelingOperatorsEditorOnly/Public/CuttingOps/EmbedPolygonsOp.h
+++ b/Engine/Plugins/Runtime/MeshModelingToolset/Source/ModelingOperatorsEditorOnly/Public/CuttingOps/EmbedPolygonsOp.h
@@ -68,6 +68,11 @@ private:
 	void RecordEmbeddedEdges(TArray<int>& PathVertIDs);
 
 	void BooleanPath(FProgressCancel* Progress);
+
+	// Fix issues like self-intersection on an input polygon
+	// Note: We could change this to return an array of general polygons in the future, to handle the most general case,
+	// but currently the embedding process does not support holes
+	TArray<FPolygon2d> CleanPolygon(const FPolygon2d& Input);
 };
 
 } // end namespace UE::Geometry
