MeshModelingTools: Fix volumes always having per-triangle normals when not going through volume to mesh tool.

diff --git a/Engine/Plugins/Runtime/MeshModelingToolset/Source/ModelingComponents/Private/ConversionUtils/VolumeToDynamicMesh.cpp b/Engine/Plugins/Runtime/MeshModelingToolset/Source/ModelingComponents/Private/ConversionUtils/VolumeToDynamicMesh.cpp
--- a/Engine/Plugins/Runtime/MeshModelingToolset/Source/ModelingComponents/Private/ConversionUtils/VolumeToDynamicMesh.cpp
+++ b/Engine/Plugins/Runtime/MeshModelingToolset/Source/ModelingComponents/Private/ConversionUtils/VolumeToDynamicMesh.cpp
@@ -5,6 +5,7 @@
 #include "ConstrainedDelaunay2.h"
 #include "Polygon2.h"
 #include "DynamicMesh/DynamicMesh3.h"
+#include "DynamicMesh/MeshNormals.h"
 #include "DynamicMesh/MeshTransforms.h"
 #include "GameFramework/Volume.h"
 #include "MeshBoundaryLoops.h"
@@ -192,6 +193,14 @@ void BrushToDynamicMesh(UModel& BrushModel, UE::Geometry::FDynamicMesh3& Mesh, c
 
 		Mesh.CompactInPlace();
 	}
+
+	if (Options.bGenerateNormals)
+	{
+		Mesh.EnableAttributes();
+		FDynamicMeshNormalOverlay* Normals = Mesh.Attributes()->PrimaryNormals();
+		FMeshNormals::InitializeOverlayTopologyFromFaceGroups(&Mesh, Normals);
+		FMeshNormals::QuickRecomputeOverlayNormals(Mesh);
+	}
 }
 
 
diff --git a/Engine/Plugins/Runtime/MeshModelingToolset/Source/ModelingComponents/Public/ConversionUtils/VolumeToDynamicMesh.h b/Engine/Plugins/Runtime/MeshModelingToolset/Source/ModelingComponents/Public/ConversionUtils/VolumeToDynamicMesh.h
--- a/Engine/Plugins/Runtime/MeshModelingToolset/Source/ModelingComponents/Public/ConversionUtils/VolumeToDynamicMesh.h
+++ b/Engine/Plugins/Runtime/MeshModelingToolset/Source/ModelingComponents/Public/ConversionUtils/VolumeToDynamicMesh.h
@@ -20,6 +20,7 @@ struct FVolumeToMeshOptions
 	bool bMergeVertices = true;
 	bool bAutoRepairMesh = true;
 	bool bOptimizeMesh = true;
+	bool bGenerateNormals = true;
 };
 
 /**
diff --git a/Engine/Plugins/Runtime/MeshModelingToolset/Source/ModelingComponentsEditorOnly/Private/ToolTargets/VolumeComponentToolTarget.cpp b/Engine/Plugins/Runtime/MeshModelingToolset/Source/ModelingComponentsEditorOnly/Private/ToolTargets/VolumeComponentToolTarget.cpp
--- a/Engine/Plugins/Runtime/MeshModelingToolset/Source/ModelingComponentsEditorOnly/Private/ToolTargets/VolumeComponentToolTarget.cpp
+++ b/Engine/Plugins/Runtime/MeshModelingToolset/Source/ModelingComponentsEditorOnly/Private/ToolTargets/VolumeComponentToolTarget.cpp
@@ -81,7 +81,6 @@ FDynamicMesh3 UVolumeComponentToolTarget::GetDynamicMesh()
 
 	FDynamicMesh3 DynamicMesh;
 	UE::Conversion::VolumeToDynamicMesh(Volume, DynamicMesh, GetVolumeToMeshOptions());
-	FMeshNormals::InitializeMeshToPerTriangleNormals(&DynamicMesh);
 
 	return DynamicMesh;
 }
