fix compile error from bad merge in UVProjectionTool.h

diff --git a/Engine/Plugins/Runtime/MeshModelingToolset/Source/MeshModelingTools/Public/UVProjectionTool.h b/Engine/Plugins/Runtime/MeshModelingToolset/Source/MeshModelingTools/Public/UVProjectionTool.h
--- a/Engine/Plugins/Runtime/MeshModelingToolset/Source/MeshModelingTools/Public/UVProjectionTool.h
+++ b/Engine/Plugins/Runtime/MeshModelingToolset/Source/MeshModelingTools/Public/UVProjectionTool.h
@@ -34,7 +34,6 @@ class MESHMODELINGTOOLS_API UUVProjectionToolBuilder : public USingleTargetWithS
 public:
 	virtual USingleTargetWithSelectionTool* CreateNewTool(const FToolBuilderState& SceneState) const override;
 
-	virtual bool CanBuildTool(const FToolBuilderState& SceneState) const override;
 	virtual bool RequiresInputSelection() const override { return false; }
 	virtual bool CanBuildTool(const FToolBuilderState& SceneState) const override;
 };
