add generated light map option to the modeling tools mesh creation settings + make the settings more consistently used by default

also, when creating a static mesh with lightmap enabled via modeling components, also set the lightmap channel and resolution to match the build settings values, to better match the importer results

#8929

diff --git a/Engine/Plugins/Runtime/MeshModelingToolset/Source/ModelingComponents/Private/ModelingComponentsSettings.cpp b/Engine/Plugins/Runtime/MeshModelingToolset/Source/ModelingComponents/Private/ModelingComponentsSettings.cpp
--- a/Engine/Plugins/Runtime/MeshModelingToolset/Source/ModelingComponents/Private/ModelingComponentsSettings.cpp
+++ b/Engine/Plugins/Runtime/MeshModelingToolset/Source/ModelingComponents/Private/ModelingComponentsSettings.cpp
@@ -1,6 +1,7 @@
 // Copyright Epic Games, Inc. All Rights Reserved.
 
 #include "ModelingComponentsSettings.h"
+#include "ModelingObjectsCreationAPI.h"
 
 #include UE_INLINE_GENERATED_CPP_BY_NAME(ModelingComponentsSettings)
 
@@ -30,4 +31,16 @@ FText UModelingComponentsEditorSettings::GetSectionDescription() const
 #endif
 
 
+void UModelingComponentsSettings::ApplyDefaultsToCreateMeshObjectParams(FCreateMeshObjectParams& Params)
+{
+	const UModelingComponentsSettings* Settings = GetDefault<UModelingComponentsSettings>();
+	if (Settings)
+	{
+		Params.bEnableCollision = Settings->bEnableCollision;
+		Params.CollisionMode = Settings->CollisionMode;
+		Params.bGenerateLightmapUVs = Settings->bGenerateLightmapUVs;
+		Params.bEnableRaytracingSupport = Settings->bEnableRayTracing;
+	}
+}
+
 #undef LOCTEXT_NAMESPACE
diff --git a/Engine/Plugins/Runtime/MeshModelingToolset/Source/ModelingComponents/Private/ModelingObjectsCreationAPI.cpp b/Engine/Plugins/Runtime/MeshModelingToolset/Source/ModelingComponents/Private/ModelingObjectsCreationAPI.cpp
--- a/Engine/Plugins/Runtime/MeshModelingToolset/Source/ModelingComponents/Private/ModelingObjectsCreationAPI.cpp
+++ b/Engine/Plugins/Runtime/MeshModelingToolset/Source/ModelingComponents/Private/ModelingObjectsCreationAPI.cpp
@@ -7,11 +7,25 @@
 #include "Misc/Paths.h"
 #include "Components/StaticMeshComponent.h"
 #include "Engine/StaticMesh.h"
+#include "ModelingComponentsSettings.h"
 #include "DynamicMesh/NonManifoldMappingSupport.h"
 
 #include UE_INLINE_GENERATED_CPP_BY_NAME(ModelingObjectsCreationAPI)
 
+// if set to 1, then we do not default-initialize our new mesh object parameters based on the Modeling Components Settings (the Modeling Tools' Project Settings)
+static TAutoConsoleVariable<bool> CVarConstructMeshObjectsWithoutModelingComponentSettings(
+	TEXT("modeling.CreateMesh.IgnoreProjectSettings"),
+	false,
+	TEXT("If enabled, do not use the preferences set in Modeling Tools' Project Settings when constructing new mesh objects"));
 
+FCreateMeshObjectParams::FCreateMeshObjectParams(bool bConstructWithDefaultModelingComponentSettings)
+{
+	if (bConstructWithDefaultModelingComponentSettings && 
+		!CVarConstructMeshObjectsWithoutModelingComponentSettings.GetValueOnGameThread())
+	{
+		UModelingComponentsSettings::ApplyDefaultsToCreateMeshObjectParams(*this);
+	}
+}
 
 void FCreateMeshObjectParams::SetMesh(FMeshDescription&& MeshDescriptionIn)
 {
diff --git a/Engine/Plugins/Runtime/MeshModelingToolset/Source/ModelingComponents/Private/PropertySets/CreateMeshObjectTypeProperties.cpp b/Engine/Plugins/Runtime/MeshModelingToolset/Source/ModelingComponents/Private/PropertySets/CreateMeshObjectTypeProperties.cpp
--- a/Engine/Plugins/Runtime/MeshModelingToolset/Source/ModelingComponents/Private/PropertySets/CreateMeshObjectTypeProperties.cpp
+++ b/Engine/Plugins/Runtime/MeshModelingToolset/Source/ModelingComponents/Private/PropertySets/CreateMeshObjectTypeProperties.cpp
@@ -110,15 +110,6 @@ bool UCreateMeshObjectTypeProperties::ConfigureCreateMeshObjectParams(FCreateMes
 	// client has to handle this case
 	ensure(OutputType != AutoIdentifier);
 
-	const UModelingComponentsSettings* Settings = GetDefault<UModelingComponentsSettings>();
-	if (Settings)
-	{
-		ParamsOut.bEnableCollision = Settings->bEnableCollision;
-		ParamsOut.CollisionMode = Settings->CollisionMode;
-
-		ParamsOut.bEnableRaytracingSupport = Settings->bEnableRayTracing;
-	}
-
 	if (OutputType == StaticMeshIdentifier)
 	{
 		ParamsOut.TypeHint = ECreateObjectTypeHint::StaticMesh;
diff --git a/Engine/Plugins/Runtime/MeshModelingToolset/Source/ModelingComponents/Public/ModelingComponentsSettings.h b/Engine/Plugins/Runtime/MeshModelingToolset/Source/ModelingComponents/Public/ModelingComponentsSettings.h
--- a/Engine/Plugins/Runtime/MeshModelingToolset/Source/ModelingComponents/Public/ModelingComponentsSettings.h
+++ b/Engine/Plugins/Runtime/MeshModelingToolset/Source/ModelingComponents/Public/ModelingComponentsSettings.h
@@ -6,7 +6,7 @@
 #include "BodySetupEnums.h"
 #include "ModelingComponentsSettings.generated.h"
 
-
+struct FCreateMeshObjectParams;
 
 /**
  * Settings for the Modeling Components plug-in. These settings are primarily used to configure two things:
@@ -48,6 +48,10 @@ public:
 	UPROPERTY(config, EditAnywhere, Category = "Modeling Tools|New Mesh Objects")
 	bool bEnableRayTracing = false;
 
+	/** Enable auto-generated Lightmap UVs for new Mesh Objects created by Modeling Tools, where supported */
+	UPROPERTY(config, EditAnywhere, Category = "Modeling Tools|New Mesh Objects")
+	bool bGenerateLightmapUVs = false;
+
 	/** Enable Collision Support for new Mesh Objects created by Modeling Tools */
 	UPROPERTY(config, EditAnywhere, Category = "Modeling Tools|New Mesh Objects")
 	bool bEnableCollision = true;
@@ -56,6 +60,8 @@ public:
 	UPROPERTY(config, EditAnywhere, Category = "Modeling Tools|New Mesh Objects")
 	TEnumAsByte<enum ECollisionTraceFlag> CollisionMode = ECollisionTraceFlag::CTF_UseComplexAsSimple;
 
+	static void ApplyDefaultsToCreateMeshObjectParams(FCreateMeshObjectParams& Params);
+
 
 };
 
diff --git a/Engine/Plugins/Runtime/MeshModelingToolset/Source/ModelingComponents/Public/ModelingObjectsCreationAPI.h b/Engine/Plugins/Runtime/MeshModelingToolset/Source/ModelingComponents/Public/ModelingObjectsCreationAPI.h
--- a/Engine/Plugins/Runtime/MeshModelingToolset/Source/ModelingComponents/Public/ModelingObjectsCreationAPI.h
+++ b/Engine/Plugins/Runtime/MeshModelingToolset/Source/ModelingComponents/Public/ModelingObjectsCreationAPI.h
@@ -76,6 +76,10 @@ enum class ECreateObjectTypeHint : uint8
 USTRUCT(Blueprintable)
 struct MODELINGCOMPONENTS_API FCreateMeshObjectParams
 {
+	// @param bConstructWithDefaultModelingComponentSettings	Whether to initialize with the default project settings.
+	// Note the modeling component settings will not be used if the CVar "modeling.CreateMesh.IgnoreProjectSettings" is enabled
+	FCreateMeshObjectParams(bool bConstructWithDefaultModelingComponentSettings = true);
+
 	GENERATED_BODY()
 
 	//
@@ -143,6 +147,10 @@ struct MODELINGCOMPONENTS_API FCreateMeshObjectParams
 	UPROPERTY(Category = "CreateMeshObjectParams", EditAnywhere)
 	bool bEnableRaytracingSupport = true;
 
+	/** Specify whether to auto-generate Lightmap UVs (if applicable for the output mesh type) */
+	UPROPERTY(Category = "CreateMeshObjectParams", EditAnywhere)
+	bool bGenerateLightmapUVs = false;
+
 	//
 	// Mesh Build Options, if this is applicable for the given mesh object
 	// (Currently somewhat specific to Assets in the Editor)
diff --git a/Engine/Plugins/Runtime/MeshModelingToolset/Source/ModelingComponentsEditorOnly/Private/AssetUtils/CreateStaticMeshUtil.cpp b/Engine/Plugins/Runtime/MeshModelingToolset/Source/ModelingComponentsEditorOnly/Private/AssetUtils/CreateStaticMeshUtil.cpp
--- a/Engine/Plugins/Runtime/MeshModelingToolset/Source/ModelingComponentsEditorOnly/Private/AssetUtils/CreateStaticMeshUtil.cpp
+++ b/Engine/Plugins/Runtime/MeshModelingToolset/Source/ModelingComponentsEditorOnly/Private/AssetUtils/CreateStaticMeshUtil.cpp
@@ -162,6 +162,14 @@ UE::AssetUtils::ECreateStaticMeshResult UE::AssetUtils::CreateStaticMeshAsset(
 	// Distance field
 	NewStaticMesh->bGenerateMeshDistanceField = Options.bAllowDistanceField;
 
+	// Set the initial overall static mesh lightmap settings to match the first lightmap build settings
+	if (Options.bGenerateLightmapUVs)
+	{
+		FMeshBuildSettings& BuildSettings = NewStaticMesh->GetSourceModel(0).BuildSettings;
+		NewStaticMesh->SetLightMapCoordinateIndex(BuildSettings.DstLightmapIndex);
+		NewStaticMesh->SetLightMapResolution(BuildSettings.MinLightmapResolution);
+	}
+
 	[[maybe_unused]] bool MarkedDirty = NewStaticMesh->MarkPackageDirty();
 	if (Options.bDeferPostEditChange == false)
 	{
diff --git a/Engine/Plugins/Runtime/MeshModelingToolset/Source/ModelingComponentsEditorOnly/Private/EditorModelingObjectsCreationAPI.cpp b/Engine/Plugins/Runtime/MeshModelingToolset/Source/ModelingComponentsEditorOnly/Private/EditorModelingObjectsCreationAPI.cpp
--- a/Engine/Plugins/Runtime/MeshModelingToolset/Source/ModelingComponentsEditorOnly/Private/EditorModelingObjectsCreationAPI.cpp
+++ b/Engine/Plugins/Runtime/MeshModelingToolset/Source/ModelingComponentsEditorOnly/Private/EditorModelingObjectsCreationAPI.cpp
@@ -293,6 +293,7 @@ FCreateMeshObjectResult UEditorModelingObjectsCreationAPI::CreateStaticMeshAsset
 	AssetOptions.bEnableRecomputeTangents = CreateMeshParams.bEnableRecomputeTangents;
 	AssetOptions.bGenerateNaniteEnabledMesh = CreateMeshParams.bEnableNanite;
 	AssetOptions.NaniteSettings = CreateMeshParams.NaniteSettings;
+	AssetOptions.bGenerateLightmapUVs = CreateMeshParams.bGenerateLightmapUVs;
 
 	AssetOptions.bCreatePhysicsBody = CreateMeshParams.bEnableCollision;
 	AssetOptions.CollisionType = CreateMeshParams.CollisionMode;
