Initialize bone attributes from the transform hierarchy when performing Skin Binding Operation.
Store bone transform in the transform hiearchy as well.

diff --git a/Engine/Plugins/Runtime/MeshModelingToolset/Source/ModelingOperators/Private/SkinningOps/SkinBindingOp.cpp b/Engine/Plugins/Runtime/MeshModelingToolset/Source/ModelingOperators/Private/SkinningOps/SkinBindingOp.cpp
--- a/Engine/Plugins/Runtime/MeshModelingToolset/Source/ModelingOperators/Private/SkinningOps/SkinBindingOp.cpp
+++ b/Engine/Plugins/Runtime/MeshModelingToolset/Source/ModelingOperators/Private/SkinningOps/SkinBindingOp.cpp
@@ -6,6 +6,7 @@
 #include "DynamicMesh/DynamicMesh3.h"
 #include "DynamicMesh/DynamicMeshAttributeSet.h"
 #include "DynamicMesh/DynamicVertexSkinWeightsAttribute.h"
+#include "DynamicMesh/DynamicBoneAttribute.h"
 #include "ReferenceSkeleton.h"
 #include "Spatial/DenseGrid3.h"
 #include "Spatial/MeshWindingNumberGrid.h"
@@ -189,17 +190,17 @@ static float DistanceToLineSegment(const FVector& P, const FVector& A, const FVe
 // store a list of line segments going from the bone transform to all the child bone transforms.
 struct FTransformHierarchyQuery
 {
-	explicit FTransformHierarchyQuery(const TArray<TPair<FTransform, int32>>& InTransformHierarchy)
+	explicit FTransformHierarchyQuery(const TArray<TPair<FTransform, FMeshBoneInfo>>& InTransformHierarchy)
 	{
 		for (int Index = 0; Index < InTransformHierarchy.Num(); Index++)
 		{
 			FTransform Xform = InTransformHierarchy[Index].Key;
-			int32 ParentIndex = InTransformHierarchy[Index].Value;
+			int32 ParentIndex = InTransformHierarchy[Index].Value.ParentIndex;
 
 			while (ParentIndex != INDEX_NONE)
 			{
 				Xform = Xform * InTransformHierarchy[ParentIndex].Key;
-				ParentIndex = InTransformHierarchy[ParentIndex].Value;
+				ParentIndex = InTransformHierarchy[ParentIndex].Value.ParentIndex;
 			}
 
 			BoneFans.Add({ Xform.GetLocation() });
@@ -208,7 +209,7 @@ struct FTransformHierarchyQuery
 		// Fill in the fan tips, as needed.
 		for (int Index = 0; Index < InTransformHierarchy.Num(); Index++)
 		{
-			const int32 ParentIndex = InTransformHierarchy[Index].Value;
+			const int32 ParentIndex = InTransformHierarchy[Index].Value.ParentIndex;
 			if (ParentIndex != INDEX_NONE)
 			{
 				BoneFans[ParentIndex].TipsPos.Add(BoneFans[Index].RootPos);
@@ -319,7 +320,7 @@ void FSkinBindingOp::SetTransformHierarchyFromReferenceSkeleton(
 		
 	for (int32 Index = 0; Index < BoneInfo.Num(); Index++)
 	{
-		TransformHierarchy.Add(MakeTuple(BonePose[Index], BoneInfo[Index].ParentIndex));
+		TransformHierarchy.Add(MakeTuple(BonePose[Index], BoneInfo[Index]));
 	}
 }
 
@@ -353,6 +354,23 @@ void FSkinBindingOp::CalculateResult(FProgressCancel* InProgress)
 		CreateSkinWeights_GeodesicVoxel(*ResultMesh, ClampedStiffness, Settings);
 		break;
 	}
+
+	// Initialize bone attributes
+	FDynamicMeshAttributeSet* AttribSet = ResultMesh->Attributes();
+
+	const int32 NumBones = TransformHierarchy.Num();
+	AttribSet->EnableBones(NumBones);
+
+	FDynamicMeshBoneNameAttribute* BoneNames = AttribSet->GetBoneNames();
+	FDynamicMeshBoneParentIndexAttribute* BoneParentIndices = AttribSet->GetBoneParentIndices();
+	FDynamicMeshBonePoseAttribute* BonePoses = AttribSet->GetBonePoses();
+
+	for (int BoneIdx = 0; BoneIdx < NumBones; ++BoneIdx)
+	{
+		BoneNames->SetValue(BoneIdx, TransformHierarchy[BoneIdx].Value.Name);
+		BoneParentIndices->SetValue(BoneIdx, TransformHierarchy[BoneIdx].Value.ParentIndex);
+		BonePoses->SetValue(BoneIdx, TransformHierarchy[BoneIdx].Key);
+	}
 }
 
 
diff --git a/Engine/Plugins/Runtime/MeshModelingToolset/Source/ModelingOperators/Public/SkinningOps/SkinBindingOp.h b/Engine/Plugins/Runtime/MeshModelingToolset/Source/ModelingOperators/Public/SkinningOps/SkinBindingOp.h
--- a/Engine/Plugins/Runtime/MeshModelingToolset/Source/ModelingOperators/Public/SkinningOps/SkinBindingOp.h
+++ b/Engine/Plugins/Runtime/MeshModelingToolset/Source/ModelingOperators/Public/SkinningOps/SkinBindingOp.h
@@ -7,7 +7,7 @@
 
 
 struct FReferenceSkeleton;
-
+struct FMeshBoneInfo;
 
 namespace UE::Geometry
 {
@@ -38,7 +38,7 @@ public:
 	// The transform hierarchy to bind to. Listed in the same order as the bones in the
 	// reference skeleton that this skeletal mesh is tied to.
 	TSharedPtr<FDynamicMesh3, ESPMode::ThreadSafe> OriginalMesh;
-	TArray<TPair<FTransform, int32>> TransformHierarchy;
+	TArray<TPair<FTransform, FMeshBoneInfo>> TransformHierarchy;
 
 	FName ProfileName = FSkeletalMeshAttributes::DefaultSkinWeightProfileName;
 	
