MeshModelingToolset: Fix crash when snapping to a static mesh without vertex buffer data

diff --git a/Engine/Plugins/Runtime/MeshModelingToolset/Source/ModelingComponents/Private/Snapping/ModelingSceneSnappingManager.cpp b/Engine/Plugins/Runtime/MeshModelingToolset/Source/ModelingComponents/Private/Snapping/ModelingSceneSnappingManager.cpp
--- a/Engine/Plugins/Runtime/MeshModelingToolset/Source/ModelingComponents/Private/Snapping/ModelingSceneSnappingManager.cpp
+++ b/Engine/Plugins/Runtime/MeshModelingToolset/Source/ModelingComponents/Private/Snapping/ModelingSceneSnappingManager.cpp
@@ -363,34 +363,46 @@ bool GetComponentHitTriangle_Internal(FHitResult HitResult, VectorType* TriVerti
 		return false;
 	}
 
-	if (UStaticMeshComponent* Component = Cast<UStaticMeshComponent>(HitResult.Component.Get()))
-	{
-		// physics collision data is created from StaticMesh RenderData
-		// so use HitResult.FaceIndex to extract triangle from the LOD0 mesh
-		// (note: this may be incorrect if there are multiple sections...in that case I think we have to
-		//  first find section whose accumulated index range would contain .FaceIndexX)
-		UStaticMesh* StaticMesh = Component->GetStaticMesh();
-		FStaticMeshLODResources& LOD = StaticMesh->GetRenderData()->LODResources[0];
-		FIndexArrayView Indices = LOD.IndexBuffer.GetArrayView();
-		int32 TriIdx = 3 * HitResult.FaceIndex;
-		FVector Positions[3];
-		Positions[0] = (FVector)LOD.VertexBuffers.PositionVertexBuffer.VertexPosition(Indices[TriIdx]);
-		Positions[1] = (FVector)LOD.VertexBuffers.PositionVertexBuffer.VertexPosition(Indices[TriIdx + 1]);
-		Positions[2] = (FVector)LOD.VertexBuffers.PositionVertexBuffer.VertexPosition(Indices[TriIdx + 2]);
-
-		// transform to world space
-		FTransform ComponentTransform = Component->GetComponentTransform();
-		Positions[0] = ComponentTransform.TransformPosition(Positions[0]);
-		Positions[1] = ComponentTransform.TransformPosition(Positions[1]);
-		Positions[2] = ComponentTransform.TransformPosition(Positions[2]);
-
-		TriVertices[0] = (VectorType)Positions[0];
-		TriVertices[1] = (VectorType)Positions[1];
-		TriVertices[2] = (VectorType)Positions[2];
-		return true;
+	UStaticMeshComponent* Component = Cast<UStaticMeshComponent>(HitResult.Component.Get());
+	if (!Component)
+	{
+		return false;
 	}
+	
+	// physics collision data is created from StaticMesh RenderData
+	// so use HitResult.FaceIndex to extract triangle from the LOD0 mesh
+	// (note: this may be incorrect if there are multiple sections...in that case I think we have to
+	//  first find section whose accumulated index range would contain .FaceIndexX)
+	UStaticMesh* StaticMesh = Component->GetStaticMesh();
+	FStaticMeshLODResources& LOD = StaticMesh->GetRenderData()->LODResources[0];
+	FIndexArrayView Indices = LOD.IndexBuffer.GetArrayView();
 
-	return false;
+	const int32 TriIdx = 3 * HitResult.FaceIndex;
+	if (TriIdx + 2 >= Indices.Num())
+	{
+		return false;
+	}
+
+	const uint32 Idx[3] = { Indices[TriIdx], Indices[TriIdx + 1], Indices[TriIdx + 2] };
+	const uint32 NumVertices = LOD.VertexBuffers.PositionVertexBuffer.GetNumVertices();
+	if (Idx[0] >= NumVertices || Idx[1] >= NumVertices || Idx[2] >= NumVertices)
+	{
+		return false;
+	}
+
+	const FVector Positions[3] = {
+		static_cast<FVector>(LOD.VertexBuffers.PositionVertexBuffer.VertexPosition(Idx[0])),
+		static_cast<FVector>(LOD.VertexBuffers.PositionVertexBuffer.VertexPosition(Idx[1])),
+		static_cast<FVector>(LOD.VertexBuffers.PositionVertexBuffer.VertexPosition(Idx[2]))
+	};
+
+	// transform to world space
+	const FTransform ComponentTransform = Component->GetComponentTransform();
+	TriVertices[0] = static_cast<VectorType>(ComponentTransform.TransformPosition(Positions[0]));
+	TriVertices[1] = static_cast<VectorType>(ComponentTransform.TransformPosition(Positions[1]));
+	TriVertices[2] = static_cast<VectorType>(ComponentTransform.TransformPosition(Positions[2]));
+
+	return true;
 }
 
 
