ModelingMode: preliminary support for selection transform visualization for Volumes and Static Meshes, and other minor improvements.

Add IGeometrySelectionTransformer::PreviewRender() API, for rendering info during selection transforms. Currently being called by  UGeometrySelectionManager::DebugRender() for active transforms.
FBasicDynamicMeshSelectionTransformer implements PreviewRender, draws selected mesh edges (tris as edges) and vertices, as well as affected edge ROI.
FDynamicMeshPolygroupTransformer overrides this behavior and draws polygroup edges/verts and edge ROI

FStaticMeshSelector and FVolumeSelector now create FDynamicMeshPolygroupTransformer for polygroup selections

FBasicDynamicMeshSelectionTransformer::bEnableSelectionTransformDrawing flag controls whether this edge drawing is enabled. Disabled on DynamicMeshComponent, enabled on StaticMesh and Volume.

Add Normals recomputation during tri and polygroup transforms in FBasicDynamicMeshSelectionTransformer and FDynamicMeshPolygroupTransformer

Add FGroupTopologyDeformer::EnumerateROIEdges() util function to support above

diff --git a/Engine/Plugins/Runtime/GeometryProcessing/Source/DynamicMesh/Public/Operations/GroupTopologyDeformer.h b/Engine/Plugins/Runtime/GeometryProcessing/Source/DynamicMesh/Public/Operations/GroupTopologyDeformer.h
--- a/Engine/Plugins/Runtime/GeometryProcessing/Source/DynamicMesh/Public/Operations/GroupTopologyDeformer.h
+++ b/Engine/Plugins/Runtime/GeometryProcessing/Source/DynamicMesh/Public/Operations/GroupTopologyDeformer.h
@@ -125,6 +125,15 @@ public:
 	/** @return the set of all overlay normals that will be modified by the deformation*/
 	const TSet<int>& GetModifiedOverlayNormals() const { return ModifiedOverlayNormals; }
 
+	/** Call EdgeSpanFunc for every group FEdgeSpan in the modified-area ROI that will be modified by the solver (does not include edges that are encompassed by the Handle) */
+	void EnumerateROIEdges(TFunctionRef<void(const FEdgeSpan& EdgeSpan)> EdgeSpanFunc) const
+	{
+		for (const FROIEdge& Edge : ROIEdges)
+		{
+			EdgeSpanFunc(Edge.Span);
+		}
+	}
+
 protected:
 	//
 	// These are functions that subclasses may wish to override
diff --git a/Engine/Plugins/Runtime/MeshModelingToolset/Source/ModelingComponents/Private/Selection/DynamicMeshPolygroupTransformer.cpp b/Engine/Plugins/Runtime/MeshModelingToolset/Source/ModelingComponents/Private/Selection/DynamicMeshPolygroupTransformer.cpp
--- a/Engine/Plugins/Runtime/MeshModelingToolset/Source/ModelingComponents/Private/Selection/DynamicMeshPolygroupTransformer.cpp
+++ b/Engine/Plugins/Runtime/MeshModelingToolset/Source/ModelingComponents/Private/Selection/DynamicMeshPolygroupTransformer.cpp
@@ -10,6 +10,7 @@
 #include "Changes/MeshVertexChange.h"
 #include "GroupTopology.h"
 #include "Operations/GroupTopologyDeformer.h"
+#include "ToolDataVisualizer.h"
 
 using namespace UE::Geometry;
 
@@ -50,7 +51,25 @@ void FDynamicMeshPolygroupTransformer::BeginTransform(const FGeometrySelection&
 			ROIMap.Add(MeshVertices[k], k);
 		}
 
-		//UpdatedPositions = InitialPositions;
+		UpdatedPositions = InitialPositions;
+
+		OverlayNormalsArray = LinearDeformer->GetModifiedOverlayNormals().Array();
+
+		ActiveSelectionEdges.Reset();
+		ActiveSelectionVertices.Reset();
+		UE::Geometry::EnumeratePolygroupSelectionElements( Selection, SourceMesh, UseTopology,
+			[&](int32 VertexID, FVector3d Position) { ActiveSelectionVertices.Add(VertexID); },
+			[&](int32 EdgeID, const FSegment3d&) { ActiveSelectionEdges.Add(SourceMesh.GetEdgeV(EdgeID)); },
+			[&](int32 TriangleID, const FTriangle3d&) {},
+			/*ApplyTransform*/nullptr, /*bMapFacesToEdgeLoops*/true);
+
+		ActiveROIEdges.Reset();
+		LinearDeformer->EnumerateROIEdges([&](const FEdgeSpan& Span) {
+			for (int32 EdgeID : Span.Edges)
+			{
+				ActiveROIEdges.Add(SourceMesh.GetEdgeV(EdgeID));
+			}
+		});
 	});
 
 	ActiveVertexChange = MakePimpl<FMeshVertexChangeBuilder>(EMeshVertexChangeComponents::VertexPositions | EMeshVertexChangeComponents::OverlayNormals);
@@ -76,10 +95,10 @@ void FDynamicMeshPolygroupTransformer::UpdateTransform(
 	int32 N = MeshVertices.Num();
 	FTransform WorldTransform = Selector->GetLocalToWorldTransform();
 
-	//for (int32 k = 0; k < N; ++k)
-	//{
-	//	UpdatedPositions[k] = PositionTransformFunc(MeshVertices[k], InitialPositions[k], WorldTransform);
-	//}
+	for (int32 k = 0; k < N; ++k)
+	{
+		UpdatedPositions[k] = PositionTransformFunc(MeshVertices[k], InitialPositions[k], WorldTransform);
+	}
 
 	Selector->GetDynamicMesh()->EditMesh([&](FDynamicMesh3& EditMesh)
 	{
@@ -98,17 +117,10 @@ void FDynamicMeshPolygroupTransformer::UpdateTransform(
 			}
 		});
 
-
-		// TODO: the right thing to do here would be to only recompute the modified overlay normals...but we have no utility for that??
-		//FMeshNormals::QuickRecomputeOverlayNormals(EditMesh);
+		FMeshNormals::RecomputeOverlayElementNormals(EditMesh, OverlayNormalsArray);
 
 	}, EDynamicMeshChangeType::DeformationEdit,
-		//EDynamicMeshAttributeChangeFlags::VertexPositions, false);
 	   EDynamicMeshAttributeChangeFlags::VertexPositions | EDynamicMeshAttributeChangeFlags::NormalsTangents, false);
-
-
-	// why does simple deformer do this every time??
-	//UpdatePendingVertexChange(false);
 }
 
 
@@ -142,4 +154,46 @@ void FDynamicMeshPolygroupTransformer::EndTransform(IToolsContextTransactionsAPI
 }
 
 
+
+void FDynamicMeshPolygroupTransformer::PreviewRender(IToolsContextRenderAPI* RenderAPI)
+{
+	if (bFallbackToSimpleTransform)
+	{
+		FBasicDynamicMeshSelectionTransformer::PreviewRender(RenderAPI);
+		return;
+	}
+	if (!bEnableSelectionTransformDrawing) return;
+	if (!LinearDeformer.IsValid()) return;
+
+	FToolDataVisualizer Visualizer;
+	Visualizer.bDepthTested = false;
+	Visualizer.BeginFrame(RenderAPI);
+
+	Visualizer.PushTransform(Selector->GetLocalToWorldTransform());
+
+	const FDynamicMesh3* Mesh = LinearDeformer->GetMesh();
+
+	Visualizer.SetPointParameters(FLinearColor(0, 0.3f, 0.95f, 1), 5.0f);
+	for (int32 VertexID : ActiveSelectionVertices)
+	{
+		Visualizer.DrawPoint(Mesh->GetVertex(VertexID));
+	}
+
+	Visualizer.SetLineParameters(FLinearColor(0, 0.3f, 0.95f, 1), 3.0f);
+	for (FIndex2i EdgeV : ActiveSelectionEdges)
+	{
+		Visualizer.DrawLine(Mesh->GetVertex(EdgeV.A), Mesh->GetVertex(EdgeV.B));
+	}
+	Visualizer.SetLineParameters(FLinearColor(0, 0.3f, 0.95f, 1), 1.0f);
+	for (FIndex2i EdgeV : ActiveROIEdges)
+	{
+		Visualizer.DrawLine(Mesh->GetVertex(EdgeV.A), Mesh->GetVertex(EdgeV.B));
+	}
+
+	Visualizer.EndFrame();
+}
+
+
+
+
 #undef LOCTEXT_NAMESPACE 
\ No newline at end of file
diff --git a/Engine/Plugins/Runtime/MeshModelingToolset/Source/ModelingComponents/Private/Selection/DynamicMeshSelector.cpp b/Engine/Plugins/Runtime/MeshModelingToolset/Source/ModelingComponents/Private/Selection/DynamicMeshSelector.cpp
--- a/Engine/Plugins/Runtime/MeshModelingToolset/Source/ModelingComponents/Private/Selection/DynamicMeshSelector.cpp
+++ b/Engine/Plugins/Runtime/MeshModelingToolset/Source/ModelingComponents/Private/Selection/DynamicMeshSelector.cpp
@@ -13,6 +13,7 @@
 #include "Spatial/SegmentTree3.h"
 #include "ToolSceneQueriesUtil.h"
 #include "Selection/DynamicMeshPolygroupTransformer.h"
+#include "ToolDataVisualizer.h"
 
 using namespace UE::Geometry;
 
@@ -809,10 +810,37 @@ void FBasicDynamicMeshSelectionTransformer::BeginTransform(const FGeometrySelect
 		UE::Geometry::VertexToTriangleOneRing(&SourceMesh, MeshVertices, TriangleROI);
 
 		// save overlay normals
+		OverlayNormalsArray.Reset();
+		OverlayNormalsSet.Reset();
 		if (SourceMesh.HasAttributes() && SourceMesh.Attributes()->PrimaryNormals() != nullptr)
 		{
 			const FDynamicMeshNormalOverlay* Normals = SourceMesh.Attributes()->PrimaryNormals();
-			UE::Geometry::TrianglesToOverlayElements(Normals, TriangleROI, OverlayNormals);
+			UE::Geometry::TrianglesToOverlayElements(Normals, TriangleROI, OverlayNormalsSet);
+			OverlayNormalsArray = OverlayNormalsSet.Array();
+		}
+
+		// collect selected mesh edges and vertices
+		ActiveSelectionEdges.Reset();
+		ActiveSelectionVertices.Reset();
+		TSet<int32> SelectionEdges;
+		UE::Geometry::EnumerateTriangleSelectionElements(Selection, SourceMesh, 
+			[&](int32 VertexID, FVector3d Position) { ActiveSelectionVertices.Add(VertexID); },
+			[&](int32 EdgeID, const FSegment3d&) { ActiveSelectionEdges.Add(SourceMesh.GetEdgeV(EdgeID)); SelectionEdges.Add(EdgeID); },
+			[&](int32 TriangleID, const FTriangle3d&) {},
+			/*ApplyTransform*/nullptr, /*bMapFacesToEdgeLoops*/true);
+
+		ActiveROIEdges.Reset();
+		for (int32 tid : TriangleROI)
+		{
+			FIndex3i TriEdges = SourceMesh.GetTriEdges(tid);
+			for (int32 k = 0; k < 3; ++k)
+			{
+				if (SelectionEdges.Contains(TriEdges[k]) == false)
+				{
+					ActiveROIEdges.Add(SourceMesh.GetEdgeV(TriEdges[k]));
+					SelectionEdges.Add(TriEdges[k]);
+				}
+			}
 		}
 	});
 
@@ -839,14 +867,10 @@ void FBasicDynamicMeshSelectionTransformer::UpdateTransform(
 			EditMesh.SetVertex(MeshVertices[k], UpdatedPositions[k]);
 		}
 
-		// TODO: the right thing to do here would be to only recompute the modified overlay normals...but we have no utility for that??
-		//FMeshNormals::QuickRecomputeOverlayNormals(EditMesh);
+		FMeshNormals::RecomputeOverlayElementNormals(EditMesh, OverlayNormalsArray);
 
 	}, EDynamicMeshChangeType::DeformationEdit,
-		//EDynamicMeshAttributeChangeFlags::VertexPositions, false);
 	   EDynamicMeshAttributeChangeFlags::VertexPositions | EDynamicMeshAttributeChangeFlags::NormalsTangents, false);
-
-	UpdatePendingVertexChange(false);
 }
 
 void FBasicDynamicMeshSelectionTransformer::UpdatePendingVertexChange(bool bFinal)
@@ -855,7 +879,7 @@ void FBasicDynamicMeshSelectionTransformer::UpdatePendingVertexChange(bool bFina
 	Selector->GetDynamicMesh()->ProcessMesh([&](const FDynamicMesh3& SourceMesh)
 	{
 		ActiveVertexChange->SaveVertices(&SourceMesh, MeshVertices, !bFinal);
-		ActiveVertexChange->SaveOverlayNormals(&SourceMesh, OverlayNormals, !bFinal);
+		ActiveVertexChange->SaveOverlayNormals(&SourceMesh, OverlayNormalsSet, !bFinal);
 	});
 }
 
@@ -878,7 +902,38 @@ void FBasicDynamicMeshSelectionTransformer::EndTransform(IToolsContextTransactio
 }
 
 
+void FBasicDynamicMeshSelectionTransformer::PreviewRender(IToolsContextRenderAPI* RenderAPI)
+{
+	if (!bEnableSelectionTransformDrawing) return;
+
+	FToolDataVisualizer Visualizer;
+	Visualizer.bDepthTested = false;
+	Visualizer.BeginFrame(RenderAPI);
+
+	Visualizer.PushTransform(Selector->GetLocalToWorldTransform());
 
+	Selector->GetDynamicMesh()->ProcessMesh([&](const FDynamicMesh3& SourceMesh)
+	{
+		Visualizer.SetPointParameters(FLinearColor(0, 0.3f, 0.95f, 1), 5.0f);
+		for (int32 VertexID : ActiveSelectionVertices)
+		{
+			Visualizer.DrawPoint(SourceMesh.GetVertex(VertexID));
+		}
+
+		Visualizer.SetLineParameters(FLinearColor(0, 0.3f, 0.95f, 1), 3.0f);
+		for (FIndex2i EdgeV : ActiveSelectionEdges)
+		{
+			Visualizer.DrawLine(SourceMesh.GetVertex(EdgeV.A), SourceMesh.GetVertex(EdgeV.B));
+		}
+		Visualizer.SetLineParameters(FLinearColor(0, 0.3f, 0.95f, 1), 1.0f);
+		for (FIndex2i EdgeV : ActiveROIEdges)
+		{
+			Visualizer.DrawLine(SourceMesh.GetVertex(EdgeV.A), SourceMesh.GetVertex(EdgeV.B));
+		}
+	});
+
+	Visualizer.EndFrame();
+}
 
 
 
diff --git a/Engine/Plugins/Runtime/MeshModelingToolset/Source/ModelingComponents/Private/Selection/GeometrySelectionManager.cpp b/Engine/Plugins/Runtime/MeshModelingToolset/Source/ModelingComponents/Private/Selection/GeometrySelectionManager.cpp
--- a/Engine/Plugins/Runtime/MeshModelingToolset/Source/ModelingComponents/Private/Selection/GeometrySelectionManager.cpp
+++ b/Engine/Plugins/Runtime/MeshModelingToolset/Source/ModelingComponents/Private/Selection/GeometrySelectionManager.cpp
@@ -1164,6 +1164,14 @@ void UGeometrySelectionManager::DebugRender(IToolsContextRenderAPI* RenderAPI)
 	// disable selection during xform to avoid overhead
 	if (IsInActiveTransformation())
 	{
+		for (int32 k = 0; k < ActiveTargetReferences.Num(); ++k)
+		{
+			if (ActiveTransformations[k] != nullptr)
+			{
+				ActiveTransformations[k]->PreviewRender(RenderAPI);
+			}
+		}
+
 		return;
 	}
 
diff --git a/Engine/Plugins/Runtime/MeshModelingToolset/Source/ModelingComponents/Public/Selection/DynamicMeshPolygroupTransformer.h b/Engine/Plugins/Runtime/MeshModelingToolset/Source/ModelingComponents/Public/Selection/DynamicMeshPolygroupTransformer.h
--- a/Engine/Plugins/Runtime/MeshModelingToolset/Source/ModelingComponents/Public/Selection/DynamicMeshPolygroupTransformer.h
+++ b/Engine/Plugins/Runtime/MeshModelingToolset/Source/ModelingComponents/Public/Selection/DynamicMeshPolygroupTransformer.h
@@ -21,6 +21,8 @@ public:
 	virtual void UpdateTransform( TFunctionRef<FVector3d(int32 VertexID, const FVector3d& InitialPosition, const FTransform& WorldTransform)> PositionTransformFunc ) override;
 	virtual void EndTransform(IToolsContextTransactionsAPI* TransactionsAPI) override;
 
+	virtual void PreviewRender(IToolsContextRenderAPI* RenderAPI) override;
+
 protected:
 	bool bFallbackToSimpleTransform = false;
 
diff --git a/Engine/Plugins/Runtime/MeshModelingToolset/Source/ModelingComponents/Public/Selection/DynamicMeshSelector.h b/Engine/Plugins/Runtime/MeshModelingToolset/Source/ModelingComponents/Public/Selection/DynamicMeshSelector.h
--- a/Engine/Plugins/Runtime/MeshModelingToolset/Source/ModelingComponents/Public/Selection/DynamicMeshSelector.h
+++ b/Engine/Plugins/Runtime/MeshModelingToolset/Source/ModelingComponents/Public/Selection/DynamicMeshSelector.h
@@ -227,6 +227,11 @@ public:
 	virtual void UpdateTransform( TFunctionRef<FVector3d(int32 VertexID, const FVector3d& InitialPosition, const FTransform& WorldTransform)> PositionTransformFunc ) override;
 	void UpdatePendingVertexChange(bool bFinal);
 
+	/** Enable line drawing of selection during transform, this is necessary in some contexts where live mesh update is too slow */
+	bool bEnableSelectionTransformDrawing = false;
+
+	virtual void PreviewRender(IToolsContextRenderAPI* RenderAPI) override;
+
 	virtual void EndTransform(IToolsContextTransactionsAPI* TransactionsAPI) override;
 
 	TFunction<void(IToolsContextTransactionsAPI* TransactionsAPI)> OnEndTransformFunc;
@@ -237,12 +242,18 @@ protected:
 	TArray<int32> MeshVertices;
 	TArray<FVector3d> InitialPositions;
 	TSet<int32> TriangleROI;
-	TSet<int32> OverlayNormals;
+	TSet<int32> OverlayNormalsSet;
+	TArray<int32> OverlayNormalsArray;
 
 	TArray<FVector3d> UpdatedPositions;
 
 	TPimplPtr<FMeshVertexChangeBuilder> ActiveVertexChange;
 
+	// used for preview rendering
+	TArray<UE::Geometry::FIndex2i> ActiveSelectionEdges;
+	TArray<int32> ActiveSelectionVertices;
+	TArray<UE::Geometry::FIndex2i> ActiveROIEdges;
+
 };
 
 
diff --git a/Engine/Plugins/Runtime/MeshModelingToolset/Source/ModelingComponents/Public/Selection/GeometrySelector.h b/Engine/Plugins/Runtime/MeshModelingToolset/Source/ModelingComponents/Public/Selection/GeometrySelector.h
--- a/Engine/Plugins/Runtime/MeshModelingToolset/Source/ModelingComponents/Public/Selection/GeometrySelector.h
+++ b/Engine/Plugins/Runtime/MeshModelingToolset/Source/ModelingComponents/Public/Selection/GeometrySelector.h
@@ -9,6 +9,7 @@
 #include "InputState.h"
 #include "ToolContextInterfaces.h"
 
+class IToolsContextRenderAPI;
 class IToolsContextTransactionsAPI;
 
 
@@ -183,6 +184,11 @@ public:
 	 */
 	virtual void UpdateTransform( TFunctionRef<FVector3d(int VertexID, const FVector3d& InitialPosition, const FTransform& WorldTransform)> VertexTransformFunc ) = 0;
 
+	/**
+	 * Optional function called per-frame to visualize the active Transformation
+	 */
+	virtual void PreviewRender(IToolsContextRenderAPI* RenderAPI) { }
+
 	/**
 	 * Finish a transform (eg called on mouse-up during 3D gizmo interaction). 
 	 * The caller provides a TransactionsAPI implementation which the Transformer/Selector can use to emit change events
diff --git a/Engine/Plugins/Runtime/MeshModelingToolset/Source/ModelingComponentsEditorOnly/Private/Selection/StaticMeshSelector.cpp b/Engine/Plugins/Runtime/MeshModelingToolset/Source/ModelingComponentsEditorOnly/Private/Selection/StaticMeshSelector.cpp
--- a/Engine/Plugins/Runtime/MeshModelingToolset/Source/ModelingComponentsEditorOnly/Private/Selection/StaticMeshSelector.cpp
+++ b/Engine/Plugins/Runtime/MeshModelingToolset/Source/ModelingComponentsEditorOnly/Private/Selection/StaticMeshSelector.cpp
@@ -16,6 +16,8 @@
 #include "MeshDescriptionToDynamicMesh.h"
 #include "DynamicMeshToMeshDescription.h"
 
+#include "Selection/DynamicMeshPolygroupTransformer.h"
+
 #include "RenderingThread.h"
 
 using namespace UE::Geometry;
@@ -76,7 +78,16 @@ IGeometrySelectionTransformer* FStaticMeshSelector::InitializeTransformation(con
 {
 	check(!ActiveTransformer);
 
-	ActiveTransformer = MakePimpl<FBasicDynamicMeshSelectionTransformer>();
+	if (Selection.TopologyType == EGeometryTopologyType::Polygroup)
+	{
+		ActiveTransformer = MakeShared<FDynamicMeshPolygroupTransformer>();
+	}
+	else
+	{
+		ActiveTransformer = MakeShared<FBasicDynamicMeshSelectionTransformer>();
+	}
+	ActiveTransformer->bEnableSelectionTransformDrawing = true;
+
 	ActiveTransformer->Initialize(this);
 	ActiveTransformer->OnEndTransformFunc = [this](IToolsContextTransactionsAPI*) 
 	{ 
diff --git a/Engine/Plugins/Runtime/MeshModelingToolset/Source/ModelingComponentsEditorOnly/Private/Selection/VolumeSelector.cpp b/Engine/Plugins/Runtime/MeshModelingToolset/Source/ModelingComponentsEditorOnly/Private/Selection/VolumeSelector.cpp
--- a/Engine/Plugins/Runtime/MeshModelingToolset/Source/ModelingComponentsEditorOnly/Private/Selection/VolumeSelector.cpp
+++ b/Engine/Plugins/Runtime/MeshModelingToolset/Source/ModelingComponentsEditorOnly/Private/Selection/VolumeSelector.cpp
@@ -16,6 +16,8 @@
 #include "ConversionUtils/VolumeToDynamicMesh.h"
 #include "ConversionUtils/DynamicMeshToVolume.h"
 
+#include "Selection/DynamicMeshPolygroupTransformer.h"
+
 using namespace UE::Geometry;
 
 #define LOCTEXT_NAMESPACE "FVolumeSelector"
@@ -68,7 +70,16 @@ IGeometrySelectionTransformer* FVolumeSelector::InitializeTransformation(const F
 {
 	check(!ActiveTransformer);
 
-	ActiveTransformer = MakePimpl<FBasicDynamicMeshSelectionTransformer>();
+	if (Selection.TopologyType == EGeometryTopologyType::Polygroup)
+	{
+		ActiveTransformer = MakeShared<FDynamicMeshPolygroupTransformer>();
+	}
+	else
+	{
+		ActiveTransformer = MakeShared<FBasicDynamicMeshSelectionTransformer>();
+	}
+	ActiveTransformer->bEnableSelectionTransformDrawing = true;
+
 	ActiveTransformer->Initialize(this);
 	ActiveTransformer->OnEndTransformFunc = [this](IToolsContextTransactionsAPI*) 
 	{ 
diff --git a/Engine/Plugins/Runtime/MeshModelingToolset/Source/ModelingComponentsEditorOnly/Public/Selection/StaticMeshSelector.h b/Engine/Plugins/Runtime/MeshModelingToolset/Source/ModelingComponentsEditorOnly/Public/Selection/StaticMeshSelector.h
--- a/Engine/Plugins/Runtime/MeshModelingToolset/Source/ModelingComponentsEditorOnly/Public/Selection/StaticMeshSelector.h
+++ b/Engine/Plugins/Runtime/MeshModelingToolset/Source/ModelingComponentsEditorOnly/Public/Selection/StaticMeshSelector.h
@@ -46,7 +46,7 @@ protected:
 	TStrongObjectPtr<UDynamicMesh> LocalTargetMesh;
 	void CopyFromStaticMesh();
 
-	TPimplPtr<FBasicDynamicMeshSelectionTransformer> ActiveTransformer;
+	TSharedPtr<FBasicDynamicMeshSelectionTransformer> ActiveTransformer;
 	void CommitMeshTransform();
 };
 
diff --git a/Engine/Plugins/Runtime/MeshModelingToolset/Source/ModelingComponentsEditorOnly/Public/Selection/VolumeSelector.h b/Engine/Plugins/Runtime/MeshModelingToolset/Source/ModelingComponentsEditorOnly/Public/Selection/VolumeSelector.h
--- a/Engine/Plugins/Runtime/MeshModelingToolset/Source/ModelingComponentsEditorOnly/Public/Selection/VolumeSelector.h
+++ b/Engine/Plugins/Runtime/MeshModelingToolset/Source/ModelingComponentsEditorOnly/Public/Selection/VolumeSelector.h
@@ -43,7 +43,7 @@ protected:
 	// TODO: this is not a great design, it would be better if something external could own this mesh...
 	TStrongObjectPtr<UDynamicMesh> LocalTargetMesh;
 
-	TPimplPtr<FBasicDynamicMeshSelectionTransformer> ActiveTransformer;
+	TSharedPtr<FBasicDynamicMeshSelectionTransformer> ActiveTransformer;
 	void CommitMeshTransform();
 };
 
