ModelingTools: add support for DynamicMeshComponent Simple Collision in Mesh to Collision Tool, Convert Tool, and Transfer Tool.
ModelingComponents: Add ShapeSet member for passing simple collision shapes in FCreateMeshObjectParams, support in relevant functions in UEditorModelingObjectsCreationAPI. Add UE::Geometry::GetCollisionShapes() functions in ComponentCollisionUtil.h, and update various functions to support DynamicMeshComponent

DynamicMeshComponent: change signature of GetSimpleCollisionShapes() to return const ref

diff --git a/Engine/Plugins/Experimental/MeshModelingToolsetExp/Source/MeshModelingToolsExp/Private/ConvertMeshesTool.cpp b/Engine/Plugins/Experimental/MeshModelingToolsetExp/Source/MeshModelingToolsExp/Private/ConvertMeshesTool.cpp
--- a/Engine/Plugins/Experimental/MeshModelingToolsetExp/Source/MeshModelingToolsExp/Private/ConvertMeshesTool.cpp
+++ b/Engine/Plugins/Experimental/MeshModelingToolsetExp/Source/MeshModelingToolsExp/Private/ConvertMeshesTool.cpp
@@ -9,6 +9,8 @@
 #include "ModelingToolTargetUtil.h"
 #include "ModelingObjectsCreationAPI.h"
 #include "Selection/ToolSelectionUtil.h"
+#include "Physics/ComponentCollisionUtil.h"
+#include "ShapeApproximation/SimpleShapeSet3.h"
 
 #include "TargetInterfaces/MeshDescriptionProvider.h"
 #include "TargetInterfaces/PrimitiveComponentBackedTarget.h"
@@ -112,6 +114,23 @@ void UConvertMeshesTool::OnShutdown(EToolShutdownType ShutdownType)
 				NewMeshObjectParams.Materials = TransferMaterials->Materials;
 			}
 			NewMeshObjectParams.SetMesh(MoveTemp(SourceMesh));
+			
+			if (BasicProperties->bTransferCollision)
+			{
+				UPrimitiveComponent* SourceComponent = UE::ToolTarget::GetTargetComponent(Targets[k]);
+				if (UE::Geometry::ComponentTypeSupportsCollision(SourceComponent, UE::Geometry::EComponentCollisionSupportLevel::ReadOnly))
+				{
+					NewMeshObjectParams.bEnableCollision = true;
+					FComponentCollisionSettings CollisionSettings = UE::Geometry::GetCollisionSettings(SourceComponent);
+					NewMeshObjectParams.CollisionMode = (ECollisionTraceFlag)CollisionSettings.CollisionTypeFlag;
+
+					FSimpleShapeSet3d ShapeSet;
+					if (UE::Geometry::GetCollisionShapes(SourceComponent, ShapeSet))
+					{
+						NewMeshObjectParams.CollisionShapeSet = MoveTemp(ShapeSet);
+					}
+				}
+			}
 
 			NewMeshObjects.Add(MoveTemp(NewMeshObjectParams));
 		}
diff --git a/Engine/Plugins/Experimental/MeshModelingToolsetExp/Source/MeshModelingToolsExp/Private/Physics/SetCollisionGeometryTool.cpp b/Engine/Plugins/Experimental/MeshModelingToolsetExp/Source/MeshModelingToolsExp/Private/Physics/SetCollisionGeometryTool.cpp
--- a/Engine/Plugins/Experimental/MeshModelingToolsetExp/Source/MeshModelingToolsExp/Private/Physics/SetCollisionGeometryTool.cpp
+++ b/Engine/Plugins/Experimental/MeshModelingToolsetExp/Source/MeshModelingToolsExp/Private/Physics/SetCollisionGeometryTool.cpp
@@ -152,7 +152,12 @@ bool USetCollisionGeometryToolBuilder::CanBuildTool(const FToolBuilderState& Sce
 	UActorComponent* LastValidTarget = nullptr;
 	SceneState.TargetManager->EnumerateSelectedAndTargetableComponents(SceneState, GetTargetRequirements(),
 		[&](UActorComponent* Component) { LastValidTarget = Component; });
-	return (LastValidTarget != nullptr && Cast<UStaticMeshComponent>(LastValidTarget) != nullptr);
+	if (LastValidTarget != nullptr)
+	{
+		return (Cast<UStaticMeshComponent>(LastValidTarget) != nullptr) ||
+			(Cast<UDynamicMeshComponent>(LastValidTarget) != nullptr);
+	}
+	return false;
 }
 
 
@@ -347,13 +352,8 @@ void USetCollisionGeometryTool::OnShutdown(EToolShutdownType ShutdownType)
 
 		GetToolManager()->BeginUndoTransaction(LOCTEXT("UpdateCollision", "Update Collision"));
 
-		// code below derived from FStaticMeshEditor::DuplicateSelectedPrims(), FStaticMeshEditor::OnCollisionSphere(), and GeomFitUtils.cpp::GenerateSphylAsSimpleCollision()
 
-		UPrimitiveComponent* Component = UE::ToolTarget::GetTargetComponent(Targets[Targets.Num() - 1]);
-		UStaticMeshComponent* StaticMeshComponent = Cast<UStaticMeshComponent>(Component);
-		TObjectPtr<UStaticMesh> StaticMesh = (StaticMeshComponent) ? StaticMeshComponent->GetStaticMesh() : nullptr;
-		UBodySetup* BodySetup = (StaticMesh) ? StaticMesh->GetBodySetup() : nullptr;
-		if (BodySetup != nullptr)
+		auto UpdateBodySetup = [this](UBodySetup* BodySetup)
 		{
 			// mark the BodySetup for modification. Do we need to modify the UStaticMesh??
 			BodySetup->Modify();
@@ -369,33 +369,53 @@ void USetCollisionGeometryTool::OnShutdown(EToolShutdownType ShutdownType)
 
 			// rebuild physics meshes
 			BodySetup->CreatePhysicsMeshes();
+		};
 
-			StaticMesh->RecreateNavCollision();
 
-			// update physics state on all components using this StaticMesh
-			for (FThreadSafeObjectIterator Iter(UStaticMeshComponent::StaticClass()); Iter; ++Iter)
+		UPrimitiveComponent* Component = UE::ToolTarget::GetTargetComponent(Targets[Targets.Num() - 1]);
+		if (UStaticMeshComponent* StaticMeshComponent = Cast<UStaticMeshComponent>(Component))
+		{
+			// code below derived from FStaticMeshEditor::DuplicateSelectedPrims(), FStaticMeshEditor::OnCollisionSphere(), and GeomFitUtils.cpp::GenerateSphylAsSimpleCollision()
+			TObjectPtr<UStaticMesh> StaticMesh = (StaticMeshComponent) ? StaticMeshComponent->GetStaticMesh() : nullptr;
+			UBodySetup* BodySetup = (StaticMesh) ? StaticMesh->GetBodySetup() : nullptr;
+			if (BodySetup != nullptr)
 			{
-				UStaticMeshComponent* SMComponent = Cast<UStaticMeshComponent>(*Iter);
-				if (SMComponent->GetStaticMesh() == StaticMesh)
+				UpdateBodySetup(BodySetup);
+
+				StaticMesh->RecreateNavCollision();
+
+				// update physics state on all components using this StaticMesh
+				for (FThreadSafeObjectIterator Iter(UStaticMeshComponent::StaticClass()); Iter; ++Iter)
 				{
-					if (SMComponent->IsPhysicsStateCreated())
+					UStaticMeshComponent* SMComponent = Cast<UStaticMeshComponent>(*Iter);
+					if (SMComponent->GetStaticMesh() == StaticMesh)
 					{
-						SMComponent->RecreatePhysicsState();
+						if (SMComponent->IsPhysicsStateCreated())
+						{
+							SMComponent->RecreatePhysicsState();
+						}
+						// Mark the render state dirty to make sure any CollisionTraceFlag changes get picked up
+						SMComponent->MarkRenderStateDirty();
 					}
-					// Mark the render state dirty to make sure any CollisionTraceFlag changes get picked up
-					SMComponent->MarkRenderStateDirty();
 				}
-			}
 
-			// do we need to do a post edit change here??
+				// do we need to do a post edit change here??
 
-			// mark static mesh as dirty so it gets resaved?
-			StaticMesh->MarkPackageDirty();
+				// mark static mesh as dirty so it gets resaved?
+				StaticMesh->MarkPackageDirty();
 
 #if WITH_EDITORONLY_DATA
-			// mark the static mesh as having customized collision so it is not regenerated on reimport
-			StaticMesh->bCustomizedCollision = true;
+				// mark the static mesh as having customized collision so it is not regenerated on reimport
+				StaticMesh->bCustomizedCollision = true;
 #endif // WITH_EDITORONLY_DATA
+			}
+		}
+		else if (UDynamicMeshComponent* DynamicMeshComponent = Cast<UDynamicMeshComponent>(Component))
+		{
+			if (UBodySetup* BodySetup = DynamicMeshComponent->GetBodySetup())
+			{
+				UpdateBodySetup(BodySetup);
+			}
 		}
 
 		// post the undo transaction
diff --git a/Engine/Plugins/Experimental/MeshModelingToolsetExp/Source/MeshModelingToolsExp/Private/TransferMeshTool.cpp b/Engine/Plugins/Experimental/MeshModelingToolsetExp/Source/MeshModelingToolsExp/Private/TransferMeshTool.cpp
--- a/Engine/Plugins/Experimental/MeshModelingToolsetExp/Source/MeshModelingToolsExp/Private/TransferMeshTool.cpp
+++ b/Engine/Plugins/Experimental/MeshModelingToolsetExp/Source/MeshModelingToolsExp/Private/TransferMeshTool.cpp
@@ -10,6 +10,8 @@
 #include "TargetInterfaces/MeshDescriptionProvider.h"
 #include "TargetInterfaces/MeshDescriptionCommitter.h"
 #include "ModelingToolTargetUtil.h"
+#include "Physics/ComponentCollisionUtil.h"
+#include "ShapeApproximation/SimpleShapeSet3.h"
 #include "Engine/StaticMesh.h"
 
 #include UE_INLINE_GENERATED_CPP_BY_NAME(TransferMeshTool)
@@ -181,6 +183,26 @@ void UTransferMeshTool::OnShutdown(EToolShutdownType ShutdownType)
 			UE::ToolTarget::CommitMeshDescriptionUpdate(Targets[1], &SourceMesh, TransferMaterials);
 		}
 
+
+		if (BasicProperties->bTransferCollision)
+		{
+			UPrimitiveComponent* SourceComponent = UE::ToolTarget::GetTargetComponent(Targets[0]);
+			if (UE::Geometry::ComponentTypeSupportsCollision(SourceComponent, UE::Geometry::EComponentCollisionSupportLevel::ReadOnly))
+			{
+				UPrimitiveComponent* TargetComponent = UE::ToolTarget::GetTargetComponent(Targets[1]);
+				if (UE::Geometry::ComponentTypeSupportsCollision(TargetComponent))
+				{
+					FComponentCollisionSettings SourceCollisionSettings = UE::Geometry::GetCollisionSettings(SourceComponent);
+
+					FSimpleShapeSet3d SourceShapeSet;
+					if (UE::Geometry::GetCollisionShapes(SourceComponent, SourceShapeSet))
+					{
+						UE::Geometry::SetSimpleCollision(TargetComponent, &SourceShapeSet, SourceCollisionSettings);
+					}
+				}
+			}
+		}
+
 		GetToolManager()->EndUndoTransaction();
 	}
 }
diff --git a/Engine/Plugins/Experimental/MeshModelingToolsetExp/Source/MeshModelingToolsExp/Public/ConvertMeshesTool.h b/Engine/Plugins/Experimental/MeshModelingToolsetExp/Source/MeshModelingToolsExp/Public/ConvertMeshesTool.h
--- a/Engine/Plugins/Experimental/MeshModelingToolsetExp/Source/MeshModelingToolsExp/Public/ConvertMeshesTool.h
+++ b/Engine/Plugins/Experimental/MeshModelingToolsetExp/Source/MeshModelingToolsExp/Public/ConvertMeshesTool.h
@@ -37,6 +37,9 @@ public:
 	// control whether the transfer materials option is displayed
 	UPROPERTY(meta = (TransientToolProperty))
 	bool bShowTransferMaterials = true;
+
+	UPROPERTY(EditAnywhere, Category = Options)
+	bool bTransferCollision = true;
 };
 
 
diff --git a/Engine/Plugins/Experimental/MeshModelingToolsetExp/Source/MeshModelingToolsExp/Public/TransferMeshTool.h b/Engine/Plugins/Experimental/MeshModelingToolsetExp/Source/MeshModelingToolsExp/Public/TransferMeshTool.h
--- a/Engine/Plugins/Experimental/MeshModelingToolsetExp/Source/MeshModelingToolsExp/Public/TransferMeshTool.h
+++ b/Engine/Plugins/Experimental/MeshModelingToolsetExp/Source/MeshModelingToolsExp/Public/TransferMeshTool.h
@@ -33,6 +33,9 @@ public:
 	UPROPERTY(EditAnywhere, Category = Options)
 	bool bTransferMaterials = true;
 
+	UPROPERTY(EditAnywhere, Category = Options)
+	bool bTransferCollision = false;
+
 	/** Specify which LOD from the Source (First) mesh to copy from */
 	UPROPERTY(EditAnywhere, Category = SourceMesh, meta = (DisplayName = "Source LOD", TransientToolProperty, NoResetToDefault, GetOptions = GetSourceLODNamesFunc, HideEditConditionToggle, EditCondition="bIsStaticMeshSource", EditConditionHides))
 	FString SourceLOD;
diff --git a/Engine/Plugins/Runtime/MeshModelingToolset/Source/ModelingComponents/Private/Physics/ComponentCollisionUtil.cpp b/Engine/Plugins/Runtime/MeshModelingToolset/Source/ModelingComponents/Private/Physics/ComponentCollisionUtil.cpp
--- a/Engine/Plugins/Runtime/MeshModelingToolset/Source/ModelingComponents/Private/Physics/ComponentCollisionUtil.cpp
+++ b/Engine/Plugins/Runtime/MeshModelingToolset/Source/ModelingComponents/Private/Physics/ComponentCollisionUtil.cpp
@@ -14,6 +14,7 @@
 #include "DynamicMesh/Operations/MergeCoincidentMeshEdges.h"
 
 #include "Physics/PhysicsDataCollection.h"
+#include "Physics/CollisionGeometryConversion.h"
 
 #include "Components/BrushComponent.h"
 #include "Components/DynamicMeshComponent.h"
@@ -26,11 +27,22 @@
 using namespace UE::Geometry;
 
 bool UE::Geometry::ComponentTypeSupportsCollision(
-	const UPrimitiveComponent* Component)
+	const UPrimitiveComponent* Component, EComponentCollisionSupportLevel SupportLevel)
 {
-	// currently only supporting StaticMeshComponent
-	const UStaticMeshComponent* StaticMeshComponent = Cast<UStaticMeshComponent>(Component);
-	return (StaticMeshComponent != nullptr);
+	// currently only supporting StaticMeshComponent and DynamicMeshComponent
+	if (Cast<UStaticMeshComponent>(Component) != nullptr)
+	{
+		return true; // ReadOnly and ReadWrite both supported
+	}
+	if (Cast<UDynamicMeshComponent>(Component) != nullptr)
+	{
+		return true; // ReadOnly and ReadWrite both supported
+	}
+	if (Cast<UBrushComponent>(Component) != nullptr)
+	{
+		return SupportLevel == EComponentCollisionSupportLevel::ReadOnly; // do not support write on BrushComponent
+	}
+	return false;
 }
 
 
@@ -38,23 +50,73 @@ FComponentCollisionSettings UE::Geometry::GetCollisionSettings(const UPrimitiveC
 {
 	FComponentCollisionSettings Settings;
 
-	const UStaticMeshComponent* StaticMeshComponent = Cast<UStaticMeshComponent>(Component);
-	if (ensure(StaticMeshComponent))
+	if (const UStaticMeshComponent* StaticMeshComponent = Cast<UStaticMeshComponent>(Component) )
 	{
-		const UStaticMesh* StaticMesh = StaticMeshComponent->GetStaticMesh();
-		if (ensure(StaticMesh))
+		if (const UBodySetup* BodySetup = GetBodySetup(Component))
 		{
-			UBodySetup* BodySetup = StaticMesh->GetBodySetup();
-			if (ensure(BodySetup))
-			{
-				Settings.CollisionTypeFlag = (int32)BodySetup->CollisionTraceFlag;
-			}
+			Settings.CollisionTypeFlag = (int32)BodySetup->CollisionTraceFlag;
+		}
+	}
+	else if (const UDynamicMeshComponent* DynamicMeshComponent = Cast<UDynamicMeshComponent>(Component))
+	{
+		Settings.CollisionTypeFlag = DynamicMeshComponent->CollisionType;
+	}
+	else if (const UBrushComponent* BrushComponent = Cast<UBrushComponent>(Component))
+	{
+		if (BrushComponent->BrushBodySetup)
+		{
+			Settings.CollisionTypeFlag = (int32)BrushComponent->BrushBodySetup->CollisionTraceFlag;
 		}
 	}
+
 	return Settings;
 }
 
 
+bool UE::Geometry::GetCollisionShapes(const UPrimitiveComponent* Component, FKAggregateGeom& AggGeom)
+{
+	if (const UStaticMeshComponent* StaticMeshComponent = Cast<UStaticMeshComponent>(Component))
+	{
+		if (const UBodySetup* BodySetup = GetBodySetup(Component))
+		{
+			AggGeom = BodySetup->AggGeom;
+			return true;
+		}
+	}
+	else if (const UDynamicMeshComponent* DynamicMeshComponent = Cast<UDynamicMeshComponent>(Component))
+	{
+		AggGeom = DynamicMeshComponent->GetSimpleCollisionShapes();
+		return true;
+	}
+	else if (const UBrushComponent* BrushComponent = Cast<UBrushComponent>(Component))
+	{
+		if (BrushComponent->BrushBodySetup)
+		{
+			AggGeom = BrushComponent->BrushBodySetup->AggGeom;
+			return true;
+		}
+	}
+
+	return false;
+}
+
+bool UE::Geometry::GetCollisionShapes(const UPrimitiveComponent* Component, FSimpleShapeSet3d& ShapeSet)
+{
+	FKAggregateGeom AggGeom;
+	if (UE::Geometry::GetCollisionShapes(Component, AggGeom))
+	{
+		UE::Geometry::GetShapeSet(AggGeom, ShapeSet);
+		return true;
+	}
+
+	return false;
+}
+
+
+
+
+
+
 void UE::Geometry::UpdateSimpleCollision(
 	UBodySetup* BodySetup, 
 	const FKAggregateGeom* NewGeometry,
@@ -73,31 +135,36 @@ void UE::Geometry::UpdateSimpleCollision(
 	// rebuild physics meshes
 	BodySetup->CreatePhysicsMeshes();
 
-	StaticMesh->RecreateNavCollision();
-
-	// update physics state on all components using this StaticMesh
-	for (FThreadSafeObjectIterator Iter(UStaticMeshComponent::StaticClass()); Iter; ++Iter)
+	if (StaticMesh != nullptr)
 	{
-		UStaticMeshComponent* SMComponent = Cast<UStaticMeshComponent>(*Iter);
-		if (SMComponent->GetStaticMesh() == StaticMesh)
+		StaticMesh->RecreateNavCollision();
+
+		// update physics state on all components using this StaticMesh
+		for (FThreadSafeObjectIterator Iter(UStaticMeshComponent::StaticClass()); Iter; ++Iter)
 		{
-			if (SMComponent->IsPhysicsStateCreated())
+			UStaticMeshComponent* SMComponent = Cast<UStaticMeshComponent>(*Iter);
+			if (SMComponent->GetStaticMesh() == StaticMesh)
 			{
-				SMComponent->RecreatePhysicsState();
+				if (SMComponent->IsPhysicsStateCreated())
+				{
+					SMComponent->RecreatePhysicsState();
+				}
 			}
 		}
-	}
 
-	// mark static mesh as dirty so it gets resaved?
-	[[maybe_unused]] bool MarkedDirty = StaticMesh->MarkPackageDirty();
+		// mark static mesh as dirty so it gets resaved?
+		[[maybe_unused]] bool MarkedDirty = StaticMesh->MarkPackageDirty();
 
 #if WITH_EDITORONLY_DATA
-	// mark the static mesh as having customized collision so it is not regenerated on reimport
-	StaticMesh->bCustomizedCollision = CollisionSettings.bIsGeneratedCollision;
+		// mark the static mesh as having customized collision so it is not regenerated on reimport
+		StaticMesh->bCustomizedCollision = CollisionSettings.bIsGeneratedCollision;
 #endif // WITH_EDITORONLY_DATA
+	}
 }
 
 
+
+
 const UBodySetup* UE::Geometry::GetBodySetup(const UPrimitiveComponent* Component)
 {
 	if (const UStaticMeshComponent* StaticMeshComponent = Cast<UStaticMeshComponent>(Component))
@@ -157,12 +224,25 @@ bool UE::Geometry::SetSimpleCollision(
 	PhysicsData.Geometry = *ShapeSet;
 	PhysicsData.CopyGeometryToAggregate();
 
-	// FPhysicsDataCollection stores its references as const, but the input Component was non-const so this is ok to do
-	UStaticMesh* StaticMesh = const_cast<UStaticMesh*>(PhysicsData.SourceStaticMesh.Get());
-	UBodySetup* BodySetup = StaticMesh->GetBodySetup();
-	UpdateSimpleCollision(BodySetup, &PhysicsData.AggGeom, StaticMesh, CollisionSettings);
+	if (UStaticMeshComponent* StaticMeshComponent = Cast<UStaticMeshComponent>(Component))
+	{
+		if (UStaticMesh* StaticMesh = StaticMeshComponent->GetStaticMesh())
+		{
+			if (UBodySetup* BodySetup = StaticMesh->GetBodySetup())
+			{
+				UpdateSimpleCollision(BodySetup, &PhysicsData.AggGeom, StaticMesh, CollisionSettings);
+				return true;
+			}
+		}
+	}
+	else if (UDynamicMeshComponent* DynamicMeshComponent = Cast<UDynamicMeshComponent>(Component))
+	{
+		DynamicMeshComponent->CollisionType = (ECollisionTraceFlag)CollisionSettings.CollisionTypeFlag;
+		DynamicMeshComponent->SetSimpleCollisionShapes(PhysicsData.AggGeom, true);
+		return true;
+	}
 
-	return true;
+	return false;
 }
 
 
@@ -178,19 +258,29 @@ bool UE::Geometry::TransformSimpleCollision(
 		return false;
 	}
 
-	FComponentCollisionSettings Settings = GetCollisionSettings(Component);
-	Settings.bIsGeneratedCollision = false;
-
 	PhysicsData.Geometry.ApplyTransform(Transform);
 	PhysicsData.ClearAggregate();
 	PhysicsData.CopyGeometryToAggregate();
 
-	// FPhysicsDataCollection stores its references as const, but the input Component was non-const so this is ok to do
-	UStaticMesh* StaticMesh = const_cast<UStaticMesh*>(PhysicsData.SourceStaticMesh.Get());
-	UBodySetup* BodySetup = StaticMesh->GetBodySetup();
-	UpdateSimpleCollision(BodySetup, &PhysicsData.AggGeom, StaticMesh, Settings);
-
-	return true;
+	if (UStaticMeshComponent* StaticMeshComponent = Cast<UStaticMeshComponent>(Component))
+	{
+		FComponentCollisionSettings Settings = GetCollisionSettings(Component);
+		Settings.bIsGeneratedCollision = false;
+		if (UStaticMesh* StaticMesh = StaticMeshComponent->GetStaticMesh())
+		{
+			if (UBodySetup* BodySetup = StaticMesh->GetBodySetup())
+			{
+				UpdateSimpleCollision(BodySetup, &PhysicsData.AggGeom, StaticMesh, Settings);
+				return true;
+			}
+		}
+	}
+	else if (UDynamicMeshComponent* DynamicMeshComponent = Cast<UDynamicMeshComponent>(Component))
+	{
+		DynamicMeshComponent->SetSimpleCollisionShapes(PhysicsData.AggGeom, true);
+		return true;
+	}
+	return false;
 }
 
 
diff --git a/Engine/Plugins/Runtime/MeshModelingToolset/Source/ModelingComponents/Public/ModelingObjectsCreationAPI.h b/Engine/Plugins/Runtime/MeshModelingToolset/Source/ModelingComponents/Public/ModelingObjectsCreationAPI.h
--- a/Engine/Plugins/Runtime/MeshModelingToolset/Source/ModelingComponents/Public/ModelingObjectsCreationAPI.h
+++ b/Engine/Plugins/Runtime/MeshModelingToolset/Source/ModelingComponents/Public/ModelingObjectsCreationAPI.h
@@ -8,6 +8,7 @@
 #include "Engine/EngineTypes.h"   // FMeshNaniteSettings
 #include "MeshDescription.h"
 #include "DynamicMesh/DynamicMesh3.h"
+#include "ShapeApproximation/SimpleShapeSet3.h"
 #include "Misc/Optional.h"
 
 #include "ModelingObjectsCreationAPI.generated.h"
@@ -138,6 +139,9 @@ struct MODELINGCOMPONENTS_API FCreateMeshObjectParams
 	UPROPERTY(Category = "CreateMeshObjectParams", EditAnywhere)
 	TEnumAsByte<enum ECollisionTraceFlag> CollisionMode = ECollisionTraceFlag::CTF_UseComplexAsSimple;
 
+	/** Collision Shapes */
+	TOptional<UE::Geometry::FSimpleShapeSet3d> CollisionShapeSet;
+
 
 	//
 	// Rendering Configuration Options, if this is applicable for the given mesh object
diff --git a/Engine/Plugins/Runtime/MeshModelingToolset/Source/ModelingComponents/Public/Physics/ComponentCollisionUtil.h b/Engine/Plugins/Runtime/MeshModelingToolset/Source/ModelingComponents/Public/Physics/ComponentCollisionUtil.h
--- a/Engine/Plugins/Runtime/MeshModelingToolset/Source/ModelingComponents/Public/Physics/ComponentCollisionUtil.h
+++ b/Engine/Plugins/Runtime/MeshModelingToolset/Source/ModelingComponents/Public/Physics/ComponentCollisionUtil.h
@@ -29,11 +29,18 @@ struct MODELINGCOMPONENTS_API FComponentCollisionSettings
 	bool bIsGeneratedCollision = true;
 };
 
+
+enum EComponentCollisionSupportLevel
+{
+	ReadOnly = 0,
+	ReadWrite = 1
+};
+
 /**
  * @return true if the component type supports collision settings
  */
 MODELINGCOMPONENTS_API bool ComponentTypeSupportsCollision(
-	const UPrimitiveComponent* Component);
+	const UPrimitiveComponent* Component, EComponentCollisionSupportLevel SupportLevel = EComponentCollisionSupportLevel::ReadWrite);
 
 
 /**
@@ -42,6 +49,17 @@ MODELINGCOMPONENTS_API bool ComponentTypeSupportsCollision(
 MODELINGCOMPONENTS_API FComponentCollisionSettings GetCollisionSettings(
 	const UPrimitiveComponent* Component);
 
+/**
+ * Get current Component collision shapes 
+ */
+MODELINGCOMPONENTS_API bool GetCollisionShapes(const UPrimitiveComponent* Component, FKAggregateGeom& AggGeomOut);
+
+/**
+ * Get current Component collision shapes 
+ */
+MODELINGCOMPONENTS_API bool GetCollisionShapes(const UPrimitiveComponent* Component, FSimpleShapeSet3d& ShapeSetOut);
+
+
 
 /**
  * Apply Transform to any Simple Collision geometry owned by Component.
@@ -79,14 +97,14 @@ MODELINGCOMPONENTS_API bool AppendSimpleCollision(
 
 /**
  * Replace existing Simple Collision geometry in BodySetup with that defined by NewGeometry,
- * and update the Component/BodySetup collision settings. Optional StaticMesh argument allows
+ * and update the Component/BodySetup collision settings. StaticMesh argument allows
  * for necessary updates to it and any active UStaticMeshComponent that reference it
  */
 MODELINGCOMPONENTS_API void UpdateSimpleCollision(
 	UBodySetup* BodySetup,
 	const FKAggregateGeom* NewGeometry,
-	UStaticMesh* StaticMesh = nullptr,
-	FComponentCollisionSettings CollisionSettings = FComponentCollisionSettings());
+	UStaticMesh* StaticMesh,
+	FComponentCollisionSettings NewCollisionSettings = FComponentCollisionSettings() );
 
 /**
  * @return BodySetup on the given SourceComponent, or nullptr if no BodySetup was found
diff --git a/Engine/Plugins/Runtime/MeshModelingToolset/Source/ModelingComponentsEditorOnly/Private/EditorModelingObjectsCreationAPI.cpp b/Engine/Plugins/Runtime/MeshModelingToolset/Source/ModelingComponentsEditorOnly/Private/EditorModelingObjectsCreationAPI.cpp
--- a/Engine/Plugins/Runtime/MeshModelingToolset/Source/ModelingComponentsEditorOnly/Private/EditorModelingObjectsCreationAPI.cpp
+++ b/Engine/Plugins/Runtime/MeshModelingToolset/Source/ModelingComponentsEditorOnly/Private/EditorModelingObjectsCreationAPI.cpp
@@ -7,6 +7,7 @@
 
 #include "AssetUtils/CreateStaticMeshUtil.h"
 #include "AssetUtils/CreateTexture2DUtil.h"
+#include "Physics/ComponentCollisionUtil.h"
 
 #include "ConversionUtils/DynamicMeshToVolume.h"
 #include "MeshDescriptionToDynamicMesh.h"
@@ -241,9 +242,15 @@ FCreateMeshObjectResult UEditorModelingObjectsCreationAPI::CreateDynamicMeshActo
 	// configure collision
 	if (CreateMeshParams.bEnableCollision)
 	{
+		if (CreateMeshParams.CollisionShapeSet.IsSet())
+		{
+			UE::Geometry::SetSimpleCollision(NewComponent, CreateMeshParams.CollisionShapeSet.GetPtrOrNull());
+		}
+
 		NewComponent->CollisionType = CreateMeshParams.CollisionMode;
 		// enable complex collision so that raycasts can hit this object
 		NewComponent->bEnableComplexCollision = true;
+
 		// force collision update
 		NewComponent->UpdateCollision(false);
 	}
@@ -357,6 +364,13 @@ FCreateMeshObjectResult UEditorModelingObjectsCreationAPI::CreateStaticMeshAsset
 		StaticMeshComponent->SetMaterial(k, ComponentMaterials[k]);
 	}
 
+	// set simple collision geometry
+	if (CreateMeshParams.CollisionShapeSet.IsSet())
+	{
+		UE::Geometry::SetSimpleCollision(StaticMeshComponent, CreateMeshParams.CollisionShapeSet.GetPtrOrNull(),
+			UE::Geometry::GetCollisionSettings(StaticMeshComponent));
+	}
+
 	// re-connect the component (?)
 	StaticMeshComponent->RegisterComponent();
 
diff --git a/Engine/Source/Runtime/GeometryFramework/Public/Components/DynamicMeshComponent.h b/Engine/Source/Runtime/GeometryFramework/Public/Components/DynamicMeshComponent.h
--- a/Engine/Source/Runtime/GeometryFramework/Public/Components/DynamicMeshComponent.h
+++ b/Engine/Source/Runtime/GeometryFramework/Public/Components/DynamicMeshComponent.h
@@ -508,7 +508,7 @@ public:
 	 */
 	virtual void SetSimpleCollisionShapes(const struct FKAggregateGeom& AggGeom, bool bUpdateCollision);
 
-	virtual struct FKAggregateGeom& GetSimpleCollisionShapes()
+	const FKAggregateGeom& GetSimpleCollisionShapes() const
 	{
 		return AggGeom;
 	}
