ModelingComponents: remove code that explicitly sets static meshes to Movable, they should be Static mobility by default, like placed static meshes

diff --git a/Engine/Plugins/Runtime/MeshModelingToolset/Source/ModelingComponentsEditorOnly/Private/EditorModelingObjectsCreationAPI.cpp b/Engine/Plugins/Runtime/MeshModelingToolset/Source/ModelingComponentsEditorOnly/Private/EditorModelingObjectsCreationAPI.cpp
--- a/Engine/Plugins/Runtime/MeshModelingToolset/Source/ModelingComponentsEditorOnly/Private/EditorModelingObjectsCreationAPI.cpp
+++ b/Engine/Plugins/Runtime/MeshModelingToolset/Source/ModelingComponentsEditorOnly/Private/EditorModelingObjectsCreationAPI.cpp
@@ -351,9 +351,6 @@ FCreateMeshObjectResult UEditorModelingObjectsCreationAPI::CreateStaticMeshAsset
 
 	// this disconnects the component from various events
 	StaticMeshComponent->UnregisterComponent();
-	// Configure flags of the component. Is this necessary?
-	StaticMeshComponent->SetMobility(EComponentMobility::Movable);
-	StaticMeshComponent->bSelectable = true;
 	// replace the UStaticMesh in the component
 	StaticMeshComponent->SetStaticMesh(NewStaticMesh);
 
