ModelingMode: Adjusted the conditions used for checking if targets can build dynamic meshes for modeling tools to include proper checks for cooked assets in the editor.

diff --git a/Engine/Plugins/Runtime/MeshModelingToolset/Source/ModelingComponentsEditorOnly/Private/ToolTargets/SkeletalMeshComponentToolTarget.cpp b/Engine/Plugins/Runtime/MeshModelingToolset/Source/ModelingComponentsEditorOnly/Private/ToolTargets/SkeletalMeshComponentToolTarget.cpp
--- a/Engine/Plugins/Runtime/MeshModelingToolset/Source/ModelingComponentsEditorOnly/Private/ToolTargets/SkeletalMeshComponentToolTarget.cpp
+++ b/Engine/Plugins/Runtime/MeshModelingToolset/Source/ModelingComponentsEditorOnly/Private/ToolTargets/SkeletalMeshComponentToolTarget.cpp
@@ -9,6 +9,7 @@
 #include "Rendering/SkeletalMeshModel.h"
 #include "ToolTargets/SkeletalMeshToolTarget.h"
 #include "SkeletalMeshAttributes.h"
+#include "UObject/Package.h"
 
 #include UE_INLINE_GENERATED_CPP_BY_NAME(SkeletalMeshComponentToolTarget)
 
@@ -197,6 +198,7 @@ bool USkeletalMeshComponentReadOnlyToolTargetFactory::CanBuildTarget(UObject* So
 
 	return Cast<USkeletalMeshComponent>(SourceObject) && Cast<USkeletalMeshComponent>(SourceObject)->GetSkeletalMeshAsset() &&
 		ExactCast<USkeletalMesh>(Cast<USkeletalMeshComponent>(SourceObject)->GetSkeletalMeshAsset()) &&
+		!ExactCast<USkeletalMesh>(Cast<USkeletalMeshComponent>(SourceObject)->GetSkeletalMeshAsset())->GetOutermost()->bIsCookedForEditor &&
 		Requirements.AreSatisfiedBy(USkeletalMeshComponentReadOnlyToolTarget::StaticClass());
 }
 
@@ -220,6 +222,7 @@ bool USkeletalMeshComponentToolTargetFactory::CanBuildTarget(UObject* SourceObje
 
 	return Cast<USkeletalMeshComponent>(SourceObject) && Cast<USkeletalMeshComponent>(SourceObject)->GetSkeletalMeshAsset() &&
 		ExactCast<USkeletalMesh>(Cast<USkeletalMeshComponent>(SourceObject)->GetSkeletalMeshAsset()) &&
+		!ExactCast<USkeletalMesh>(Cast<USkeletalMeshComponent>(SourceObject)->GetSkeletalMeshAsset())->GetOutermost()->bIsCookedForEditor &&
 		Requirements.AreSatisfiedBy(USkeletalMeshComponentToolTarget::StaticClass());
 }
 
diff --git a/Engine/Plugins/Runtime/MeshModelingToolset/Source/ModelingComponentsEditorOnly/Private/ToolTargets/SkeletalMeshToolTarget.cpp b/Engine/Plugins/Runtime/MeshModelingToolset/Source/ModelingComponentsEditorOnly/Private/ToolTargets/SkeletalMeshToolTarget.cpp
--- a/Engine/Plugins/Runtime/MeshModelingToolset/Source/ModelingComponentsEditorOnly/Private/ToolTargets/SkeletalMeshToolTarget.cpp
+++ b/Engine/Plugins/Runtime/MeshModelingToolset/Source/ModelingComponentsEditorOnly/Private/ToolTargets/SkeletalMeshToolTarget.cpp
@@ -9,6 +9,7 @@
 #include "Widgets/Notifications/SNotificationList.h"
 #include "Framework/Notifications/NotificationManager.h"
 #include "SkeletalMeshAttributes.h"
+#include "UObject/Package.h"
 
 #include UE_INLINE_GENERATED_CPP_BY_NAME(SkeletalMeshToolTarget)
 
@@ -258,7 +259,10 @@ bool USkeletalMeshReadOnlyToolTargetFactory::CanBuildTarget(UObject* SourceObjec
 	// just add another factory that allows that class specifically (but make sure that
 	// GetMeshDescription and such work properly)
 
-	return ExactCast<USkeletalMesh>(SourceObject) 
+	const USkeletalMesh* SkeletalMesh =  ExactCast<USkeletalMesh>(SourceObject);
+
+	return SkeletalMesh
+		&& !SkeletalMesh->GetOutermost()->bIsCookedForEditor
 		&& Requirements.AreSatisfiedBy(USkeletalMeshReadOnlyToolTarget::StaticClass());
 }
 
@@ -279,8 +283,11 @@ bool USkeletalMeshToolTargetFactory::CanBuildTarget(UObject* SourceObject, const
 	// If you want to make the tool target work with some subclass of USkeletalMesh,
 	// just add another factory that allows that class specifically (but make sure that
 	// GetMeshDescription and such work properly)
+	
+	const USkeletalMesh* SkeletalMesh =  ExactCast<USkeletalMesh>(SourceObject);
 
-	return ExactCast<USkeletalMesh>(SourceObject) 
+	return SkeletalMesh
+		&& !SkeletalMesh->GetOutermost()->bIsCookedForEditor
 		&& Requirements.AreSatisfiedBy(USkeletalMeshToolTarget::StaticClass());
 }
 
diff --git a/Engine/Plugins/Runtime/MeshModelingToolset/Source/ModelingComponentsEditorOnly/Private/ToolTargets/StaticMeshComponentToolTarget.cpp b/Engine/Plugins/Runtime/MeshModelingToolset/Source/ModelingComponentsEditorOnly/Private/ToolTargets/StaticMeshComponentToolTarget.cpp
--- a/Engine/Plugins/Runtime/MeshModelingToolset/Source/ModelingComponentsEditorOnly/Private/ToolTargets/StaticMeshComponentToolTarget.cpp
+++ b/Engine/Plugins/Runtime/MeshModelingToolset/Source/ModelingComponentsEditorOnly/Private/ToolTargets/StaticMeshComponentToolTarget.cpp
@@ -12,6 +12,7 @@
 #include "ToolTargets/StaticMeshToolTarget.h"
 #include "AssetUtils/MeshDescriptionUtil.h"
 #include "StaticMeshAttributes.h"
+#include "UObject/Package.h"
 
 #include UE_INLINE_GENERATED_CPP_BY_NAME(StaticMeshComponentToolTarget)
 
@@ -282,7 +283,7 @@ bool UStaticMeshComponentToolTargetFactory::CanBuildTarget(UObject* SourceObject
 {
 	const UStaticMeshComponent* Component = GetValid(Cast<UStaticMeshComponent>(SourceObject));
 	return Component && !Component->IsUnreachable() && Component->IsValidLowLevel() && Component->GetStaticMesh()
-		&& (Component->GetStaticMesh()->GetNumSourceModels() > 0)
+		&& !Component->GetStaticMesh()->GetOutermost()->bIsCookedForEditor
 		&& Requirements.AreSatisfiedBy(UStaticMeshComponentToolTarget::StaticClass());
 }
 
diff --git a/Engine/Plugins/Runtime/MeshModelingToolset/Source/ModelingComponentsEditorOnly/Private/ToolTargets/StaticMeshToolTarget.cpp b/Engine/Plugins/Runtime/MeshModelingToolset/Source/ModelingComponentsEditorOnly/Private/ToolTargets/StaticMeshToolTarget.cpp
--- a/Engine/Plugins/Runtime/MeshModelingToolset/Source/ModelingComponentsEditorOnly/Private/ToolTargets/StaticMeshToolTarget.cpp
+++ b/Engine/Plugins/Runtime/MeshModelingToolset/Source/ModelingComponentsEditorOnly/Private/ToolTargets/StaticMeshToolTarget.cpp
@@ -12,6 +12,7 @@
 #include "Widgets/Notifications/SNotificationList.h"
 #include "Framework/Notifications/NotificationManager.h"
 #include "StaticMeshAttributes.h"
+#include "UObject/Package.h"
 
 #include UE_INLINE_GENERATED_CPP_BY_NAME(StaticMeshToolTarget)
 
@@ -341,7 +342,7 @@ bool UStaticMeshToolTargetFactory::CanBuildTarget(UObject* SourceObject, const F
 {
 	const UStaticMesh* StaticMesh = GetValid(Cast<UStaticMesh>(SourceObject));
 	return StaticMesh && !StaticMesh->IsUnreachable() && StaticMesh->IsValidLowLevel()
-		&& (StaticMesh->GetNumSourceModels() > 0)
+		&& !StaticMesh->GetOutermost()->bIsCookedForEditor
 		&& Requirements.AreSatisfiedBy(UStaticMeshToolTarget::StaticClass());
 }
 
