Made FMeshBakerDynamicMeshSampler explicitly non-copyable to prevent a bunch of deprecation warnings about FMeshBakerDynamicMeshSampler::DetailNormalMap from their implicit copy constructors.
Made FSceneCapturePhotoSetSampler explicitly non-copyable to prevent its special member functions being DLL-exported.

diff --git a/Engine/Plugins/Runtime/GeometryProcessing/Source/DynamicMesh/Public/Sampling/MeshBakerCommon.h b/Engine/Plugins/Runtime/GeometryProcessing/Source/DynamicMesh/Public/Sampling/MeshBakerCommon.h
--- a/Engine/Plugins/Runtime/GeometryProcessing/Source/DynamicMesh/Public/Sampling/MeshBakerCommon.h
+++ b/Engine/Plugins/Runtime/GeometryProcessing/Source/DynamicMesh/Public/Sampling/MeshBakerCommon.h
@@ -366,6 +366,8 @@ public:
 class FMeshBakerDynamicMeshSampler : public IMeshBakerDetailSampler
 {
 public:
+	UE_NONCOPYABLE(FMeshBakerDynamicMeshSampler)
+
 	FMeshBakerDynamicMeshSampler(const FDynamicMesh3* Mesh, const FDynamicMeshAABBTree3* Spatial, const FMeshTangentsd* Tangents = nullptr)
 		: DetailMesh(Mesh), DetailSpatial(Spatial), DetailTangents(Tangents)
 	{
diff --git a/Engine/Plugins/Runtime/MeshModelingToolset/Source/ModelingComponents/Public/Baking/RenderCaptureFunctions.h b/Engine/Plugins/Runtime/MeshModelingToolset/Source/ModelingComponents/Public/Baking/RenderCaptureFunctions.h
--- a/Engine/Plugins/Runtime/MeshModelingToolset/Source/ModelingComponents/Public/Baking/RenderCaptureFunctions.h
+++ b/Engine/Plugins/Runtime/MeshModelingToolset/Source/ModelingComponents/Public/Baking/RenderCaptureFunctions.h
@@ -75,6 +75,8 @@ public:
 		const FDynamicMeshAABBTree3* Spatial,
 		const FMeshTangentsd* Tangents); 
 
+	UE_NONCOPYABLE(FSceneCapturePhotoSetSampler);
+
 	virtual bool SupportsCustomCorrespondence() const override;
 
 	// Warning: Expects that Sample.BaseSample.SurfacePoint and Sample.BaseNormal are set when the function is called
