Reorder the hires mesh to fit the LOD 0 in all workflow

#rnx

diff --git a/Engine/Plugins/Experimental/GeometryScripting/Source/GeometryScriptingCore/Private/MeshAssetFunctions.cpp b/Engine/Plugins/Experimental/GeometryScripting/Source/GeometryScriptingCore/Private/MeshAssetFunctions.cpp
--- a/Engine/Plugins/Experimental/GeometryScripting/Source/GeometryScriptingCore/Private/MeshAssetFunctions.cpp
+++ b/Engine/Plugins/Experimental/GeometryScripting/Source/GeometryScriptingCore/Private/MeshAssetFunctions.cpp
@@ -293,6 +293,13 @@ UDynamicMesh*  UGeometryScriptLibrary_StaticMeshFunctions::CopyMeshToStaticMesh(
 			});
 
 
+		if (FMeshDescription* BaseLodMeshDescription = ToStaticMeshAsset->GetMeshDescription(0))
+		{
+			FString MaterialNameConflictMsg = TEXT("[Asset ") + ToStaticMeshAsset->GetPathName() + TEXT("] Nanite hi - res CopyMeshToStaticMesh have some material name that differ from the LOD 0 material name.Your nanite hi - res should use the same material names the LOD 0 use to ensure we can remap the section in the same order.");
+			FString MaterialCountConflictMsg = TEXT("[Asset ") + ToStaticMeshAsset->GetPathName() + TEXT("] Nanite hi-res CopyMeshToStaticMesh use more material then LOD 0. Your nanite hi-res should have less or equal number of material. Any extra material will be remap to the first material use by LOD 0.");
+			FStaticMeshOperations::ReorderMeshDescriptionPolygonGroups(*BaseLodMeshDescription, *NewHiResMD, MaterialNameConflictMsg, MaterialCountConflictMsg);
+		}
+
 		ToStaticMeshAsset->CommitHiResMeshDescription();
 	}
 	else
diff --git a/Engine/Plugins/Experimental/MeshLODToolset/Source/MeshLODToolset/Private/Graphs/GenerateStaticMeshLODProcess.cpp b/Engine/Plugins/Experimental/MeshLODToolset/Source/MeshLODToolset/Private/Graphs/GenerateStaticMeshLODProcess.cpp
--- a/Engine/Plugins/Experimental/MeshLODToolset/Source/MeshLODToolset/Private/Graphs/GenerateStaticMeshLODProcess.cpp
+++ b/Engine/Plugins/Experimental/MeshLODToolset/Source/MeshLODToolset/Private/Graphs/GenerateStaticMeshLODProcess.cpp
@@ -40,6 +40,7 @@
 #include "PhysicsEngine/BodySetup.h"
 #include "RenderingThread.h"
 #include "StaticMeshAttributes.h"
+#include "StaticMeshOperations.h"
 
 #include UE_INLINE_GENERATED_CPP_BY_NAME(GenerateStaticMeshLODProcess)
 
@@ -1704,6 +1705,13 @@ void UGenerateStaticMeshLODProcess::UpdateSourceStaticMeshAsset(bool bSetNewHDSo
 	{
 		SourceStaticMesh->ModifyHiResMeshDescription();
 
+		if (FMeshDescription* BaseLodMeshDescription = SourceStaticMesh->GetMeshDescription(0))
+		{
+			FString MaterialNameConflictMsg = TEXT("[Asset ") + SourceStaticMesh->GetPathName() + TEXT("] Nanite hi - res UpdateSourceStaticMeshAsset have some material name that differ from the LOD 0 material name.Your nanite hi - res should use the same material names the LOD 0 use to ensure we can remap the section in the same order.");
+			FString MaterialCountConflictMsg = TEXT("[Asset ") + SourceStaticMesh->GetPathName() + TEXT("] Nanite hi-res UpdateSourceStaticMeshAsset use more material then LOD 0. Your nanite hi-res should have less or equal number of material. Any extra material will be remap to the first material use by LOD 0.");
+			FStaticMeshOperations::ReorderMeshDescriptionPolygonGroups(*BaseLodMeshDescription, *SourceMeshDescription, MaterialNameConflictMsg, MaterialCountConflictMsg);
+		}
+
 		FMeshDescription* NewSourceMD = SourceStaticMesh->CreateHiResMeshDescription();
 		*NewSourceMD = *SourceMeshDescription;		// todo: can MoveTemp here, we don't need this memory anymore??
 
diff --git a/Engine/Plugins/Runtime/MeshModelingToolset/Source/ModelingComponentsEditorOnly/Private/ToolTargets/StaticMeshToolTarget.cpp b/Engine/Plugins/Runtime/MeshModelingToolset/Source/ModelingComponentsEditorOnly/Private/ToolTargets/StaticMeshToolTarget.cpp
--- a/Engine/Plugins/Runtime/MeshModelingToolset/Source/ModelingComponentsEditorOnly/Private/ToolTargets/StaticMeshToolTarget.cpp
+++ b/Engine/Plugins/Runtime/MeshModelingToolset/Source/ModelingComponentsEditorOnly/Private/ToolTargets/StaticMeshToolTarget.cpp
@@ -13,6 +13,7 @@
 #include "Framework/Notifications/NotificationManager.h"
 #include "StaticMeshAttributes.h"
 #include "UObject/Package.h"
+#include "StaticMeshOperations.h"
 
 #include UE_INLINE_GENERATED_CPP_BY_NAME(StaticMeshToolTarget)
 
@@ -298,6 +299,16 @@ void UStaticMeshToolTarget::CommitMeshDescription(UStaticMesh* StaticMeshIn, con
 		StaticMeshIn->ModifyMeshDescription((int32)EditingLODIn);
 	}
 
+	if (EditingLODIn == EMeshLODIdentifier::HiResSource || EditingLODIn == EMeshLODIdentifier::LOD0)
+	{
+		const FMeshDescription* BaseLODMeshDescription = (EditingLODIn == EMeshLODIdentifier::LOD0) ? UpdateMeshDescription : StaticMeshIn->GetMeshDescription(0);
+		FMeshDescription* HiResMeshDescription = (EditingLODIn == EMeshLODIdentifier::HiResSource) ? UpdateMeshDescription : StaticMeshIn->GetHiResMeshDescription();
+
+		FString MaterialNameConflictMsg = TEXT("[Asset ") + StaticMeshIn->GetPathName() + TEXT("] Nanite hi - res StaticMeshToolTarget have some material name that differ from the LOD 0 material name.Your nanite hi - res should use the same material names the LOD 0 use to ensure we can remap the section in the same order.");
+		FString MaterialCountConflictMsg = TEXT("[Asset ") + StaticMeshIn->GetPathName() + TEXT("] Nanite hi-res StaticMeshToolTarget use more material then LOD 0. Your nanite hi-res should have less or equal number of material. Any extra material will be remap to the first material use by LOD 0.");
+		FStaticMeshOperations::ReorderMeshDescriptionPolygonGroups(*BaseLODMeshDescription, *HiResMeshDescription, MaterialNameConflictMsg, MaterialCountConflictMsg);
+	}
+
 	FCommitterParams CommitterParams;
 	CommitterParams.MeshDescriptionOut = UpdateMeshDescription;
 
