ModelingToolset: rewrite material setup in UE::AssetUtils::CreateStaticMeshAsset to not call UStaticMesh::SetMaterial(), which forces a PostEditChange() on each call. Instead build and set a list of FStaticMaterial, which shouldn't result in spurious ddc builds.

diff --git a/Engine/Plugins/Runtime/MeshModelingToolset/Source/ModelingComponentsEditorOnly/Private/AssetUtils/CreateStaticMeshUtil.cpp b/Engine/Plugins/Runtime/MeshModelingToolset/Source/ModelingComponentsEditorOnly/Private/AssetUtils/CreateStaticMeshUtil.cpp
--- a/Engine/Plugins/Runtime/MeshModelingToolset/Source/ModelingComponentsEditorOnly/Private/AssetUtils/CreateStaticMeshUtil.cpp
+++ b/Engine/Plugins/Runtime/MeshModelingToolset/Source/ModelingComponentsEditorOnly/Private/AssetUtils/CreateStaticMeshUtil.cpp
@@ -12,6 +12,11 @@
 #include "DynamicMeshToMeshDescription.h"
 #include "UObject/Package.h"
 
+#include "MaterialDomain.h"
+#include "Materials/Material.h"
+#include "Materials/MaterialInterface.h"
+#include "AssetUtils/StaticMeshMaterialUtil.h"
+
 using namespace UE::AssetUtils;
 
 using namespace UE::Geometry;
@@ -71,21 +76,31 @@ UE::AssetUtils::ECreateStaticMeshResult UE::AssetUtils::CreateStaticMeshAsset(
 		NewStaticMesh->GetBodySetup()->CollisionTraceFlag = Options.CollisionType;
 	}
 
-	// add a material slot. Must always have one material slot.
+	// Must always have one material slot.
 	int32 UseNumMaterialSlots = FMath::Max(1, Options.NumMaterialSlots);
-	for (int MatIdx = 0; MatIdx < UseNumMaterialSlots; MatIdx++)
-	{
-		NewStaticMesh->GetStaticMaterials().Add(FStaticMaterial());
-	}
-
-	// set materials if the count matches
-	if (Options.AssetMaterials.Num() == UseNumMaterialSlots)
+	// initialize list of Static Materials. Set Default Surface Material in each slot if no
+	// other Material was provided. Derive slot names from material names 
+	TArray<FStaticMaterial> StaticMaterials;
+	for (int32 MatIdx = 0; MatIdx < UseNumMaterialSlots; ++MatIdx)
 	{
-		for (int MatIdx = 0; MatIdx < UseNumMaterialSlots; MatIdx++)
+		FStaticMaterial NewMaterial;
+		// fallback to default material if no material is found
+		if (MatIdx < Options.AssetMaterials.Num() && Options.AssetMaterials[MatIdx] != nullptr)
+		{
+			NewMaterial.MaterialInterface = Options.AssetMaterials[MatIdx];
+			NewMaterial.MaterialSlotName = UE::AssetUtils::GenerateNewMaterialSlotName(StaticMaterials, NewMaterial.MaterialInterface, MatIdx);
+		}
+		else
 		{
-			NewStaticMesh->SetMaterial(MatIdx, Options.AssetMaterials[MatIdx]);
+			NewMaterial.MaterialInterface = UMaterial::GetDefaultMaterial(MD_Surface);
+			NewMaterial.MaterialSlotName = UE::AssetUtils::GenerateNewMaterialSlotName(StaticMaterials, nullptr, MatIdx);
 		}
+		NewMaterial.ImportedMaterialSlotName = NewMaterial.MaterialSlotName;
+		NewMaterial.UVChannelData = FMeshUVChannelInfo(1.f);		// this avoids an ensure in  UStaticMesh::GetUVChannelData
+		StaticMaterials.Add(NewMaterial);
 	}
+	NewStaticMesh->SetStaticMaterials(StaticMaterials);
+
 
 	// determine maximum number of sections across all mesh LODs
 	int32 MaxNumSections = 0;
