MeshModelingTools: Fix incorrect UVs in Push/Pull activity in PolyEd.

diff --git a/Engine/Plugins/Runtime/MeshModelingToolset/Source/MeshModelingTools/Private/ToolActivities/PolyEditExtrudeActivity.cpp b/Engine/Plugins/Runtime/MeshModelingToolset/Source/MeshModelingTools/Private/ToolActivities/PolyEditExtrudeActivity.cpp
--- a/Engine/Plugins/Runtime/MeshModelingToolset/Source/MeshModelingTools/Private/ToolActivities/PolyEditExtrudeActivity.cpp
+++ b/Engine/Plugins/Runtime/MeshModelingToolset/Source/MeshModelingTools/Private/ToolActivities/PolyEditExtrudeActivity.cpp
@@ -113,7 +113,11 @@ TUniquePtr<FDynamicMeshOperator> UPolyEditExtrudeActivity::MakeNewOperator()
 	
 	Op->ExtrudeDistance = GetDistanceMode(*this, PropertySetToUse) == EPolyEditExtrudeDistanceMode::Fixed ?
 		GetFixedDistance(*this, PropertySetToUse) : ExtrudeHeightMechanic->CurrentHeight;
-	Op->UVScaleFactor = UVScaleFactor;
+	// TODO: The extruder has been changed such that the boolean and move-and-stitch paths treat UV scaling differently,
+	// with the boolean path using a newer path that does its own normalization of UVs, while the move-and-stitch
+	// still uses the legacy version. They both probably should just use the newer version, but for now we'll pass
+	// the path-appropriate value here...
+	Op->UVScaleFactor = Op->ExtrudeMode == FExtrudeOp::EExtrudeMode::MoveAndStitch ? UVScaleFactor : 1.0;
 
 	FTransform3d WorldTransform(ActivityContext->Preview->PreviewMesh->GetTransform());
 	Op->SetResultTransform(WorldTransform);
