ModelingComponents: Restore check for SourceModel existence during CanBuild in static mesh tool targets.

diff --git a/Engine/Plugins/Runtime/MeshModelingToolset/Source/ModelingComponentsEditorOnly/Private/ToolTargets/StaticMeshComponentToolTarget.cpp b/Engine/Plugins/Runtime/MeshModelingToolset/Source/ModelingComponentsEditorOnly/Private/ToolTargets/StaticMeshComponentToolTarget.cpp
--- a/Engine/Plugins/Runtime/MeshModelingToolset/Source/ModelingComponentsEditorOnly/Private/ToolTargets/StaticMeshComponentToolTarget.cpp
+++ b/Engine/Plugins/Runtime/MeshModelingToolset/Source/ModelingComponentsEditorOnly/Private/ToolTargets/StaticMeshComponentToolTarget.cpp
@@ -284,6 +284,7 @@ bool UStaticMeshComponentToolTargetFactory::CanBuildTarget(UObject* SourceObject
 	const UStaticMeshComponent* Component = GetValid(Cast<UStaticMeshComponent>(SourceObject));
 	return Component && !Component->IsUnreachable() && Component->IsValidLowLevel() && Component->GetStaticMesh()
 		&& !Component->GetStaticMesh()->GetOutermost()->bIsCookedForEditor
+		&& (Component->GetStaticMesh()->GetNumSourceModels() > 0)
 		&& Requirements.AreSatisfiedBy(UStaticMeshComponentToolTarget::StaticClass());
 }
 
diff --git a/Engine/Plugins/Runtime/MeshModelingToolset/Source/ModelingComponentsEditorOnly/Private/ToolTargets/StaticMeshToolTarget.cpp b/Engine/Plugins/Runtime/MeshModelingToolset/Source/ModelingComponentsEditorOnly/Private/ToolTargets/StaticMeshToolTarget.cpp
--- a/Engine/Plugins/Runtime/MeshModelingToolset/Source/ModelingComponentsEditorOnly/Private/ToolTargets/StaticMeshToolTarget.cpp
+++ b/Engine/Plugins/Runtime/MeshModelingToolset/Source/ModelingComponentsEditorOnly/Private/ToolTargets/StaticMeshToolTarget.cpp
@@ -343,6 +343,7 @@ bool UStaticMeshToolTargetFactory::CanBuildTarget(UObject* SourceObject, const F
 	const UStaticMesh* StaticMesh = GetValid(Cast<UStaticMesh>(SourceObject));
 	return StaticMesh && !StaticMesh->IsUnreachable() && StaticMesh->IsValidLowLevel()
 		&& !StaticMesh->GetOutermost()->bIsCookedForEditor
+		&& (StaticMesh->GetNumSourceModels() > 0)
 		&& Requirements.AreSatisfiedBy(UStaticMeshToolTarget::StaticClass());
 }
 
