-Added support for skin weight and bone attributes when converting skinned meshes to dynamic meshes.
-Removed extra spaces.

#rnx

diff --git a/Engine/Plugins/Runtime/MeshModelingToolset/Source/ModelingComponents/ModelingComponents.Build.cs b/Engine/Plugins/Runtime/MeshModelingToolset/Source/ModelingComponents/ModelingComponents.Build.cs
--- a/Engine/Plugins/Runtime/MeshModelingToolset/Source/ModelingComponents/ModelingComponents.Build.cs
+++ b/Engine/Plugins/Runtime/MeshModelingToolset/Source/ModelingComponents/ModelingComponents.Build.cs
@@ -33,6 +33,7 @@ public class ModelingComponents : ModuleRules
 				"RenderCore",
 				"RHI",
 				"ImageWriteQueue",
+				"SkeletalMeshDescription", // FSkeletalMeshAttributes::DefaultSkinWeightProfileName
 				// ... add private dependencies that you statically link with here ...
 			}
 		);
diff --git a/Engine/Plugins/Runtime/MeshModelingToolset/Source/ModelingComponents/Private/ConversionUtils/SkinnedMeshToDynamicMesh.cpp b/Engine/Plugins/Runtime/MeshModelingToolset/Source/ModelingComponents/Private/ConversionUtils/SkinnedMeshToDynamicMesh.cpp
--- a/Engine/Plugins/Runtime/MeshModelingToolset/Source/ModelingComponents/Private/ConversionUtils/SkinnedMeshToDynamicMesh.cpp
+++ b/Engine/Plugins/Runtime/MeshModelingToolset/Source/ModelingComponents/Private/ConversionUtils/SkinnedMeshToDynamicMesh.cpp
@@ -12,7 +12,8 @@
 #include "SkeletalRenderPublic.h"
 #include "ToDynamicMesh.h"
 #include "Rendering/SkeletalMeshRenderData.h"
-
+#include "BoneWeights.h"
+#include "SkeletalMeshAttributes.h"
 
 namespace UE 
 {
@@ -48,7 +49,7 @@ public:
 
 		// This step does the actual skinning (and isn't as const as the api suggests..)
 		Component->GetCPUSkinnedVertices(SkinnedVertices, RequestedLOD);
-		const FSkeletalMeshLODRenderData& LODData  = SkeletalMeshRenderData.LODRenderData[RequestedLOD];
+		const FSkeletalMeshLODRenderData& LODData = SkeletalMeshRenderData.LODRenderData[RequestedLOD];
 		
 
 		const int32 NumSections = LODData.RenderSections.Num();
@@ -78,7 +79,7 @@ public:
 		
 		// construct a list of all valid VertIDs for this mesh.
 		VertIDs.Reserve(NumVerts);
-		for (const  FSkelMeshRenderSection& Section : LODData.RenderSections)
+		for (const FSkelMeshRenderSection& Section : LODData.RenderSections)
 		{
 			if (SkipSection(Section))
 			{
@@ -87,18 +88,52 @@ public:
 
 			const int32 BaseVertexIndex = static_cast<int32>( Section.BaseVertexIndex);
 			const int32 NumSectionVtx = static_cast<int32>( Section.NumVertices );
-			for (int32 VtxIndex = BaseVertexIndex;  VtxIndex < NumSectionVtx + BaseVertexIndex;  ++VtxIndex)
+			for (int32 VtxIndex = BaseVertexIndex; VtxIndex < NumSectionVtx + BaseVertexIndex; ++VtxIndex)
 			{
 				VertIDs.Add(VtxIndex);
 			}
 		}
 
+
+		// generate vertex weights and remap the indices
+		const FSkinWeightVertexBuffer* SkinWeightBuffer = Component->GetSkinWeightBuffer(LOD);
+		const bool bSkinWeightsValid = SkinWeightBuffer && SkinWeightBuffer->GetDataVertexBuffer() && SkinWeightBuffer->GetDataVertexBuffer()->IsInitialized() && SkinWeightBuffer->GetDataVertexBuffer()->GetNumVertices() > 0;
+
+		if (bSkinWeightsValid)
+		{
+			SkinWeights.Reserve(NumVerts);
+			for (const FSkelMeshRenderSection& Section : LODData.RenderSections)
+			{
+				if (SkipSection(Section))
+				{
+					continue;
+				}
+
+				const int32 BaseVertexIndex = static_cast<int32>(Section.BaseVertexIndex);
+				const int32 NumSectionVtx = static_cast<int32>(Section.NumVertices);
+				for (int32 VtxIndex = BaseVertexIndex; VtxIndex < NumSectionVtx + BaseVertexIndex; ++VtxIndex)
+				{
+					FSkinWeightInfo SrcWeights = SkinWeightBuffer->GetVertexSkinWeights(VtxIndex);
+
+					for (int32 Idx = 0; Idx < MAX_TOTAL_INFLUENCES; ++Idx)
+					{
+						const int32 SectionBoneIdx = static_cast<int32>(SrcWeights.InfluenceBones[Idx]);
+						if (ensure(SectionBoneIdx < Section.BoneMap.Num()))
+						{
+							SrcWeights.InfluenceBones[Idx] = Section.BoneMap[SectionBoneIdx];
+						}
+					}
+					SkinWeights.Add(SrcWeights);
+				}
+			}
+		}
+
 		// construct list of all valid TriIDs for this mesh.  
 		TriIDs.Reserve(NumTris);
 		TriIDtoSectionID.Init(-1, NumTris);
 		for (int32 s = 0; s < NumSections; ++s)
 		{
-			const  FSkelMeshRenderSection& Section = LODData.RenderSections[s];
+			const FSkelMeshRenderSection& Section = LODData.RenderSections[s];
 			if (SkipSection(Section))
 			{
 				continue; 
@@ -126,7 +161,7 @@ public:
 
 	int32 NumUVLayers() const
 	{
-		const FSkeletalMeshLODRenderData& LODData  = SkeletalMeshRenderData.LODRenderData[LOD];
+		const FSkeletalMeshLODRenderData& LODData = SkeletalMeshRenderData.LODRenderData[LOD];
 		return static_cast<int32>( LODData.StaticVertexBuffers.StaticMeshVertexBuffer.GetNumTexCoords() );
 	}
 
@@ -170,8 +205,8 @@ public:
 	
 	bool HasColors() const
 	{
-		const FSkeletalMeshLODRenderData& LODData  = SkeletalMeshRenderData.LODRenderData[LOD];
-		return  (LODData.StaticVertexBuffers.ColorVertexBuffer.IsInitialized() && LODData.StaticVertexBuffers.ColorVertexBuffer.GetNumVertices() > 0);
+		const FSkeletalMeshLODRenderData& LODData = SkeletalMeshRenderData.LODRenderData[LOD];
+		return (LODData.StaticVertexBuffers.ColorVertexBuffer.IsInitialized() && LODData.StaticVertexBuffers.ColorVertexBuffer.GetNumVertices() > 0);
 	}
 
 	//-- Access to per-wedge attributes --//
@@ -217,9 +252,9 @@ public:
 		if (HasColors())
 		{ 
 			const int32 VID = GetVertID(WID);
-			const FSkeletalMeshLODRenderData& LODData  = SkeletalMeshRenderData.LODRenderData[LOD];
+			const FSkeletalMeshLODRenderData& LODData = SkeletalMeshRenderData.LODRenderData[LOD];
 			const FColor& Color = LODData.StaticVertexBuffers.ColorVertexBuffer.VertexColor(VID);
-			LinearColor =  Color.ReinterpretAsLinear();
+			LinearColor = Color.ReinterpretAsLinear();
 		}
 		return FVector4f(LinearColor.R, LinearColor.G, LinearColor.B, LinearColor.A);
 	}
@@ -230,8 +265,8 @@ public:
 		int32 SectionID = TriIDtoSectionID[TriID];
 		check(SectionID > -1);
 
-		const FSkeletalMeshLODRenderData& LODData  = SkeletalMeshRenderData.LODRenderData[LOD];
-		const  FSkelMeshRenderSection& Section = LODData.RenderSections[SectionID];
+		const FSkeletalMeshLODRenderData& LODData = SkeletalMeshRenderData.LODRenderData[LOD];
+		const FSkelMeshRenderSection& Section = LODData.RenderSections[SectionID];
 		
 		int32 MaterialIndex = Section.MaterialIndex;
 		// use the remapping of material indices if there is a valid value
@@ -243,7 +278,72 @@ public:
 		return MaterialIndex;
 	}
 
+	int32 NumSkinWeightAttributes() const 
+	{ 
+		const bool bSkinWeightsValid = SkinWeights.Num() == SkinnedVertices.Num();
+
+		return static_cast<int32>(bSkinWeightsValid);
+	}
+
+	UE::AnimationCore::FBoneWeights GetVertexSkinWeight(int32 SkinWeightAttributeIndex, VertIDType VertexID) const
+	{
+		checkfSlow(SkinWeightAttributeIndex == 0, TEXT("USkinnedMeshComponent can have at most one skin weight attribute"));
+		checkfSlow(VertexID < SkinWeights.Num() && VertexID >= 0, TEXT("Invalid vertex id"));
+
+		const FSkinWeightInfo& WeightInfo = SkinWeights[VertexID];
 
+		const UE::AnimationCore::FBoneWeights Weights = UE::AnimationCore::FBoneWeights::Create(WeightInfo.InfluenceBones, WeightInfo.InfluenceWeights);
+
+		return Weights;
+	}
+
+	FName GetSkinWeightAttributeName(int32 SkinWeightAttributeIndex) const
+	{
+		checkfSlow(SkinWeightAttributeIndex == 0, TEXT("USkinnedMeshComponent can have at most one skin weight attribute"));
+		
+		return FSkeletalMeshAttributes::DefaultSkinWeightProfileName;
+	};
+	
+
+	int32 GetNumBones() const 
+	{ 
+		return Asset->GetRefSkeleton().GetRawBoneNum(); 
+	}
+
+    FName GetBoneName(int32 BoneIdx) const 
+	{
+		if (ensure(BoneIdx >= 0 && BoneIdx < GetNumBones()))
+		{
+			return Asset->GetRefSkeleton().GetRawRefBoneInfo()[BoneIdx].Name;
+		}
+
+		return NAME_None;
+	}
+
+	int32 GetBoneParentIndex(int32 BoneIdx) const
+	{
+		if (ensure(BoneIdx >= 0 && BoneIdx < GetNumBones()))
+		{
+			return Asset->GetRefSkeleton().GetRawRefBoneInfo()[BoneIdx].ParentIndex;
+		}
+
+		return INDEX_NONE;
+	}
+
+	FTransform GetBonePose(int32 BoneIdx) const
+	{
+		if (ensure(BoneIdx >= 0 && BoneIdx < GetNumBones()))
+		{
+			return Asset->GetRefSkeleton().GetRawRefBonePose()[BoneIdx];
+		}
+
+		return FTransform::Identity;
+	}
+
+	FVector4f GetBoneColor(int32 BoneIdx) const
+	{
+		return FVector4f::One();
+	}
 
 	//-- null implementation of shared attributes: Skeletal mesh model doesn't use these --//
 	const TArray<int32>& GetUVIDs(int32 LayerID) const { return EmptyArray; }
@@ -290,7 +390,7 @@ private:
 
 private:
 
-	TArray<FFinalSkinVertex>  SkinnedVertices;
+	TArray<FFinalSkinVertex> SkinnedVertices;
 	USkinnedMeshComponent* Component;
 	int32                  LOD;
 	bool                   bUseDisabledSections;
@@ -307,6 +407,8 @@ private:
 
 	TArray<int32> TriIDtoSectionID;
 	TArray<int32> EmptyArray;
+
+	TArray<FSkinWeightInfo> SkinWeights;
 };
 
 
@@ -326,7 +428,7 @@ void SkinnedMeshComponentToDynamicMesh(USkinnedMeshComponent& SkinnedMeshCompone
 		return;
 	}
 
-	if (RequestedLOD < 0 || RequestedLOD  > NumLODs - 1)
+	if (RequestedLOD < 0 || RequestedLOD > NumLODs - 1)
 	{
 		return;
 	}
diff --git a/Engine/Source/Runtime/MeshConversion/Public/ToDynamicMesh.h b/Engine/Source/Runtime/MeshConversion/Public/ToDynamicMesh.h
--- a/Engine/Source/Runtime/MeshConversion/Public/ToDynamicMesh.h
+++ b/Engine/Source/Runtime/MeshConversion/Public/ToDynamicMesh.h
@@ -6,6 +6,8 @@
 #include "DynamicMesh/DynamicMesh3.h"
 #include "DynamicMesh/DynamicMeshAttributeSet.h"
 #include "DynamicMesh/DynamicMeshOverlay.h"
+#include "DynamicMesh/DynamicVertexSkinWeightsAttribute.h"
+#include "DynamicMesh/DynamicBoneAttribute.h"
 #include "VectorTypes.h"
 #include "Async/Async.h"
 
@@ -308,6 +310,18 @@ public:
 *	    int32 NumWeightMapLayers() const;
 *       float GetVertexWeight(int32 WeightMapIndex, int32 SrcVertID) const;
 *	    FName GetWeightMapName(int32 WeightMapIndex) const;
+*
+*		// skin weight attributes information
+*		int32 NumSkinWeightAttributes() const;
+*		UE::AnimationCore::FBoneWeights GetVertexSkinWeight(int32 SkinWeightAttributeIndex, VertIDType VtxID) const;
+*		FName GetSkinWeightAttributeName(int32 SkinWeightAttributeIndex) const;
+*
+*		// bone attributes information
+*		int32 GetNumBones() const;
+*		FName GetBoneName(int32 BoneIdx) const;
+*		int32 GetBoneParentIndex(int32 BoneIdx) const;
+*		FTransform GetBonePose(int32 BoneIdx) const;
+*		FVector4f GetBoneColor(int32 BoneIdx) const;
 *	};
 *
 */
@@ -367,7 +381,6 @@ protected:
 
 		// -- Create Overlays
 		const int32 NumUVLayers = MeshIn.NumUVLayers();
-		const int32 NumWeightMaps = MeshIn.NumWeightMapLayers();
 
 		MeshOut.EnableAttributes(); // by default 1-UV layer and 1-normal layer
 
@@ -514,6 +527,7 @@ protected:
 		}
 
 		//populate WeightMap attribute
+		const int32 NumWeightMaps = MeshIn.NumWeightMapLayers();
 		MeshOut.Attributes()->SetNumWeightLayers(NumWeightMaps);
 		for (int WeightMapIndex = 0; WeightMapIndex < NumWeightMaps; WeightMapIndex++)	//-V654 //-V621 (The static analyzer complains if it knows NumWeightMaps is 0 for a given SrcMeshType)
 		{
@@ -533,6 +547,65 @@ protected:
 			Pending.Add(MoveTemp(WeightMapFuture));
 		}
 
+		
+		// populate skinning weight attribute
+		const int32 NumSkinWeightAttributes = MeshIn.NumSkinWeightAttributes();
+		
+		// first, attach all the skinning weight attributes (this also allocates memory to store skinning weights for each attribute)
+		for (int AttributeIndex = 0; AttributeIndex < NumSkinWeightAttributes; ++AttributeIndex) 
+		{
+			UE::Geometry::FDynamicMeshVertexSkinWeightsAttribute* Attribute = new UE::Geometry::FDynamicMeshVertexSkinWeightsAttribute(&MeshOut);
+			const FName AttribName = MeshIn.GetSkinWeightAttributeName(AttributeIndex);
+			Attribute->SetName(AttribName);
+			MeshOut.Attributes()->AttachSkinWeightsAttribute(AttribName, Attribute);
+		}
+
+		// now set all of the skin weight data
+		for (int AttributeIndex = 0; AttributeIndex < NumSkinWeightAttributes; ++AttributeIndex)
+		{
+			auto SkinAttributeFeature = Async(EAsyncExecution::ThreadPool, [this, &MeshIn, &MeshOut, AttributeIndex]()
+			{
+				const FName AttribName = MeshIn.GetSkinWeightAttributeName(AttributeIndex);
+				
+				UE::Geometry::FDynamicMeshVertexSkinWeightsAttribute* OutAttribute = MeshOut.Attributes()->GetSkinWeightsAttribute(AttribName);
+				checkSlow(OutAttribute != nullptr);
+
+				if (OutAttribute) 
+				{
+					for (int32 VertexID : MeshOut.VertexIndicesItr())
+					{
+						SrcVertIDType SrcVertID = MyBase::ToSrcVertIDMap[VertexID];
+						UE::AnimationCore::FBoneWeights SkinWeight = MeshIn.GetVertexSkinWeight(AttributeIndex, SrcVertID);
+						OutAttribute->SetValue(VertexID, SkinWeight);
+					}
+				}
+			});
+			Pending.Add(MoveTemp(SkinAttributeFeature));
+		}
+
+		// populate bones attribute
+		const int NumBones = MeshIn.GetNumBones();
+		if (MeshIn.GetNumBones())
+		{	
+			MeshOut.Attributes()->EnableBones(NumBones);
+			FDynamicMeshBoneNameAttribute* BoneNameAttrib = MeshOut.Attributes()->GetBoneNames();
+			FDynamicMeshBoneParentIndexAttribute* BoneParentIndices = MeshOut.Attributes()->GetBoneParentIndices();
+			FDynamicMeshBonePoseAttribute* BonePoses = MeshOut.Attributes()->GetBonePoses();
+			FDynamicMeshBoneColorAttribute* BoneColors = MeshOut.Attributes()->GetBoneColors();
+
+			auto BoneAttributeFeature = Async(EAsyncExecution::ThreadPool, [this, &MeshIn, &MeshOut, BoneNameAttrib, BoneParentIndices, BonePoses, BoneColors, NumBones]()
+			{
+				for (int32 BoneIdx = 0; BoneIdx < NumBones; ++BoneIdx)
+				{
+					BoneNameAttrib->SetValue(BoneIdx, MeshIn.GetBoneName(BoneIdx));
+					BoneParentIndices->SetValue(BoneIdx, MeshIn.GetBoneParentIndex(BoneIdx));
+					BonePoses->SetValue(BoneIdx, MeshIn.GetBonePose(BoneIdx));
+					BoneColors->SetValue(BoneIdx, MeshIn.GetBoneColor(BoneIdx));
+				}
+			});
+			Pending.Add(MoveTemp(BoneAttributeFeature));
+		}
+
 		// TODO: PolyGroup layers
 
 		// wait for all work to be done
diff --git a/Engine/Source/Runtime/MeshConversionEngineTypes/Private/SkeletalMeshLODModelToDynamicMesh.cpp b/Engine/Source/Runtime/MeshConversionEngineTypes/Private/SkeletalMeshLODModelToDynamicMesh.cpp
--- a/Engine/Source/Runtime/MeshConversionEngineTypes/Private/SkeletalMeshLODModelToDynamicMesh.cpp
+++ b/Engine/Source/Runtime/MeshConversionEngineTypes/Private/SkeletalMeshLODModelToDynamicMesh.cpp
@@ -20,7 +20,7 @@ public:
 	typedef int32 NormalIDType;
 	typedef int32 ColorIDType;
 
-	FSkeletalMeshLODModelWrapper(const  FSkeletalMeshLODModel& MeshIn, bool bUseDisabledSections) :
+	FSkeletalMeshLODModelWrapper(const FSkeletalMeshLODModel& MeshIn, bool bUseDisabledSections) :
 		bIncludeDisabledSections(bUseDisabledSections),
 		Mesh(&MeshIn)
 	{
@@ -271,6 +271,17 @@ public:
 	FVector4f GetColor(ColorIDType ID) const { check(0); return FVector4f(); }
 	bool GetColorTri(const TriIDType&, ColorIDType& ID0, ColorIDType& ID1, ColorIDType& ID2) const { ID0 = ID1 = ID2 = ColorIDType(-1); return false; }
 
+	// TODO: This file is only used in an outdated geometry blueprint library that is not being used. So adding an empty 
+	// implementation for skin weights and bones for now. See SkinnedMeshToDynamicMesh.h/cpp for reference.
+	int32 NumSkinWeightAttributes() const { return 0; }
+	UE::AnimationCore::FBoneWeights GetVertexSkinWeight(int32 SkinWeightAttributeIndex, VertIDType VertexID) const { return UE::AnimationCore::FBoneWeights(); }
+	FName GetSkinWeightAttributeName(int32 SkinWeightAttributeIndex) const { return NAME_None; }
+
+	int32 GetNumBones() const { return 0; }
+	FName GetBoneName(int32 BoneIdx) const { return NAME_None; }
+	int32 GetBoneParentIndex(int32 BoneIdx) const { return INDEX_NONE; }
+	FTransform GetBonePose(int32 BoneIdx) const { return FTransform::Identity; }
+	FVector4f GetBoneColor(int32 BoneIdx) const { return FVector4f::One(); }
 	
 	// -- additional methods, not required by the conversion interface
 	const  FSkeletalMeshLODModel& GetSrcMesh() const
@@ -304,20 +315,20 @@ private:
 
 private:
 
-	bool   bIncludeDisabledSections;
+	bool bIncludeDisabledSections;
 	TArray<VertIDType> VertIDToSectionID;
 	TArray<TriIDType> TriIDToSectionID;
 	TArray<TriIDType> TriIDs;
 	TArray<VertIDType> VertIDs;
 	
-	const  FSkeletalMeshLODModel* Mesh;
+	const FSkeletalMeshLODModel* Mesh;
 };
 
 
 
 
 
-void FSkeletalMeshLODModelToDynamicMesh::Convert(const  FSkeletalMeshLODModel* MeshIn, FDynamicMesh3& MeshOut, bool bCopyTangents)
+void FSkeletalMeshLODModelToDynamicMesh::Convert(const FSkeletalMeshLODModel* MeshIn, FDynamicMesh3& MeshOut, bool bCopyTangents)
 {
 
 	
@@ -382,7 +393,7 @@ void FSkeletalMeshLODModelToDynamicMesh::Convert(const  FSkeletalMeshLODModel* M
 
 #else
 
-void  FSkeletalMeshLODModelToDynamicMesh::Convert(const  FSkeletalMeshLODModel* MeshIn, FDynamicMesh3& MeshOut, bool bCopyTangents)
+void FSkeletalMeshLODModelToDynamicMesh::Convert(const FSkeletalMeshLODModel* MeshIn, FDynamicMesh3& MeshOut, bool bCopyTangents)
 {
 	// Conversion only supported with editor.
 	check(0);
