ModelingMode: Use actor factory when creating volumes

diff --git a/Engine/Plugins/Runtime/MeshModelingToolset/Source/ModelingComponentsEditorOnly/Private/EditorModelingObjectsCreationAPI.cpp b/Engine/Plugins/Runtime/MeshModelingToolset/Source/ModelingComponentsEditorOnly/Private/EditorModelingObjectsCreationAPI.cpp
--- a/Engine/Plugins/Runtime/MeshModelingToolset/Source/ModelingComponentsEditorOnly/Private/EditorModelingObjectsCreationAPI.cpp
+++ b/Engine/Plugins/Runtime/MeshModelingToolset/Source/ModelingComponentsEditorOnly/Private/EditorModelingObjectsCreationAPI.cpp
@@ -142,17 +142,29 @@ TArray<UMaterialInterface*> UEditorModelingObjectsCreationAPI::FilterMaterials(c
 
 FCreateMeshObjectResult UEditorModelingObjectsCreationAPI::CreateVolume(FCreateMeshObjectParams&& CreateMeshParams)
 {
-	// spawn new actor
-	FActorSpawnParameters SpawnInfo;
-	FTransform NewActorTransform = FTransform::Identity;
+	// determine volume actor type
 	UClass* VolumeClass = ABlockingVolume::StaticClass();
 	if (CreateMeshParams.TypeHintClass
-		&& Cast<AVolume>(CreateMeshParams.TypeHintClass.Get()->GetDefaultObject(false)) != nullptr )
+		&& Cast<AVolume>(CreateMeshParams.TypeHintClass.Get()->GetDefaultObject(false)) != nullptr)
 	{
 		VolumeClass = CreateMeshParams.TypeHintClass;
 	}
 
-	AVolume* NewVolumeActor = (AVolume*)CreateMeshParams.TargetWorld->SpawnActor(VolumeClass, &NewActorTransform, SpawnInfo);
+	// create new volume actor using factory
+	AVolume* const NewVolumeActor = [VolumeClass, &CreateMeshParams]() -> AVolume*
+	{
+		if (UActorFactory* const VolumeFactory = FActorFactoryAssetProxy::GetFactoryForAssetObject(VolumeClass))
+		{
+			AActor* const Actor = VolumeFactory->CreateActor(VolumeClass, CreateMeshParams.TargetWorld->GetCurrentLevel(), FTransform::Identity);
+			FActorLabelUtilities::SetActorLabelUnique(Actor, CreateMeshParams.BaseName);
+			return Cast<AVolume>(Actor);
+		}
+		return nullptr;
+	}();
+	if (!NewVolumeActor)
+	{
+		return FCreateMeshObjectResult{ECreateModelingObjectResult::Failed_ActorCreationFailed};
+	}
 
 	NewVolumeActor->BrushType = EBrushType::Brush_Add;
 	UModel* Model = NewObject<UModel>(NewVolumeActor);
