add support for querying LOD support in a mesh description provider, and use this to add an option to apply Bake Transform and Pivot tools to all LODs

diff --git a/Engine/Plugins/Experimental/MeshModelingToolsetExp/Source/MeshModelingToolsExp/Private/BakeTransformTool.cpp b/Engine/Plugins/Experimental/MeshModelingToolsetExp/Source/MeshModelingToolsExp/Private/BakeTransformTool.cpp
--- a/Engine/Plugins/Experimental/MeshModelingToolsetExp/Source/MeshModelingToolsExp/Private/BakeTransformTool.cpp
+++ b/Engine/Plugins/Experimental/MeshModelingToolsetExp/Source/MeshModelingToolsExp/Private/BakeTransformTool.cpp
@@ -125,7 +125,15 @@ void UBakeTransformTool::UpdateAssets()
 	TArray<FBox> BoundsOfScaledRotatedMesh;
 	for (int32 ComponentIdx = 0; ComponentIdx < Targets.Num(); ComponentIdx++)
 	{
-		const FMeshDescription* MeshDescription = UE::ToolTarget::GetMeshDescription(Targets[ComponentIdx]);
+		bool bTargetSupportsLODs = false;
+		TArray<EMeshLODIdentifier> LODs = UE::ToolTarget::GetMeshDescriptionLODs(Targets[ComponentIdx], bTargetSupportsLODs, !BasicProperties->bApplyToAllLODs);
+		// Make sure all the LODs after the first are loaded, but don't do anything with them
+		for (int32 LODIdx = 1; LODIdx < LODs.Num(); ++LODIdx)
+		{
+			UE::ToolTarget::GetMeshDescription(Targets[ComponentIdx], FGetMeshParameters(bTargetSupportsLODs, LODs[LODIdx]));
+		}
+		// Use the first LOD to update the bounds (or the default LOD, if LODs are not supported or if the update does not need to apply to all LODs)
+		const FMeshDescription* MeshDescription = UE::ToolTarget::GetMeshDescription(Targets[ComponentIdx], FGetMeshParameters(bTargetSupportsLODs, LODs[0]));
 		if (MapToFirstOccurrences[ComponentIdx] < ComponentIdx)
 		{
 			BoundsOfScaledRotatedMesh.Add(BoundsOfScaledRotatedMesh[MapToFirstOccurrences[ComponentIdx]]);
@@ -296,20 +304,24 @@ void UBakeTransformTool::UpdateAssets()
 
 		if (MapToFirstOccurrences[ComponentIdx] == ComponentIdx)
 		{
-			// apply edit
-			FMeshDescription SourceMesh(UE::ToolTarget::GetMeshDescriptionCopy(Targets[ComponentIdx]));
-			FMeshDescriptionEditableTriangleMeshAdapter EditableMeshDescAdapter(&SourceMesh);
+			bool bTargetSupportsLODs = false;
+			TArray<EMeshLODIdentifier> LODs = UE::ToolTarget::GetMeshDescriptionLODs(Targets[ComponentIdx], bTargetSupportsLODs, !BasicProperties->bApplyToAllLODs);
+			for (EMeshLODIdentifier LOD : LODs)
+			{
+				FMeshDescription SourceMesh(UE::ToolTarget::GetMeshDescriptionCopy(Targets[ComponentIdx], FGetMeshParameters(bTargetSupportsLODs, LOD)));
+				FMeshDescriptionEditableTriangleMeshAdapter EditableMeshDescAdapter(&SourceMesh);
 
-			MeshAdapterTransforms::ApplyTransform(EditableMeshDescAdapter, ToBakePart);
+				MeshAdapterTransforms::ApplyTransform(EditableMeshDescAdapter, ToBakePart);
 
-			FVector3d BakeScaleVec = ToBakePart.GetScale();
-			if (BakeScaleVec.X * BakeScaleVec.Y * BakeScaleVec.Z < 0)
-			{
-				SourceMesh.ReverseAllPolygonFacing();
-			}
+				FVector3d BakeScaleVec = ToBakePart.GetScale();
+				if (BakeScaleVec.X * BakeScaleVec.Y * BakeScaleVec.Z < 0)
+				{
+					SourceMesh.ReverseAllPolygonFacing();
+				}
 
-			// todo: support vertex-only update
-			UE::ToolTarget::CommitMeshDescriptionUpdate(Target, &SourceMesh);
+				// todo: support vertex-only update
+				UE::ToolTarget::CommitMeshDescriptionUpdate(Target, &SourceMesh, nullptr /*no material set changes*/, FCommitMeshParameters(bTargetSupportsLODs, LOD));
+			}
 
 			if (!bNeedSeparateTransactionForSimpleCollision)
 			{
diff --git a/Engine/Plugins/Experimental/MeshModelingToolsetExp/Source/MeshModelingToolsExp/Private/EditPivotTool.cpp b/Engine/Plugins/Experimental/MeshModelingToolsetExp/Source/MeshModelingToolsExp/Private/EditPivotTool.cpp
--- a/Engine/Plugins/Experimental/MeshModelingToolsetExp/Source/MeshModelingToolsExp/Private/EditPivotTool.cpp
+++ b/Engine/Plugins/Experimental/MeshModelingToolsetExp/Source/MeshModelingToolsExp/Private/EditPivotTool.cpp
@@ -407,7 +407,12 @@ void UEditPivotTool::UpdateAssets(const FFrame3d& NewPivotWorldFrame)
 	// TODO: this may not be necessary anymore. Also may not be the most efficient
 	for (int32 ComponentIdx = 0; ComponentIdx < Targets.Num(); ComponentIdx++)
 	{
-		const FMeshDescription* MeshDescription = UE::ToolTarget::GetMeshDescription(Targets[ComponentIdx]);
+		bool bTargetSupportsLODs = false;
+		TArray<EMeshLODIdentifier> LODs = UE::ToolTarget::GetMeshDescriptionLODs(Targets[ComponentIdx], bTargetSupportsLODs, !TransformProps->bApplyToAllLODs);
+		for (EMeshLODIdentifier LOD : LODs)
+		{
+			UE::ToolTarget::GetMeshDescription(Targets[ComponentIdx], FGetMeshParameters(bTargetSupportsLODs, LOD));
+		}
 	}
 
 	// If the targets have any convex hulls, we need to update those first in a separate transaction to work around a crash
@@ -595,16 +600,21 @@ void UEditPivotTool::UpdateAssets(const FFrame3d& NewPivotWorldFrame)
 		}
 		else if (MapToFirstOccurrences[ComponentIdx] == ComponentIdx)
 		{
-			FMeshDescription SourceMesh(UE::ToolTarget::GetMeshDescriptionCopy(Target));
-			FMeshDescriptionEditableTriangleMeshAdapter EditableMeshDescAdapter(&SourceMesh);
-			
-			for (const FTransformSRT3d& ToBake : BakeTransforms[ComponentIdx].GetTransforms())
+			bool bTargetSupportsLODs = false;
+			TArray<EMeshLODIdentifier> LODs = UE::ToolTarget::GetMeshDescriptionLODs(Targets[ComponentIdx], bTargetSupportsLODs, !TransformProps->bApplyToAllLODs);
+			for (EMeshLODIdentifier LOD : LODs)
 			{
-				MeshAdapterTransforms::ApplyTransform(EditableMeshDescAdapter, ToBake);
+				FMeshDescription SourceMesh(UE::ToolTarget::GetMeshDescriptionCopy(Target, FGetMeshParameters(bTargetSupportsLODs, LOD)));
+				FMeshDescriptionEditableTriangleMeshAdapter EditableMeshDescAdapter(&SourceMesh);
+
+				for (const FTransformSRT3d& ToBake : BakeTransforms[ComponentIdx].GetTransforms())
+				{
+					MeshAdapterTransforms::ApplyTransform(EditableMeshDescAdapter, ToBake);
+				}
+
+				// todo: support vertex-only update
+				UE::ToolTarget::CommitMeshDescriptionUpdate(Target, &SourceMesh, nullptr /*no material set changes*/, FCommitMeshParameters(bTargetSupportsLODs, LOD));
 			}
-			
-			// todo: support vertex-only update
-			UE::ToolTarget::CommitMeshDescriptionUpdate(Target, &SourceMesh);
 
 			// transform simple collision geometry
 			if (!bNeedSeparateTransactionForSimpleCollision)
diff --git a/Engine/Plugins/Experimental/MeshModelingToolsetExp/Source/MeshModelingToolsExp/Public/BakeTransformTool.h b/Engine/Plugins/Experimental/MeshModelingToolsetExp/Source/MeshModelingToolsExp/Public/BakeTransformTool.h
--- a/Engine/Plugins/Experimental/MeshModelingToolsetExp/Source/MeshModelingToolsExp/Public/BakeTransformTool.h
+++ b/Engine/Plugins/Experimental/MeshModelingToolsetExp/Source/MeshModelingToolsExp/Public/BakeTransformTool.h
@@ -52,6 +52,10 @@ class MESHMODELINGTOOLSEXP_API UBakeTransformToolProperties : public UInteractiv
 public:
 	UBakeTransformToolProperties();
 
+	/** If checked, the baked transform will be applied to all available LODs. Has no effect on selections without LODs. */
+	UPROPERTY(EditAnywhere, Category = Options)
+	bool bApplyToAllLODs = true;
+
 	/** Bake rotation */
 	UPROPERTY(EditAnywhere, Category = Options)
 	bool bBakeRotation = true;
diff --git a/Engine/Plugins/Experimental/MeshModelingToolsetExp/Source/MeshModelingToolsExp/Public/EditPivotTool.h b/Engine/Plugins/Experimental/MeshModelingToolsetExp/Source/MeshModelingToolsExp/Public/EditPivotTool.h
--- a/Engine/Plugins/Experimental/MeshModelingToolsetExp/Source/MeshModelingToolsExp/Public/EditPivotTool.h
+++ b/Engine/Plugins/Experimental/MeshModelingToolsetExp/Source/MeshModelingToolsExp/Public/EditPivotTool.h
@@ -59,6 +59,10 @@ class MESHMODELINGTOOLSEXP_API UEditPivotToolProperties : public UInteractiveToo
 	GENERATED_BODY()
 public:
 
+	/** If checked, the baked transform will be applied to all available LODs. Has no effect on selections without LODs. */
+	UPROPERTY(EditAnywhere, Category = Options)
+	bool bApplyToAllLODs = true;
+
 	/** When enabled, click-drag to reposition the Pivot */
 	UPROPERTY(EditAnywhere, Category = Options)
 	bool bEnableSnapDragging = false;
diff --git a/Engine/Plugins/Runtime/MeshModelingToolset/Source/ModelingComponents/Private/ModelingToolTargetUtil.cpp b/Engine/Plugins/Runtime/MeshModelingToolset/Source/ModelingComponents/Private/ModelingToolTargetUtil.cpp
--- a/Engine/Plugins/Runtime/MeshModelingToolset/Source/ModelingComponents/Private/ModelingToolTargetUtil.cpp
+++ b/Engine/Plugins/Runtime/MeshModelingToolset/Source/ModelingComponents/Private/ModelingToolTargetUtil.cpp
@@ -147,6 +147,20 @@ bool UE::ToolTarget::CommitMaterialSetUpdate(
 }
 
 
+TArray<EMeshLODIdentifier> UE::ToolTarget::GetMeshDescriptionLODs(UToolTarget* Target, bool& bOutTargetSupportsLODs, bool bOnlyReturnDefaultLOD, bool bExcludeAutoGeneratedLODs)
+{
+	if (!bOnlyReturnDefaultLOD)
+	{
+		if (IMeshDescriptionProvider* MeshDescriptionProvider = Cast<IMeshDescriptionProvider>(Target))
+		{
+			bOutTargetSupportsLODs = MeshDescriptionProvider->SupportsLODs();
+			return MeshDescriptionProvider->GetAvailableLODs(bExcludeAutoGeneratedLODs);
+		}
+	}
+	bOutTargetSupportsLODs = false;
+	TArray<EMeshLODIdentifier> DefaultArray{ EMeshLODIdentifier::Default };
+	return DefaultArray;
+}
 
 const FMeshDescription* UE::ToolTarget::GetMeshDescription(UToolTarget* Target, const FGetMeshParameters& GetMeshParams)
 {
@@ -231,7 +245,7 @@ FDynamicMesh3 UE::ToolTarget::GetDynamicMeshCopy(UToolTarget* Target, bool bWant
 }
 
 
-UE::ToolTarget::EDynamicMeshUpdateResult UE::ToolTarget::CommitMeshDescriptionUpdate(UToolTarget* Target, const FMeshDescription* UpdatedMesh, const FComponentMaterialSet* UpdatedMaterials)
+UE::ToolTarget::EDynamicMeshUpdateResult UE::ToolTarget::CommitMeshDescriptionUpdate(UToolTarget* Target, const FMeshDescription* UpdatedMesh, const FComponentMaterialSet* UpdatedMaterials, const FCommitMeshParameters& CommitParams)
 {
 	if (!ensure(UpdatedMesh != nullptr))
 	{
@@ -249,16 +263,16 @@ UE::ToolTarget::EDynamicMeshUpdateResult UE::ToolTarget::CommitMeshDescriptionUp
 		CommitMaterialSetUpdate(Target, *UpdatedMaterials, true);
 	}
 
-	bool bOK = MeshDescriptionCommitter->CommitMeshDescription(*UpdatedMesh);
+	bool bOK = MeshDescriptionCommitter->CommitMeshDescription(*UpdatedMesh, CommitParams);
 	return (bOK) ? EDynamicMeshUpdateResult::Ok : EDynamicMeshUpdateResult::Failed;
 }
 
 
-UE::ToolTarget::EDynamicMeshUpdateResult UE::ToolTarget::CommitMeshDescriptionUpdate(UToolTarget* Target, FMeshDescription&& UpdatedMesh)
+UE::ToolTarget::EDynamicMeshUpdateResult UE::ToolTarget::CommitMeshDescriptionUpdate(UToolTarget* Target, FMeshDescription&& UpdatedMesh, const FCommitMeshParameters& CommitParams)
 {
 	if (IMeshDescriptionCommitter* MeshDescriptionCommitter = Cast<IMeshDescriptionCommitter>(Target))
 	{
-		bool bOK = MeshDescriptionCommitter->CommitMeshDescription(MoveTemp(UpdatedMesh));
+		bool bOK = MeshDescriptionCommitter->CommitMeshDescription(MoveTemp(UpdatedMesh), CommitParams);
 		return (bOK) ? EDynamicMeshUpdateResult::Ok : EDynamicMeshUpdateResult::Failed;
 	}
 	return EDynamicMeshUpdateResult::Failed;
@@ -268,7 +282,8 @@ UE::ToolTarget::EDynamicMeshUpdateResult UE::ToolTarget::CommitMeshDescriptionUp
 UE::ToolTarget::EDynamicMeshUpdateResult UE::ToolTarget::CommitMeshDescriptionUpdateViaDynamicMesh(
 	UToolTarget* Target, 
 	const UE::Geometry::FDynamicMesh3& UpdatedMesh, 
-	bool bHaveModifiedTopology)
+	bool bHaveModifiedTopology,
+	const FCommitMeshParameters& CommitParams)
 {
 	IMeshDescriptionCommitter* MeshDescriptionCommitter = Cast<IMeshDescriptionCommitter>(Target);
 	if (!ensure(MeshDescriptionCommitter))
@@ -295,7 +310,7 @@ UE::ToolTarget::EDynamicMeshUpdateResult UE::ToolTarget::CommitMeshDescriptionUp
 	}
 	
 
-	bool bOK = MeshDescriptionCommitter->CommitMeshDescription(MoveTemp(ConvertedMesh));
+	bool bOK = MeshDescriptionCommitter->CommitMeshDescription(MoveTemp(ConvertedMesh), CommitParams);
 	return (bOK) ? EDynamicMeshUpdateResult::Ok : EDynamicMeshUpdateResult::Failed;
 }
 
diff --git a/Engine/Plugins/Runtime/MeshModelingToolset/Source/ModelingComponents/Public/ModelingToolTargetUtil.h b/Engine/Plugins/Runtime/MeshModelingToolset/Source/ModelingComponents/Public/ModelingToolTargetUtil.h
--- a/Engine/Plugins/Runtime/MeshModelingToolset/Source/ModelingComponents/Public/ModelingToolTargetUtil.h
+++ b/Engine/Plugins/Runtime/MeshModelingToolset/Source/ModelingComponents/Public/ModelingToolTargetUtil.h
@@ -90,6 +90,19 @@ MODELINGCOMPONENTS_API bool CommitMaterialSetUpdate(
 	const FComponentMaterialSet& UpdatedMaterials,
 	bool bApplyToAsset = true);
 
+/**
+ * @param Target					The tool target to query for LODs
+ * @param bOutTargetSupportsLODs	Will be set to 'false' if the target does not support querying MeshDescriptions with LODs
+ * @param bOnlyReturnDefaultLOD		If true, always return an array of just the Default LOD
+ * @param bExcludeAutoGeneratedLODs	If true, exclude any auto-generated LODs (which have no associated Mesh Description) from the LOD list
+ * @return An array of valid LODs for the mesh data underlying the ToolTarget, or just the Default LOD if the target doesn't support LODs or bOnlyReturnDefaultLOD is true
+ */
+MODELINGCOMPONENTS_API TArray<EMeshLODIdentifier> GetMeshDescriptionLODs(
+	UToolTarget* Target,
+	bool& bOutTargetSupportsLODs,
+	bool bOnlyReturnDefaultLOD = false,
+	bool bExcludeAutoGeneratedLODs = true);
+
 
 /**
  * @return the MeshDescription underlying a ToolTarget, if it has such a mesh. May be generated internally by the ToolTarget. May be nullptr if the Target does not have a mesh.
@@ -142,7 +155,8 @@ enum class EDynamicMeshUpdateResult
 MODELINGCOMPONENTS_API EDynamicMeshUpdateResult CommitMeshDescriptionUpdate(
 	UToolTarget* Target, 
 	const FMeshDescription* UpdatedMesh, 
-	const FComponentMaterialSet* UpdatedMaterials = nullptr);
+	const FComponentMaterialSet* UpdatedMaterials = nullptr,
+	const FCommitMeshParameters& CommitParams = FCommitMeshParameters());
 
 /**
 * Update the Mesh in a ToolTarget based on the provided MeshDescription, and optional material set
@@ -150,7 +164,8 @@ MODELINGCOMPONENTS_API EDynamicMeshUpdateResult CommitMeshDescriptionUpdate(
 */
 MODELINGCOMPONENTS_API EDynamicMeshUpdateResult CommitMeshDescriptionUpdate(
 	UToolTarget* Target, 
-	FMeshDescription&& UpdatedMesh);
+	FMeshDescription&& UpdatedMesh,
+	const FCommitMeshParameters& CommitParams = FCommitMeshParameters());
 
 /**
 * Update the Mesh in a ToolTarget based on the provided MeshDescription, and optional material set
@@ -159,7 +174,8 @@ MODELINGCOMPONENTS_API EDynamicMeshUpdateResult CommitMeshDescriptionUpdate(
 MODELINGCOMPONENTS_API EDynamicMeshUpdateResult CommitMeshDescriptionUpdateViaDynamicMesh(
 	UToolTarget* Target, 
 	const UE::Geometry::FDynamicMesh3& UpdatedMesh,
-	bool bHaveModifiedTopology);
+	bool bHaveModifiedTopology,
+	const FCommitMeshParameters& CommitParams = FCommitMeshParameters());
 
 
 /**
diff --git a/Engine/Plugins/Runtime/MeshModelingToolset/Source/ModelingComponentsEditorOnly/Private/ToolTargets/StaticMeshComponentToolTarget.cpp b/Engine/Plugins/Runtime/MeshModelingToolset/Source/ModelingComponentsEditorOnly/Private/ToolTargets/StaticMeshComponentToolTarget.cpp
--- a/Engine/Plugins/Runtime/MeshModelingToolset/Source/ModelingComponentsEditorOnly/Private/ToolTargets/StaticMeshComponentToolTarget.cpp
+++ b/Engine/Plugins/Runtime/MeshModelingToolset/Source/ModelingComponentsEditorOnly/Private/ToolTargets/StaticMeshComponentToolTarget.cpp
@@ -217,6 +217,13 @@ FMeshDescription UStaticMeshComponentToolTarget::GetMeshDescriptionCopy(const FG
 }
 
 
+TArray<EMeshLODIdentifier> UStaticMeshComponentToolTarget::GetAvailableLODs(bool bSkipAutoGenerated) const
+{
+	UStaticMesh* StaticMesh = Cast<UStaticMeshComponent>(Component)->GetStaticMesh();
+	return UStaticMeshToolTarget::GetAvailableLODs(StaticMesh, bSkipAutoGenerated);
+}
+
+
 void UStaticMeshComponentToolTarget::CommitMeshDescription(const FCommitter& Committer, const FCommitMeshParameters& CommitParams)
 {
 	if (ensure(IsValid()) == false) return;
diff --git a/Engine/Plugins/Runtime/MeshModelingToolset/Source/ModelingComponentsEditorOnly/Private/ToolTargets/StaticMeshToolTarget.cpp b/Engine/Plugins/Runtime/MeshModelingToolset/Source/ModelingComponentsEditorOnly/Private/ToolTargets/StaticMeshToolTarget.cpp
--- a/Engine/Plugins/Runtime/MeshModelingToolset/Source/ModelingComponentsEditorOnly/Private/ToolTargets/StaticMeshToolTarget.cpp
+++ b/Engine/Plugins/Runtime/MeshModelingToolset/Source/ModelingComponentsEditorOnly/Private/ToolTargets/StaticMeshToolTarget.cpp
@@ -73,6 +73,41 @@ EMeshLODIdentifier UStaticMeshToolTarget::GetValidEditingLOD(const UStaticMesh*
 	return ValidEditingLOD;
 }
 
+TArray<EMeshLODIdentifier> UStaticMeshToolTarget::GetAvailableLODs(bool bSkipAutoGenerated) const
+{
+	return GetAvailableLODs(StaticMesh, bSkipAutoGenerated);
+}
+
+TArray<EMeshLODIdentifier> UStaticMeshToolTarget::GetAvailableLODs(const UStaticMesh* InStaticMesh, bool bSkipAutoGenerated)
+{
+	TArray<EMeshLODIdentifier> AvailableLODs;
+	
+	if (!InStaticMesh)
+	{
+		return AvailableLODs;
+	}
+
+	if (InStaticMesh->IsHiResMeshDescriptionValid())
+	{
+		AvailableLODs.Add(EMeshLODIdentifier::HiResSource);
+	}
+	int32 NumExistingLOD = FMath::Min(8, InStaticMesh->GetNumSourceModels());
+	for (int32 LOD = 0; LOD < NumExistingLOD; ++LOD)
+	{
+		if (bSkipAutoGenerated)
+		{
+			// Generated LODs have no source model data
+			if (!InStaticMesh->GetSourceModel(LOD).IsSourceModelInitialized())
+			{
+				continue;
+			}
+		}
+		AvailableLODs.Add((EMeshLODIdentifier)LOD);
+	}
+
+	return AvailableLODs;
+}
+
 
 bool UStaticMeshToolTarget::IsValid() const
 {
diff --git a/Engine/Plugins/Runtime/MeshModelingToolset/Source/ModelingComponentsEditorOnly/Public/ToolTargets/StaticMeshComponentToolTarget.h b/Engine/Plugins/Runtime/MeshModelingToolset/Source/ModelingComponentsEditorOnly/Public/ToolTargets/StaticMeshComponentToolTarget.h
--- a/Engine/Plugins/Runtime/MeshModelingToolset/Source/ModelingComponentsEditorOnly/Public/ToolTargets/StaticMeshComponentToolTarget.h
+++ b/Engine/Plugins/Runtime/MeshModelingToolset/Source/ModelingComponentsEditorOnly/Public/ToolTargets/StaticMeshComponentToolTarget.h
@@ -48,6 +48,8 @@ public:
 	virtual const FMeshDescription* GetMeshDescription(const FGetMeshParameters& GetMeshParams = FGetMeshParameters()) override;
 	virtual FMeshDescription GetMeshDescriptionCopy(const FGetMeshParameters& GetMeshParams) override;
 	virtual FMeshDescription GetEmptyMeshDescription() override;
+	virtual bool SupportsLODs() const override { return true; }
+	virtual TArray<EMeshLODIdentifier> GetAvailableLODs(bool bSkipAutoGenerated = true) const override;
 
 	// IMeshDescritpionCommitter implementation
 	virtual void CommitMeshDescription(const FCommitter& Committer, const FCommitMeshParameters& CommitParams = FCommitMeshParameters()) override;
diff --git a/Engine/Plugins/Runtime/MeshModelingToolset/Source/ModelingComponentsEditorOnly/Public/ToolTargets/StaticMeshToolTarget.h b/Engine/Plugins/Runtime/MeshModelingToolset/Source/ModelingComponentsEditorOnly/Public/ToolTargets/StaticMeshToolTarget.h
--- a/Engine/Plugins/Runtime/MeshModelingToolset/Source/ModelingComponentsEditorOnly/Public/ToolTargets/StaticMeshToolTarget.h
+++ b/Engine/Plugins/Runtime/MeshModelingToolset/Source/ModelingComponentsEditorOnly/Public/ToolTargets/StaticMeshToolTarget.h
@@ -52,6 +52,8 @@ public:
 	// IMeshDescriptionProvider implementation
 	virtual const FMeshDescription* GetMeshDescription(const FGetMeshParameters& GetMeshParams = FGetMeshParameters()) override;
 	virtual FMeshDescription GetEmptyMeshDescription() override;
+	virtual bool SupportsLODs() const override { return true; }
+	virtual TArray<EMeshLODIdentifier> GetAvailableLODs(bool bSkipAutoGenerated = true) const override;
 
 	// IMeshDescritpionCommitter implementation
 	virtual void CommitMeshDescription(const FCommitter& Committer, const FCommitMeshParameters& CommitParams = FCommitMeshParameters()) override;
@@ -87,6 +89,7 @@ protected:
 	static bool IsValid(const UStaticMesh* StaticMesh, EMeshLODIdentifier EditingLOD);
 	static EMeshLODIdentifier GetValidEditingLOD(const UStaticMesh* StaticMesh, 
 		EMeshLODIdentifier RequestedEditingLOD);
+	static TArray<EMeshLODIdentifier> GetAvailableLODs(const UStaticMesh* StaticMesh, bool bSkipAutoGenerated);
 	static void CommitMeshDescription(UStaticMesh* StaticMesh, 
 		const FCommitter& Committer, EMeshLODIdentifier EditingLODIn);
 	static void GetMaterialSet(const UStaticMesh* StaticMesh, 
diff --git a/Engine/Source/Runtime/InteractiveToolsFramework/Public/TargetInterfaces/MeshDescriptionProvider.h b/Engine/Source/Runtime/InteractiveToolsFramework/Public/TargetInterfaces/MeshDescriptionProvider.h
--- a/Engine/Source/Runtime/InteractiveToolsFramework/Public/TargetInterfaces/MeshDescriptionProvider.h
+++ b/Engine/Source/Runtime/InteractiveToolsFramework/Public/TargetInterfaces/MeshDescriptionProvider.h
@@ -46,4 +46,23 @@ public:
 		return *GetMeshDescription(GetMeshParams);
 	}
 
+	/**
+	 * For providers that have LODs (i.e. if SupportsLODs returns true), returns an array of all available LODs. Otherwise, returns an array with just the Default LOD.
+	 * @param bSkipAutoGenerated	Skip auto-generated LODs, which have no associated source Mesh Description
+	 */
+	virtual TArray<EMeshLODIdentifier> GetAvailableLODs(bool bSkipAutoGenerated = true) const
+	{
+		TArray<EMeshLODIdentifier> DefaultLODArray;
+		DefaultLODArray.Add(EMeshLODIdentifier::Default);
+		return DefaultLODArray;
+	}
+
+	/**
+	 * @return true if the provider is able to provide different Mesh Descriptions for different requested LODs.  If so, GetValidLODs() will return the available LODs.
+	 */
+	virtual bool SupportsLODs() const
+	{
+		return false;
+	}
+
 };
\ No newline at end of file
diff --git a/Engine/Source/Runtime/InteractiveToolsFramework/Public/TargetInterfaces/MeshTargetInterfaceTypes.h b/Engine/Source/Runtime/InteractiveToolsFramework/Public/TargetInterfaces/MeshTargetInterfaceTypes.h
--- a/Engine/Source/Runtime/InteractiveToolsFramework/Public/TargetInterfaces/MeshTargetInterfaceTypes.h
+++ b/Engine/Source/Runtime/InteractiveToolsFramework/Public/TargetInterfaces/MeshTargetInterfaceTypes.h
@@ -31,6 +31,11 @@ typedef EMeshLODIdentifier EStaticMeshEditingLOD;
  */
 struct INTERACTIVETOOLSFRAMEWORK_API FGetMeshParameters
 {
+	FGetMeshParameters() = default;
+	explicit FGetMeshParameters(bool bHaveRequestLOD, EMeshLODIdentifier RequestLOD)
+		: bHaveRequestLOD(bHaveRequestLOD), RequestLOD(RequestLOD)
+	{}
+
 	/** If true, RequestLOD specifies which mesh LOD should be returned, if available. If not available, behavior is implementation-dependent. */
 	bool bHaveRequestLOD = false;
 	/** Specify which LOD to fetch from a mesh-containing ToolTarget. Ignored unless bHaveRequestLOD == true. */
@@ -50,6 +55,11 @@ struct INTERACTIVETOOLSFRAMEWORK_API FGetMeshParameters
  */
 struct INTERACTIVETOOLSFRAMEWORK_API FCommitMeshParameters
 {
+	FCommitMeshParameters() = default;
+	explicit FCommitMeshParameters(bool bHaveTargetLOD, EMeshLODIdentifier TargetLOD) :
+		bHaveTargetLOD(bHaveTargetLOD), TargetLOD(TargetLOD)
+	{}
+
 	/** If true, TargetLOD specifies which mesh LOD should be set/updated */
 	bool bHaveTargetLOD = false;
 	/** Specify which LOD to set in a mesh-containing ToolTarget. Ignored unless bHaveTargetLOD == true. */
