Mesh Modeling - Static Mesh and Static Mesh Component Tool Targets.
Update Can Build Target to distinguish generated LODs.   Tools shouldn't edit generated LODs.

diff --git a/Engine/Plugins/Runtime/MeshModelingToolset/Source/ModelingComponentsEditorOnly/Private/ToolTargets/StaticMeshComponentToolTarget.cpp b/Engine/Plugins/Runtime/MeshModelingToolset/Source/ModelingComponentsEditorOnly/Private/ToolTargets/StaticMeshComponentToolTarget.cpp
--- a/Engine/Plugins/Runtime/MeshModelingToolset/Source/ModelingComponentsEditorOnly/Private/ToolTargets/StaticMeshComponentToolTarget.cpp
+++ b/Engine/Plugins/Runtime/MeshModelingToolset/Source/ModelingComponentsEditorOnly/Private/ToolTargets/StaticMeshComponentToolTarget.cpp
@@ -289,9 +289,10 @@ IInterface_CollisionDataProvider* UStaticMeshComponentToolTarget::GetComplexColl
 bool UStaticMeshComponentToolTargetFactory::CanBuildTarget(UObject* SourceObject, const FToolTargetTypeRequirements& Requirements) const
 {
 	const UStaticMeshComponent* Component = GetValid(Cast<UStaticMeshComponent>(SourceObject));
-	return Component && !Component->IsUnreachable() && Component->IsValidLowLevel() && Component->GetStaticMesh()
-		&& !Component->GetStaticMesh()->GetOutermost()->bIsCookedForEditor
-		&& (Component->GetStaticMesh()->GetNumSourceModels() > 0)
+	const UStaticMesh* StaticMesh = (Component && !Component->IsUnreachable() && Component->IsValidLowLevel()) ? Component->GetStaticMesh() : nullptr;
+	return StaticMesh
+		&& !StaticMesh->GetOutermost()->bIsCookedForEditor
+		&& UStaticMeshToolTarget::HasNonGeneratedLOD(StaticMesh, EditingLOD)
 		&& Requirements.AreSatisfiedBy(UStaticMeshComponentToolTarget::StaticClass());
 }
 
diff --git a/Engine/Plugins/Runtime/MeshModelingToolset/Source/ModelingComponentsEditorOnly/Private/ToolTargets/StaticMeshToolTarget.cpp b/Engine/Plugins/Runtime/MeshModelingToolset/Source/ModelingComponentsEditorOnly/Private/ToolTargets/StaticMeshToolTarget.cpp
--- a/Engine/Plugins/Runtime/MeshModelingToolset/Source/ModelingComponentsEditorOnly/Private/ToolTargets/StaticMeshToolTarget.cpp
+++ b/Engine/Plugins/Runtime/MeshModelingToolset/Source/ModelingComponentsEditorOnly/Private/ToolTargets/StaticMeshToolTarget.cpp
@@ -136,6 +136,39 @@ bool UStaticMeshToolTarget::IsValid(const UStaticMesh* StaticMeshIn, EMeshLODIde
 	return true;
 }
 
+bool UStaticMeshToolTarget::HasNonGeneratedLOD(const UStaticMesh* StaticMeshIn, EMeshLODIdentifier MeshLOD)
+{
+	if (!StaticMeshIn)
+	{
+		return false;
+	}
+	
+	if (MeshLOD == EMeshLODIdentifier::Default)
+	{
+		MeshLOD = EMeshLODIdentifier::LOD0;
+	}
+	if (MeshLOD == EMeshLODIdentifier::HiResSource)
+	{
+		return StaticMeshIn->IsHiResMeshDescriptionValid();
+	}
+	const int32 NumExistingLOD = FMath::Min(8, StaticMeshIn->GetNumSourceModels());
+	if (MeshLOD == EMeshLODIdentifier::MaxQuality)
+	{
+		const bool bHasHiRes = StaticMeshIn->IsHiResMeshDescriptionValid();
+		const bool bHasLOD0 = (NumExistingLOD > 0) ? StaticMeshIn->GetSourceModel(0).IsSourceModelInitialized() : false;
+		return (bHasHiRes || bHasLOD0);
+	}
+	else
+	{
+		const int32 LODint = static_cast<int32>(MeshLOD);
+		if (LODint >= NumExistingLOD)
+		{ 
+			return false; 
+		}
+
+		return StaticMeshIn->GetSourceModel(LODint).IsSourceModelInitialized();
+	}
+}
 
 int32 UStaticMeshToolTarget::GetNumMaterials() const
 {
@@ -378,7 +411,7 @@ bool UStaticMeshToolTargetFactory::CanBuildTarget(UObject* SourceObject, const F
 	const UStaticMesh* StaticMesh = GetValid(Cast<UStaticMesh>(SourceObject));
 	return StaticMesh && !StaticMesh->IsUnreachable() && StaticMesh->IsValidLowLevel()
 		&& !StaticMesh->GetOutermost()->bIsCookedForEditor
-		&& (StaticMesh->GetNumSourceModels() > 0)
+		&& UStaticMeshToolTarget::HasNonGeneratedLOD(StaticMesh, EditingLOD)
 		&& Requirements.AreSatisfiedBy(UStaticMeshToolTarget::StaticClass());
 }
 
diff --git a/Engine/Plugins/Runtime/MeshModelingToolset/Source/ModelingComponentsEditorOnly/Public/ToolTargets/StaticMeshToolTarget.h b/Engine/Plugins/Runtime/MeshModelingToolset/Source/ModelingComponentsEditorOnly/Public/ToolTargets/StaticMeshToolTarget.h
--- a/Engine/Plugins/Runtime/MeshModelingToolset/Source/ModelingComponentsEditorOnly/Public/ToolTargets/StaticMeshToolTarget.h
+++ b/Engine/Plugins/Runtime/MeshModelingToolset/Source/ModelingComponentsEditorOnly/Public/ToolTargets/StaticMeshToolTarget.h
@@ -86,7 +86,11 @@ protected:
 
 	friend class UStaticMeshComponentToolTarget;
 
+	friend class UStaticMeshComponentToolTargetFactory;
+
 	static bool IsValid(const UStaticMesh* StaticMesh, EMeshLODIdentifier EditingLOD);
+	// returns true if the mesh has a non-generated LOD of the requested level.  Interprets EMeshLODIdentifier::Default as LOD0
+	static bool HasNonGeneratedLOD(const UStaticMesh* StaticMesh, EMeshLODIdentifier MeshLOD);
 	static EMeshLODIdentifier GetValidEditingLOD(const UStaticMesh* StaticMesh, 
 		EMeshLODIdentifier RequestedEditingLOD);
 	static TArray<EMeshLODIdentifier> GetAvailableLODs(const UStaticMesh* StaticMesh, bool bSkipAutoGenerated);
