UVEditor: Enable support for UV island normalization in Layout operations.

diff --git a/Engine/Plugins/Experimental/UVEditor/Source/UVEditorTools/Private/UVEditorLayoutTool.cpp b/Engine/Plugins/Experimental/UVEditor/Source/UVEditorTools/Private/UVEditorLayoutTool.cpp
--- a/Engine/Plugins/Experimental/UVEditor/Source/UVEditorTools/Private/UVEditorLayoutTool.cpp
+++ b/Engine/Plugins/Experimental/UVEditor/Source/UVEditorTools/Private/UVEditorLayoutTool.cpp
@@ -54,6 +54,8 @@ void UUVEditorLayoutTool::Setup()
 	UContextObjectStore* ContextStore = GetToolManager()->GetContextObjectStore();
 	UVToolSelectionAPI = ContextStore->FindContext<UUVToolSelectionAPI>();
 
+	UVTool2DViewportAPI = ContextStore->FindContext<UUVTool2DViewportAPI>();
+
 	UUVToolSelectionAPI::FHighlightOptions HighlightOptions;
 	HighlightOptions.bBaseHighlightOnPreviews = true;
 	HighlightOptions.bAutoUpdateUnwrap = true;
@@ -73,6 +75,16 @@ void UUVEditorLayoutTool::Setup()
 				                                                        UE::Geometry::FUVToolSelection::EType::Triangle).SelectedIDs);
 		}
 
+		if (UVTool2DViewportAPI)
+		{
+			TMap<int32, int32> TextureResolutionPerUDIM;
+			for (const FUDIMBlock& UDIM : UVTool2DViewportAPI->GetUDIMBlocks())
+			{
+				TextureResolutionPerUDIM.Add(UDIM.UDIM, UDIM.TextureResolution);
+			}
+			Factory->TextureResolutionPerUDIM = TextureResolutionPerUDIM;
+		}
+
 		Target.AppliedPreview->ChangeOpFactory(Factory);
 		Target.AppliedPreview->OnMeshUpdated.AddWeakLambda(this, [this, &Target](UMeshOpPreviewWithBackgroundCompute* Preview) {
 			Target.UpdateUnwrapPreviewFromAppliedPreview();
diff --git a/Engine/Plugins/Experimental/UVEditor/Source/UVEditorTools/Public/UVEditorLayoutTool.h b/Engine/Plugins/Experimental/UVEditor/Source/UVEditorTools/Public/UVEditorLayoutTool.h
--- a/Engine/Plugins/Experimental/UVEditor/Source/UVEditorTools/Public/UVEditorLayoutTool.h
+++ b/Engine/Plugins/Experimental/UVEditor/Source/UVEditorTools/Public/UVEditorLayoutTool.h
@@ -14,6 +14,7 @@
 class UUVEditorToolMeshInput;
 class UUVLayoutProperties;
 class UUVLayoutOperatorFactory;
+class UUVTool2DViewportAPI;
 
 UCLASS()
 class UVEDITORTOOLS_API UUVEditorLayoutToolBuilder : public UInteractiveToolBuilder
@@ -71,6 +72,10 @@ protected:
 	UPROPERTY()
 	TObjectPtr<UUVToolSelectionAPI> UVToolSelectionAPI = nullptr;
 
+	//~ For UDIM information access
+	UPROPERTY()
+	TObjectPtr< UUVTool2DViewportAPI> UVTool2DViewportAPI = nullptr; 
+
 	//
 	// Analytics
 	//
diff --git a/Engine/Plugins/Runtime/GeometryProcessing/Source/DynamicMesh/Private/Parameterization/MeshUVPacking.cpp b/Engine/Plugins/Runtime/GeometryProcessing/Source/DynamicMesh/Private/Parameterization/MeshUVPacking.cpp
--- a/Engine/Plugins/Runtime/GeometryProcessing/Source/DynamicMesh/Private/Parameterization/MeshUVPacking.cpp
+++ b/Engine/Plugins/Runtime/GeometryProcessing/Source/DynamicMesh/Private/Parameterization/MeshUVPacking.cpp
@@ -68,6 +68,7 @@ bool FDynamicMeshUVPacker::StandardPack()
 	Packer.bAllowFlips = bAllowFlips;
 	Packer.GutterSize = GutterSize;
 	Packer.TextureResolution = TextureResolution;
+	Packer.bScaleIslandsByWorldSpaceTexelRatio = bScaleIslandsByWorldSpaceTexelRatio;
 
 	FUVOverlayView MeshView(UVOverlay);
 	FMeshConnectedComponents UVComponents = CollectUVIslandsToPack(MeshView);
@@ -86,6 +87,7 @@ bool FDynamicMeshUVPacker::StackPack()
 	Packer.bAllowFlips = bAllowFlips;
 	Packer.GutterSize = GutterSize;
 	Packer.TextureResolution = TextureResolution;
+	Packer.bScaleIslandsByWorldSpaceTexelRatio = bScaleIslandsByWorldSpaceTexelRatio;
 
 	FUVOverlayView MeshView(UVOverlay);
 	FMeshConnectedComponents UVComponents = CollectUVIslandsToPack(MeshView);
diff --git a/Engine/Plugins/Runtime/GeometryProcessing/Source/DynamicMesh/Public/Parameterization/MeshUVPacking.h b/Engine/Plugins/Runtime/GeometryProcessing/Source/DynamicMesh/Public/Parameterization/MeshUVPacking.h
--- a/Engine/Plugins/Runtime/GeometryProcessing/Source/DynamicMesh/Public/Parameterization/MeshUVPacking.h
+++ b/Engine/Plugins/Runtime/GeometryProcessing/Source/DynamicMesh/Public/Parameterization/MeshUVPacking.h
@@ -37,6 +37,9 @@ public:
 	/** If true, islands can be flipped in addition to rotate/translate/scale */
 	bool bAllowFlips = false;
 
+	/** Attempt to rescale islands to match texel-to-world-space ratio across islands, based on ratio of World- and UV-space edge lengths */
+	bool bScaleIslandsByWorldSpaceTexelRatio = false;
+
 	explicit FDynamicMeshUVPacker(FDynamicMeshUVOverlay* UVOverlay);
 	explicit FDynamicMeshUVPacker(FDynamicMeshUVOverlay* UVOverlay, TUniquePtr<TArray<int32>>&& TidsToRepackIn);
 
@@ -55,6 +58,7 @@ public:
 	 */
 	bool StackPack();
 
+
 protected:
 
 	FMeshConnectedComponents CollectUVIslandsToPack(const FUVOverlayView& MeshView);
diff --git a/Engine/Plugins/Runtime/MeshModelingToolset/Source/ModelingOperators/ModelingOperators.Build.cs b/Engine/Plugins/Runtime/MeshModelingToolset/Source/ModelingOperators/ModelingOperators.Build.cs
--- a/Engine/Plugins/Runtime/MeshModelingToolset/Source/ModelingOperators/ModelingOperators.Build.cs
+++ b/Engine/Plugins/Runtime/MeshModelingToolset/Source/ModelingOperators/ModelingOperators.Build.cs
@@ -15,7 +15,8 @@ public class ModelingOperators : ModuleRules
 				"DynamicMesh",
 				"MeshConversion",
 				"GeometryAlgorithms", // required for constrained Delaunay triangulation
-				"SkeletalMeshDescription"
+				"SkeletalMeshDescription",
+				"TextureUtilitiesCommon", // required for UDIM manipulation
 				// ... add other public dependencies that you statically link with here ...
 			}
 		);
diff --git a/Engine/Plugins/Runtime/MeshModelingToolset/Source/ModelingOperators/Private/ParameterizationOps/UVLayoutOp.cpp b/Engine/Plugins/Runtime/MeshModelingToolset/Source/ModelingOperators/Private/ParameterizationOps/UVLayoutOp.cpp
--- a/Engine/Plugins/Runtime/MeshModelingToolset/Source/ModelingOperators/Private/ParameterizationOps/UVLayoutOp.cpp
+++ b/Engine/Plugins/Runtime/MeshModelingToolset/Source/ModelingOperators/Private/ParameterizationOps/UVLayoutOp.cpp
@@ -8,6 +8,7 @@
 #include "Properties/UVLayoutProperties.h"
 #include "Selections/MeshConnectedComponents.h"
 #include "Parameterization/MeshUDIMClassifier.h"
+#include "UDIMUtilities.h"
 
 #include UE_INLINE_GENERATED_CPP_BY_NAME(UVLayoutOp)
 
@@ -86,6 +87,12 @@ void FUVLayoutOp::CalculateResult(FProgressCancel* Progress)
 		{
 			TUniquePtr<TArray<int32>> TileTids;
 			TileTids = MakeUnique<TArray<int32>>(TileClassifier.TidsForTile(TileIndex));
+			const int32 TileID = UE::TextureUtilitiesCommon::GetUDIMIndex(TileIndex.X, TileIndex.Y);
+			float TileResolution = TextureResolution;
+			if (TextureResolutionPerUDIM.IsSet())
+			{
+				TileResolution = TextureResolutionPerUDIM.GetValue().FindOrAdd(TileID, this->TextureResolution);
+			}
 
 			// Do this first, so we don't need to keep the TileTids around after moving it into the packer.
 			TSet<int32> ElementsToMove;
@@ -100,6 +107,7 @@ void FUVLayoutOp::CalculateResult(FProgressCancel* Progress)
 
 			// TODO: There is a second connected components call inside the packer that might be unnessessary. Could be a future optimization.
 			FDynamicMeshUVPacker Packer(UseUVLayer, MoveTemp(TileTids) );
+			Packer.TextureResolution = TileResolution;					
 			ExecutePacker(Packer);
 			if (Progress && Progress->Cancelled())
 			{
@@ -132,6 +140,7 @@ void FUVLayoutOp::CalculateResult(FProgressCancel* Progress)
 		}
 		
 		FDynamicMeshUVPacker Packer(UseUVLayer, MoveTemp(TidsToLayout) );
+		Packer.TextureResolution = this->TextureResolution;
 		ExecutePacker(Packer);
 		if (Progress && Progress->Cancelled())
 		{
@@ -152,8 +161,7 @@ void FUVLayoutOp::CalculateResult(FProgressCancel* Progress)
 }
 
 void FUVLayoutOp::ExecutePacker(FDynamicMeshUVPacker& Packer)
-{
-	Packer.TextureResolution = this->TextureResolution;
+{	
 	Packer.GutterSize = this->GutterSize;
 	Packer.bAllowFlips = this->bAllowFlips;
 
@@ -173,6 +181,15 @@ void FUVLayoutOp::ExecutePacker(FDynamicMeshUVPacker& Packer)
 			return;
 		}
 	}
+	else if (UVLayoutMode == EUVLayoutOpLayoutModes::Normalize)
+	{
+		Packer.bScaleIslandsByWorldSpaceTexelRatio = true;
+		if (Packer.StandardPack() == false)
+		{
+			// failed... what to do?
+			return;
+		}
+	}
 }
 
 TUniquePtr<FDynamicMeshOperator> UUVLayoutOperatorFactory::MakeNewOperator()
@@ -192,6 +209,9 @@ TUniquePtr<FDynamicMeshOperator> UUVLayoutOperatorFactory::MakeNewOperator()
 	case EUVLayoutType::Repack:
 		Op->UVLayoutMode = EUVLayoutOpLayoutModes::RepackToUnitRect;
 		break;
+	case EUVLayoutType::Normalize:
+		Op->UVLayoutMode = EUVLayoutOpLayoutModes::Normalize;
+		break;
 	}
 
 	Op->UVLayerIndex = GetSelectedUVChannel();
@@ -201,6 +221,8 @@ TUniquePtr<FDynamicMeshOperator> UUVLayoutOperatorFactory::MakeNewOperator()
 	Op->UVTranslation = FVector2f(Settings->Translation);
 	Op->SetTransform(TargetTransform);
 	Op->bMaintainOriginatingUDIM = Settings->bEnableUDIMLayout;
+	Op->Selection = Selection;
+	Op->TextureResolutionPerUDIM = TextureResolutionPerUDIM;
 
 	return Op;
 }
diff --git a/Engine/Plugins/Runtime/MeshModelingToolset/Source/ModelingOperators/Public/ParameterizationOps/UVLayoutOp.h b/Engine/Plugins/Runtime/MeshModelingToolset/Source/ModelingOperators/Public/ParameterizationOps/UVLayoutOp.h
--- a/Engine/Plugins/Runtime/MeshModelingToolset/Source/ModelingOperators/Public/ParameterizationOps/UVLayoutOp.h
+++ b/Engine/Plugins/Runtime/MeshModelingToolset/Source/ModelingOperators/Public/ParameterizationOps/UVLayoutOp.h
@@ -21,7 +21,8 @@ enum class EUVLayoutOpLayoutModes
 {
 	TransformOnly = 0,
 	RepackToUnitRect = 1,
-	StackInUnitRect = 2
+	StackInUnitRect = 2,
+	Normalize = 3
 };
 
 class MODELINGOPERATORS_API FUVLayoutOp : public FDynamicMeshOperator
@@ -42,6 +43,7 @@ public:
 	float GutterSize = 1.0;
 	bool bMaintainOriginatingUDIM = false;
 	TOptional<TSet<int32>> Selection;
+	TOptional<TMap<int32, int32>> TextureResolutionPerUDIM;
 	FVector2f UVTranslation = FVector2f::Zero();
 
 	void SetTransform(const FTransformSRT3d& Transform);
@@ -80,4 +82,5 @@ public:
 	TSharedPtr<UE::Geometry::FDynamicMesh3, ESPMode::ThreadSafe> OriginalMesh;
 	TUniqueFunction<int32()> GetSelectedUVChannel = []() { return 0; };
 	FTransform TargetTransform;
+	TOptional<TMap<int32, int32>> TextureResolutionPerUDIM;
 };
\ No newline at end of file
diff --git a/Engine/Plugins/Runtime/MeshModelingToolset/Source/ModelingOperators/Public/Properties/UVLayoutProperties.h b/Engine/Plugins/Runtime/MeshModelingToolset/Source/ModelingOperators/Public/Properties/UVLayoutProperties.h
--- a/Engine/Plugins/Runtime/MeshModelingToolset/Source/ModelingOperators/Public/Properties/UVLayoutProperties.h
+++ b/Engine/Plugins/Runtime/MeshModelingToolset/Source/ModelingOperators/Public/Properties/UVLayoutProperties.h
@@ -18,7 +18,9 @@ enum class EUVLayoutType
 	/** Uniformly scale and translate each UV island individually to pack it into the unit square, i.e. fit between 0 and 1 with overlap */
 	Stack,
 	/** Uniformly scale and translate UV islands collectively to pack them into the unit square, i.e. fit between 0 and 1 with no overlap */
-	Repack
+	Repack,
+	/** Scale and translate UV islands to normalize the UV islands' area to match an average texel density. */
+	Normalize
 };
 
 
