MeshModelingTools: Make FPhysicsDataCollection::CopyGeometryToAggregate clear the AggGeom before filling it. Fixes MeshToCollision tool duplicating simple geometry when appending

diff --git a/Engine/Plugins/Runtime/MeshModelingToolset/Source/ModelingComponents/Private/Physics/PhysicsDataCollection.cpp b/Engine/Plugins/Runtime/MeshModelingToolset/Source/ModelingComponents/Private/Physics/PhysicsDataCollection.cpp
--- a/Engine/Plugins/Runtime/MeshModelingToolset/Source/ModelingComponents/Private/Physics/PhysicsDataCollection.cpp
+++ b/Engine/Plugins/Runtime/MeshModelingToolset/Source/ModelingComponents/Private/Physics/PhysicsDataCollection.cpp
@@ -86,6 +86,8 @@ void FPhysicsDataCollection::ClearAggregate()
 
 void FPhysicsDataCollection::CopyGeometryToAggregate()
 {
+	ClearAggregate();
+
 	for (FBoxShape3d& BoxGeom : Geometry.Boxes)
 	{
 		FKBoxElem Element;
