Fixed mismatched copy constructor/assignment operator in FGeometryIdentifier.

diff --git a/Engine/Plugins/Runtime/MeshModelingToolset/Source/ModelingComponents/Public/Selection/GeometrySelector.h b/Engine/Plugins/Runtime/MeshModelingToolset/Source/ModelingComponents/Public/Selection/GeometrySelector.h
--- a/Engine/Plugins/Runtime/MeshModelingToolset/Source/ModelingComponents/Public/Selection/GeometrySelector.h
+++ b/Engine/Plugins/Runtime/MeshModelingToolset/Source/ModelingComponents/Public/Selection/GeometrySelector.h
@@ -63,7 +63,6 @@ struct MODELINGCOMPONENTS_API FGeometryIdentifier
 	TWeakObjectPtr<UObject> TargetObject = nullptr;
 
 	FGeometryIdentifier() = default;
-	FGeometryIdentifier(const FGeometryIdentifier& Other) = default;
 
 	/**
 	 * @return true if TargetObject is defined and still exists
