ModelingMode: support IGeometrySelector::Sleep() and ::Restore() failing in GeometrySelectionManager. FBaseDynamicMeshSelector will now fail to sleep if TargetMesh is not valid. SupportsSleep() now also checks for valid target mesh.

diff --git a/Engine/Plugins/Runtime/MeshModelingToolset/Source/ModelingComponents/Private/Selection/DynamicMeshSelector.cpp b/Engine/Plugins/Runtime/MeshModelingToolset/Source/ModelingComponents/Private/Selection/DynamicMeshSelector.cpp
--- a/Engine/Plugins/Runtime/MeshModelingToolset/Source/ModelingComponents/Private/Selection/DynamicMeshSelector.cpp
+++ b/Engine/Plugins/Runtime/MeshModelingToolset/Source/ModelingComponents/Private/Selection/DynamicMeshSelector.cpp
@@ -67,6 +67,12 @@ void FBaseDynamicMeshSelector::InvalidateOnMeshChange(FDynamicMeshChangeInfo Cha
 	NotifyGeometryModified();
 }
 
+bool FBaseDynamicMeshSelector::SupportsSleep() const
+{
+	// can only sleep if target mesh is valid
+	return TargetMesh.IsValid();
+}
+
 void FBaseDynamicMeshSelector::Shutdown()
 {
 	if (TargetMesh.IsValid())
@@ -83,7 +89,10 @@ void FBaseDynamicMeshSelector::Shutdown()
 
 bool FBaseDynamicMeshSelector::Sleep()
 {
-	check(TargetMesh.IsValid());
+	if (TargetMesh.IsValid() == false)
+	{
+		return false;
+	}
 
 	TargetMesh->OnMeshChanged().Remove(TargetMesh_OnMeshChangedHandle);
 	TargetMesh_OnMeshChangedHandle.Reset();
@@ -100,7 +109,10 @@ bool FBaseDynamicMeshSelector::Sleep()
 
 bool FBaseDynamicMeshSelector::Restore()
 {
-	check(SleepingTargetMesh.IsValid());
+	if (SleepingTargetMesh.IsValid() == false)
+	{
+		return false;
+	}
 
 	TargetMesh = SleepingTargetMesh.Get();
 	SleepingTargetMesh = nullptr;
diff --git a/Engine/Plugins/Runtime/MeshModelingToolset/Source/ModelingComponents/Private/Selection/GeometrySelectionManager.cpp b/Engine/Plugins/Runtime/MeshModelingToolset/Source/ModelingComponents/Private/Selection/GeometrySelectionManager.cpp
--- a/Engine/Plugins/Runtime/MeshModelingToolset/Source/ModelingComponents/Private/Selection/GeometrySelectionManager.cpp
+++ b/Engine/Plugins/Runtime/MeshModelingToolset/Source/ModelingComponents/Private/Selection/GeometrySelectionManager.cpp
@@ -47,7 +47,7 @@ void UGeometrySelectionManager::Shutdown()
 
 	for (TSharedPtr<FGeometrySelectionTarget> Target : ActiveTargetReferences)
 	{
-		SleepOrShutdownTarget(Target.Get(), false);
+		SleepOrShutdownTarget(Target, false);
 	}
 
 	ResetTargetCache();
@@ -315,7 +315,7 @@ void UGeometrySelectionManager::ClearActiveTargets()
 
 	for (TSharedPtr<FGeometrySelectionTarget> Target : ActiveTargetReferences)
 	{
-		SleepOrShutdownTarget(Target.Get(), false);
+		SleepOrShutdownTarget(Target, false);
 	}
 
 	ActiveTargetReferences.Reset();
@@ -493,18 +493,22 @@ void UGeometrySelectionManager::SetTargetsOnUndoRedo(TArray<FGeometryIdentifier>
 
 
 
-void UGeometrySelectionManager::SleepOrShutdownTarget(FGeometrySelectionTarget* Target, bool bForceShutdown)
+void UGeometrySelectionManager::SleepOrShutdownTarget(TSharedPtr<FGeometrySelectionTarget> Target, bool bForceShutdown)
 {
 	if (Target->Selector->SupportsSleep() && bForceShutdown == false)
 	{
-		bool bOK = Target->Selector->Sleep();
-		check(bOK);
-	}
-	else
-	{
-		Target->Selector->GetOnGeometryModifed().Remove(Target->OnGeometryModifiedHandle);
-		Target->Selector->Shutdown();
+		if (Target->Selector->Sleep())
+		{
+			return;
+		}
 	}
+
+	// if target cannot sleep or if sleeping failed, make sure it is not in the target
+	// cache so that we do not try to restore it later
+	TargetCache.Remove(Target->TargetIdentifier);
+
+	Target->Selector->GetOnGeometryModifed().Remove(Target->OnGeometryModifiedHandle);
+	Target->Selector->Shutdown();
 }
 
 TSharedPtr<UGeometrySelectionManager::FGeometrySelectionTarget> UGeometrySelectionManager::GetCachedTarget(FGeometryIdentifier TargetIdentifier, const IGeometrySelectorFactory* UseFactory)
@@ -513,15 +517,22 @@ TSharedPtr<UGeometrySelectionManager::FGeometrySelectionTarget> UGeometrySelecti
 	{
 		TSharedPtr<FGeometrySelectionTarget> FoundTarget = TargetCache[TargetIdentifier];
 		FoundTarget->Selection.Reset();
-		FoundTarget->Selector->Restore();
-
-		// ensure these are current, as they may have changed while Target was asleep
-		FoundTarget->Selection.ElementType = GetSelectionElementType();
-		FoundTarget->Selection.TopologyType = GetSelectionTopologyType();
-		bool bEnableTopologyFilter = (FoundTarget->Selection.TopologyType == EGeometryTopologyType::Polygroup && FoundTarget->Selection.ElementType != EGeometryElementType::Vertex);
-		FoundTarget->SelectionEditor->UpdateQueryConfig(GetCurrentSelectionQueryConfig(), bEnableTopologyFilter);
+		bool bRestored = FoundTarget->Selector->Restore();
+		if (bRestored)
+		{
+			// ensure these are current, as they may have changed while Target was asleep
+			FoundTarget->Selection.ElementType = GetSelectionElementType();
+			FoundTarget->Selection.TopologyType = GetSelectionTopologyType();
+			bool bEnableTopologyFilter = (FoundTarget->Selection.TopologyType == EGeometryTopologyType::Polygroup && FoundTarget->Selection.ElementType != EGeometryElementType::Vertex);
+			FoundTarget->SelectionEditor->UpdateQueryConfig(GetCurrentSelectionQueryConfig(), bEnableTopologyFilter);
 
-		return FoundTarget;
+			return FoundTarget;
+		}
+		else
+		{
+			// if restore failed, something is wrong w/ TargetCache, remove this Target
+			TargetCache.Remove(TargetIdentifier);
+		}
 	}
 
 	// if we are in a situation where we don't have a cache, currently we need the Factory to exist?
@@ -556,7 +567,7 @@ void UGeometrySelectionManager::ResetTargetCache()
 {
 	for (TPair<FGeometryIdentifier, TSharedPtr<FGeometrySelectionTarget>> Pair : TargetCache)
 	{
-		SleepOrShutdownTarget(Pair.Value.Get(), true);
+		SleepOrShutdownTarget(Pair.Value, true);
 	}
 	TargetCache.Reset();
 }
diff --git a/Engine/Plugins/Runtime/MeshModelingToolset/Source/ModelingComponents/Public/Selection/DynamicMeshSelector.h b/Engine/Plugins/Runtime/MeshModelingToolset/Source/ModelingComponents/Public/Selection/DynamicMeshSelector.h
--- a/Engine/Plugins/Runtime/MeshModelingToolset/Source/ModelingComponents/Public/Selection/DynamicMeshSelector.h
+++ b/Engine/Plugins/Runtime/MeshModelingToolset/Source/ModelingComponents/Public/Selection/DynamicMeshSelector.h
@@ -53,7 +53,7 @@ public:
 	//
 
 	virtual void Shutdown() override;
-	virtual bool SupportsSleep() const { return true; }
+	virtual bool SupportsSleep() const;
 	virtual bool Sleep();
 	virtual bool Restore();
 
diff --git a/Engine/Plugins/Runtime/MeshModelingToolset/Source/ModelingComponents/Public/Selection/GeometrySelectionManager.h b/Engine/Plugins/Runtime/MeshModelingToolset/Source/ModelingComponents/Public/Selection/GeometrySelectionManager.h
--- a/Engine/Plugins/Runtime/MeshModelingToolset/Source/ModelingComponents/Public/Selection/GeometrySelectionManager.h
+++ b/Engine/Plugins/Runtime/MeshModelingToolset/Source/ModelingComponents/Public/Selection/GeometrySelectionManager.h
@@ -415,7 +415,7 @@ protected:
 	// Functional, but needs to be smarter.
 	//
 	TMap<FGeometryIdentifier, TSharedPtr<FGeometrySelectionTarget>> TargetCache;
-	void SleepOrShutdownTarget(FGeometrySelectionTarget* Target, bool bForceShutdown);
+	void SleepOrShutdownTarget(TSharedPtr<FGeometrySelectionTarget> Target, bool bForceShutdown);
 	TSharedPtr<FGeometrySelectionTarget> GetCachedTarget(FGeometryIdentifier Identifier, const IGeometrySelectorFactory* UseFactory);
 	void ResetTargetCache();
 	void SetTargetsOnUndoRedo(TArray<FGeometryIdentifier> NewTargets);
