MeshModelingTools: Fix Nan warning in PolyExt.

diff --git a/Engine/Plugins/Runtime/MeshModelingToolset/Source/MeshModelingTools/Private/DrawPolygonTool.cpp b/Engine/Plugins/Runtime/MeshModelingToolset/Source/MeshModelingTools/Private/DrawPolygonTool.cpp
--- a/Engine/Plugins/Runtime/MeshModelingToolset/Source/MeshModelingTools/Private/DrawPolygonTool.cpp
+++ b/Engine/Plugins/Runtime/MeshModelingToolset/Source/MeshModelingTools/Private/DrawPolygonTool.cpp
@@ -877,7 +877,8 @@ void UDrawPolygonTool::GetPolygonParametersFromFixedPoints(const TArray<FVector3
 	AngleRad = FMathd::Atan2(Delta.Y, Delta.X);
 
 	double Radius = Delta.Length();
-	FVector2d AxisX = Delta / Radius;
+	FVector2d AxisX = Radius != 0 ? Delta / Radius
+		: FVector2d(1,0); // arbitrary if delta was 0 vector
 	FVector2d AxisY = -UE::Geometry::PerpCW(AxisX);
 	FVector2d HeightPt = DrawFrame.ToPlaneUV((FixedPoints.Num() == 3) ? FixedPoints[2] : FixedPoints[1], 2);
 	FVector2d HeightDelta = HeightPt - FirstReferencePt;
