MeshModelingTools: Don't transact static mesh asset creation because we are not able to undo it properly anyway.

diff --git a/Engine/Plugins/Runtime/MeshModelingToolset/Source/ModelingComponentsEditorOnly/Private/AssetUtils/CreateStaticMeshUtil.cpp b/Engine/Plugins/Runtime/MeshModelingToolset/Source/ModelingComponentsEditorOnly/Private/AssetUtils/CreateStaticMeshUtil.cpp
--- a/Engine/Plugins/Runtime/MeshModelingToolset/Source/ModelingComponentsEditorOnly/Private/AssetUtils/CreateStaticMeshUtil.cpp
+++ b/Engine/Plugins/Runtime/MeshModelingToolset/Source/ModelingComponentsEditorOnly/Private/AssetUtils/CreateStaticMeshUtil.cpp
@@ -2,6 +2,7 @@
 
 #include "AssetUtils/CreateStaticMeshUtil.h"
 
+#include "CoreGlobals.h" // GUndo
 #include "Engine/StaticMesh.h"
 #include "Engine/StaticMeshActor.h"
 
@@ -26,6 +27,16 @@ UE::AssetUtils::ECreateStaticMeshResult UE::AssetUtils::CreateStaticMeshAsset(
 	FStaticMeshAssetOptions& Options,
 	FStaticMeshResults& ResultsOut)
 {
+	// MINOR HACK: undoing static mesh asset creation causes a bunch of problems because the asset is not deleted,
+	// and is instead left in an invalid state that can cause crashes later. Since we can't undo full asset creation
+	// anyway, we are going to make sure that we don't transact the static mesh asset creation at all. We still want
+	// to allow the calling code to enclose this in another transaction because upstream code may not want to care
+	// about what type of object we are creating, so it shouldn't have to be careful with its transaction placement.
+	// So we temporarily disconnect the current transaction object, if any.
+	ITransaction* UndoState = GUndo;
+	GUndo = nullptr; // Pretend we're not in a transaction
+	ON_SCOPE_EXIT{ GUndo = UndoState; }; // Revert
+
 	FString NewObjectName = FPackageName::GetLongPackageAssetName(Options.NewAssetPath);
 
 	UPackage* UsePackage;
