ModelingMode: avoid a crash on shutdown when DynamicMeshActor is selected, due to TargetCache being modified during iteration

diff --git a/Engine/Plugins/Runtime/MeshModelingToolset/Source/ModelingComponents/Private/Selection/GeometrySelectionManager.cpp b/Engine/Plugins/Runtime/MeshModelingToolset/Source/ModelingComponents/Private/Selection/GeometrySelectionManager.cpp
--- a/Engine/Plugins/Runtime/MeshModelingToolset/Source/ModelingComponents/Private/Selection/GeometrySelectionManager.cpp
+++ b/Engine/Plugins/Runtime/MeshModelingToolset/Source/ModelingComponents/Private/Selection/GeometrySelectionManager.cpp
@@ -565,9 +565,12 @@ TSharedPtr<UGeometrySelectionManager::FGeometrySelectionTarget> UGeometrySelecti
 
 void UGeometrySelectionManager::ResetTargetCache()
 {
-	for (TPair<FGeometryIdentifier, TSharedPtr<FGeometrySelectionTarget>> Pair : TargetCache)
+	// SleepOrShutdownTarget may modify TargetCache
+	TArray<TSharedPtr<FGeometrySelectionTarget>> ToShutdown;
+	TargetCache.GenerateValueArray(ToShutdown);
+	for (TSharedPtr<FGeometrySelectionTarget> Target : ToShutdown)
 	{
-		SleepOrShutdownTarget(Pair.Value, true);
+		SleepOrShutdownTarget(Target, true);
 	}
 	TargetCache.Reset();
 }
