Fix GetElapsedComputeTime so that the preview material kicks in after the proper delay

diff --git a/Engine/Plugins/Runtime/MeshModelingToolset/Source/ModelingComponents/Public/MeshOpPreviewHelpers.h b/Engine/Plugins/Runtime/MeshModelingToolset/Source/ModelingComponents/Public/MeshOpPreviewHelpers.h
--- a/Engine/Plugins/Runtime/MeshModelingToolset/Source/ModelingComponents/Public/MeshOpPreviewHelpers.h
+++ b/Engine/Plugins/Runtime/MeshModelingToolset/Source/ModelingComponents/Public/MeshOpPreviewHelpers.h
@@ -375,11 +375,7 @@ public:
 		check(BackgroundCompute);
 		if (BackgroundCompute)
 		{
-			typename ComputeSourceType::FStatus Status = BackgroundCompute->CheckStatus();
-			if (Status.TaskStatus == UE::Geometry::EBackgroundComputeTaskStatus::InProgress)
-			{
-				return static_cast<float>(Status.ElapsedTime);
-			}
+			return (float)BackgroundCompute->GetElapsedComputeTime();
 		}
 		return 0.f;
 	}
diff --git a/Engine/Plugins/Runtime/MeshModelingToolset/Source/ModelingOperators/Public/BackgroundModelingComputeSource.h b/Engine/Plugins/Runtime/MeshModelingToolset/Source/ModelingOperators/Public/BackgroundModelingComputeSource.h
--- a/Engine/Plugins/Runtime/MeshModelingToolset/Source/ModelingOperators/Public/BackgroundModelingComputeSource.h
+++ b/Engine/Plugins/Runtime/MeshModelingToolset/Source/ModelingOperators/Public/BackgroundModelingComputeSource.h
@@ -188,8 +188,6 @@ public:
 
 	/**
 	 * Return status of the active background computation.
-	 * Warning: Status.ElapsedTime may not be accurate because Tick is not called with the correct DeltaTime argument in
-	 * Development builds when the user is interacting with a slider (as of 13th Dec 2021)
 	 */
 	FStatus CheckStatus() const;
 
@@ -201,7 +199,14 @@ public:
 	/**
 	 * @return duration in seconds of current computation
 	 */
-	double GetElapsedComputeTime() const { return (AccumTime - LastInvalidateTime); }
+	double GetElapsedComputeTime() const {
+		FStatus Status = CheckStatus();
+		if (Status.TaskStatus == EBackgroundComputeTaskStatus::InProgress)
+		{
+			return Status.ElapsedTime;
+		}
+		return 0.;
+	}
 
 public:
 	/** Default wait delay for cancel/restart cycle */
