GeometryScript: Fixed issue where CreateSkeletalMeshAssetFromMesh was not correctly assigning materials.

diff --git a/Engine/Plugins/Runtime/MeshModelingToolset/Source/ModelingComponentsEditorOnly/Private/AssetUtils/CreateSkeletalMeshUtil.cpp b/Engine/Plugins/Runtime/MeshModelingToolset/Source/ModelingComponentsEditorOnly/Private/AssetUtils/CreateSkeletalMeshUtil.cpp
--- a/Engine/Plugins/Runtime/MeshModelingToolset/Source/ModelingComponentsEditorOnly/Private/AssetUtils/CreateSkeletalMeshUtil.cpp
+++ b/Engine/Plugins/Runtime/MeshModelingToolset/Source/ModelingComponentsEditorOnly/Private/AssetUtils/CreateSkeletalMeshUtil.cpp
@@ -95,7 +95,7 @@ UE::AssetUtils::ECreateSkeletalMeshResult UE::AssetUtils::CreateSkeletalMeshAsse
 	}
 
 	// ensure there is at least one material
-	if (Materials.Num() == 0)
+	if (Materials.Num() == 0 && MaterialView.IsEmpty())
 	{
 		Materials.Add(FSkeletalMaterial());
 		MaterialView = Materials;
