MeshModelingTools: Harden tool targets against underlying objects being replaced via blueprints and/or garbage collected.

diff --git a/Engine/Plugins/Experimental/ChaosClothAssetEditor/Source/ChaosClothAssetEditorTools/Private/ChaosClothAsset/ClothComponentToolTarget.cpp b/Engine/Plugins/Experimental/ChaosClothAssetEditor/Source/ChaosClothAssetEditorTools/Private/ChaosClothAsset/ClothComponentToolTarget.cpp
--- a/Engine/Plugins/Experimental/ChaosClothAssetEditor/Source/ChaosClothAssetEditorTools/Private/ChaosClothAsset/ClothComponentToolTarget.cpp
+++ b/Engine/Plugins/Experimental/ChaosClothAssetEditor/Source/ChaosClothAssetEditorTools/Private/ChaosClothAsset/ClothComponentToolTarget.cpp
@@ -54,9 +54,9 @@ bool UClothComponentToolTargetFactory::CanBuildTarget(UObject* SourceObject, con
 UToolTarget* UClothComponentToolTargetFactory::BuildTarget(UObject* SourceObject, const FToolTargetTypeRequirements& Requirements)
 {
 	UClothComponentToolTarget* const Target = NewObject<UClothComponentToolTarget>(); // TODO: Should we set an outer here?
-	Target->Component = Cast<UChaosClothComponent>(SourceObject);
+	Target->InitializeComponent(Cast<UChaosClothComponent>(SourceObject));
 
-	ensure(Target->Component && Requirements.AreSatisfiedBy(Target));
+	ensure(Target->Component.IsValid() && Requirements.AreSatisfiedBy(Target));
 
 	return Target;
 }
diff --git a/Engine/Plugins/Runtime/MeshModelingToolset/Source/ModelingComponentsEditorOnly/Private/ToolTargets/DynamicMeshComponentToolTarget.cpp b/Engine/Plugins/Runtime/MeshModelingToolset/Source/ModelingComponentsEditorOnly/Private/ToolTargets/DynamicMeshComponentToolTarget.cpp
--- a/Engine/Plugins/Runtime/MeshModelingToolset/Source/ModelingComponentsEditorOnly/Private/ToolTargets/DynamicMeshComponentToolTarget.cpp
+++ b/Engine/Plugins/Runtime/MeshModelingToolset/Source/ModelingComponentsEditorOnly/Private/ToolTargets/DynamicMeshComponentToolTarget.cpp
@@ -236,8 +236,8 @@ bool UDynamicMeshComponentToolTargetFactory::CanBuildTarget(UObject* SourceObjec
 UToolTarget* UDynamicMeshComponentToolTargetFactory::BuildTarget(UObject* SourceObject, const FToolTargetTypeRequirements& Requirements)
 {
 	UDynamicMeshComponentToolTarget* Target = NewObject<UDynamicMeshComponentToolTarget>();// TODO: Should we set an outer here?
-	Target->Component = Cast<UDynamicMeshComponent>(SourceObject);
-	check(Target->Component && Requirements.AreSatisfiedBy(Target));
+	Target->InitializeComponent(Cast<UDynamicMeshComponent>(SourceObject));
+	checkSlow(Target->Component.IsValid() && Requirements.AreSatisfiedBy(Target));
 
 	return Target;
 }
diff --git a/Engine/Plugins/Runtime/MeshModelingToolset/Source/ModelingComponentsEditorOnly/Private/ToolTargets/SkeletalMeshComponentToolTarget.cpp b/Engine/Plugins/Runtime/MeshModelingToolset/Source/ModelingComponentsEditorOnly/Private/ToolTargets/SkeletalMeshComponentToolTarget.cpp
--- a/Engine/Plugins/Runtime/MeshModelingToolset/Source/ModelingComponentsEditorOnly/Private/ToolTargets/SkeletalMeshComponentToolTarget.cpp
+++ b/Engine/Plugins/Runtime/MeshModelingToolset/Source/ModelingComponentsEditorOnly/Private/ToolTargets/SkeletalMeshComponentToolTarget.cpp
@@ -38,7 +38,7 @@ bool USkeletalMeshComponentReadOnlyToolTarget::IsValid() const
 	}
 	const USkeletalMesh* SkeletalMesh = SkeletalMeshComponent->GetSkeletalMeshAsset();
 
-	return USkeletalMeshReadOnlyToolTarget::IsValid(SkeletalMesh);
+	return SkeletalMesh != nullptr;
 }
 
 int32 USkeletalMeshComponentReadOnlyToolTarget::GetNumMaterials() const
@@ -80,7 +80,7 @@ bool USkeletalMeshComponentReadOnlyToolTarget::CommitMaterialSetUpdate(const FCo
 		USkeletalMesh* SkeletalMesh = Cast<USkeletalMeshComponent>(Component)->GetSkeletalMeshAsset();
 
 		// unregister the component while we update it's static mesh
-		TUniquePtr<FComponentReregisterContext> ComponentReregisterContext = MakeUnique<FComponentReregisterContext>(Component);
+		TUniquePtr<FComponentReregisterContext> ComponentReregisterContext = MakeUnique<FComponentReregisterContext>(Component.Get());
 
 		return USkeletalMeshToolTarget::CommitMaterialSetUpdate(SkeletalMesh, MaterialSet, bApplyToAsset);
 	}
@@ -163,7 +163,7 @@ void USkeletalMeshComponentToolTarget::CommitMeshDescription(const FCommitter& C
 	}
 
 	// unregister the component while we update its skeletal mesh
-	FComponentReregisterContext ComponentReregisterContext(Component);
+	FComponentReregisterContext ComponentReregisterContext(Component.Get());
 
 	USkeletalMeshToolTarget::CommitMeshDescription(SkeletalMesh, CachedMeshDescription.Get(), Committer);
 
@@ -205,8 +205,8 @@ bool USkeletalMeshComponentReadOnlyToolTargetFactory::CanBuildTarget(UObject* So
 UToolTarget* USkeletalMeshComponentReadOnlyToolTargetFactory::BuildTarget(UObject* SourceObject, const FToolTargetTypeRequirements& Requirements)
 {
 	USkeletalMeshComponentReadOnlyToolTarget* Target = NewObject<USkeletalMeshComponentReadOnlyToolTarget>();
-	Target->Component = Cast<USkeletalMeshComponent>(SourceObject);
-	check(Target->Component && Requirements.AreSatisfiedBy(Target));
+	Target->InitializeComponent(Cast<USkeletalMeshComponent>(SourceObject));
+	checkSlow(Target->Component.IsValid() && Requirements.AreSatisfiedBy(Target));
 
 	return Target;
 }
@@ -229,8 +229,8 @@ bool USkeletalMeshComponentToolTargetFactory::CanBuildTarget(UObject* SourceObje
 UToolTarget* USkeletalMeshComponentToolTargetFactory::BuildTarget(UObject* SourceObject, const FToolTargetTypeRequirements& Requirements)
 {
 	USkeletalMeshComponentToolTarget* Target = NewObject<USkeletalMeshComponentToolTarget>();
-	Target->Component = Cast<USkeletalMeshComponent>(SourceObject);
-	check(Target->Component && Requirements.AreSatisfiedBy(Target));
+	Target->InitializeComponent(Cast<USkeletalMeshComponent>(SourceObject));
+	checkSlow(Target->Component.IsValid() && Requirements.AreSatisfiedBy(Target));
 
 	return Target;
 }
diff --git a/Engine/Plugins/Runtime/MeshModelingToolset/Source/ModelingComponentsEditorOnly/Private/ToolTargets/SkeletalMeshToolTarget.cpp b/Engine/Plugins/Runtime/MeshModelingToolset/Source/ModelingComponentsEditorOnly/Private/ToolTargets/SkeletalMeshToolTarget.cpp
--- a/Engine/Plugins/Runtime/MeshModelingToolset/Source/ModelingComponentsEditorOnly/Private/ToolTargets/SkeletalMeshToolTarget.cpp
+++ b/Engine/Plugins/Runtime/MeshModelingToolset/Source/ModelingComponentsEditorOnly/Private/ToolTargets/SkeletalMeshToolTarget.cpp
@@ -30,7 +30,7 @@ namespace USkeletalMeshToolTargetLocals
 
 bool USkeletalMeshReadOnlyToolTarget::IsValid() const
 {
-	return IsValid(SkeletalMesh);
+	return SkeletalMesh.IsValid();
 }
 
 bool USkeletalMeshReadOnlyToolTarget::IsValid(const USkeletalMesh* SkeletalMeshIn)
@@ -57,7 +57,7 @@ UMaterialInterface* USkeletalMeshReadOnlyToolTarget::GetMaterial(int32 MaterialI
 void USkeletalMeshReadOnlyToolTarget::GetMaterialSet(FComponentMaterialSet& MaterialSetOut, bool bPreferAssetMaterials) const
 {
 	if (!ensure(IsValid())) return;
-	GetMaterialSet(SkeletalMesh, MaterialSetOut, bPreferAssetMaterials);
+	GetMaterialSet(SkeletalMesh.Get(), MaterialSetOut, bPreferAssetMaterials);
 }
 
 void USkeletalMeshReadOnlyToolTarget::GetMaterialSet(const USkeletalMesh* SkeletalMeshIn, FComponentMaterialSet& MaterialSetOut,
@@ -74,7 +74,7 @@ void USkeletalMeshReadOnlyToolTarget::GetMaterialSet(const USkeletalMesh* Skelet
 bool USkeletalMeshReadOnlyToolTarget::CommitMaterialSetUpdate(const FComponentMaterialSet& MaterialSet, bool bApplyToAsset)
 {
 	if (!ensure(IsValid())) return false;
-	return CommitMaterialSetUpdate(SkeletalMesh, MaterialSet, bApplyToAsset);
+	return CommitMaterialSetUpdate(SkeletalMesh.Get(), MaterialSet, bApplyToAsset);
 }
 
 bool USkeletalMeshReadOnlyToolTarget::CommitMaterialSetUpdate(USkeletalMesh* SkeletalMeshIn, 
@@ -133,7 +133,7 @@ const FMeshDescription* USkeletalMeshReadOnlyToolTarget::GetMeshDescription(cons
 	if (!CachedMeshDescription.IsValid())
 	{
 		CachedMeshDescription = MakeUnique<FMeshDescription>();
-		GetMeshDescription(SkeletalMesh, *CachedMeshDescription);
+		GetMeshDescription(SkeletalMesh.Get(), *CachedMeshDescription);
 	}
 
 	return CachedMeshDescription.Get();
@@ -162,7 +162,7 @@ FDynamicMesh3 USkeletalMeshReadOnlyToolTarget::GetDynamicMesh()
 
 USkeletalMesh* USkeletalMeshReadOnlyToolTarget::GetSkeletalMesh() const
 {
-	return IsValid() ? SkeletalMesh : nullptr;
+	return IsValid() ? SkeletalMesh.Get() : nullptr;
 }
 
 //
@@ -176,9 +176,9 @@ void USkeletalMeshToolTarget::CommitMeshDescription(const FCommitter& Committer,
 	if (!CachedMeshDescription.IsValid())
 	{
 		CachedMeshDescription = MakeUnique<FMeshDescription>();
-		GetMeshDescription(SkeletalMesh, *CachedMeshDescription);
+		GetMeshDescription(SkeletalMesh.Get(), *CachedMeshDescription);
 	}
-	CommitMeshDescription(SkeletalMesh, CachedMeshDescription.Get(), Committer);
+	CommitMeshDescription(SkeletalMesh.Get(), CachedMeshDescription.Get(), Committer);
 }
 
 void USkeletalMeshToolTarget::CommitMeshDescription(USkeletalMesh* SkeletalMeshIn,
@@ -238,7 +238,7 @@ UToolTarget* USkeletalMeshReadOnlyToolTargetFactory::BuildTarget(UObject* Source
 {
 	USkeletalMeshReadOnlyToolTarget* Target = NewObject<USkeletalMeshReadOnlyToolTarget>();
 	Target->SkeletalMesh = Cast<USkeletalMesh>(SourceObject);
-	check(Target->SkeletalMesh && Requirements.AreSatisfiedBy(Target));
+	checkSlow(Target->SkeletalMesh.IsValid() && Requirements.AreSatisfiedBy(Target));
 
 	return Target;
 }
@@ -263,7 +263,7 @@ UToolTarget* USkeletalMeshToolTargetFactory::BuildTarget(UObject* SourceObject,
 {
 	USkeletalMeshToolTarget* Target = NewObject<USkeletalMeshToolTarget>();
 	Target->SkeletalMesh = Cast<USkeletalMesh>(SourceObject);
-	check(Target->SkeletalMesh && Requirements.AreSatisfiedBy(Target));
+	checkSlow(Target->SkeletalMesh.IsValid() && Requirements.AreSatisfiedBy(Target));
 
 	return Target;
 }
diff --git a/Engine/Plugins/Runtime/MeshModelingToolset/Source/ModelingComponentsEditorOnly/Private/ToolTargets/StaticMeshComponentToolTarget.cpp b/Engine/Plugins/Runtime/MeshModelingToolset/Source/ModelingComponentsEditorOnly/Private/ToolTargets/StaticMeshComponentToolTarget.cpp
--- a/Engine/Plugins/Runtime/MeshModelingToolset/Source/ModelingComponentsEditorOnly/Private/ToolTargets/StaticMeshComponentToolTarget.cpp
+++ b/Engine/Plugins/Runtime/MeshModelingToolset/Source/ModelingComponentsEditorOnly/Private/ToolTargets/StaticMeshComponentToolTarget.cpp
@@ -46,7 +46,7 @@ bool UStaticMeshComponentToolTarget::IsValid() const
 	}
 	UStaticMesh* StaticMesh = StaticMeshComponent->GetStaticMesh();
 	
-	return UStaticMeshToolTarget::IsValid(StaticMesh, EditingLOD);
+	return UStaticMeshToolTarget::HasNonGeneratedLOD(StaticMesh, EditingLOD);
 }
 
 
@@ -90,7 +90,7 @@ bool UStaticMeshComponentToolTarget::CommitMaterialSetUpdate(const FComponentMat
 		UStaticMesh* StaticMesh = Cast<UStaticMeshComponent>(Component)->GetStaticMesh();
 
 		// unregister the component while we update it's static mesh
-		TUniquePtr<FComponentReregisterContext> ComponentReregisterContext = MakeUnique<FComponentReregisterContext>(Component);
+		TUniquePtr<FComponentReregisterContext> ComponentReregisterContext = MakeUnique<FComponentReregisterContext>(Component.Get());
 
 		return UStaticMeshToolTarget::CommitMaterialSetUpdate(StaticMesh, MaterialSet, bApplyToAsset);
 	}
@@ -233,7 +233,7 @@ void UStaticMeshComponentToolTarget::CommitMeshDescription(const FCommitter& Com
 	EMeshLODIdentifier WriteToLOD = (CommitParams.bHaveTargetLOD && CommitParams.TargetLOD != EMeshLODIdentifier::Default) ? CommitParams.TargetLOD : EditingLOD;
 
 	// unregister the component while we update its static mesh
-	FComponentReregisterContext ComponentReregisterContext(Component);
+	FComponentReregisterContext ComponentReregisterContext(Component.Get());
 
 	UStaticMeshToolTarget::CommitMeshDescription(StaticMesh, Committer, WriteToLOD);
 
@@ -299,9 +299,9 @@ bool UStaticMeshComponentToolTargetFactory::CanBuildTarget(UObject* SourceObject
 UToolTarget* UStaticMeshComponentToolTargetFactory::BuildTarget(UObject* SourceObject, const FToolTargetTypeRequirements& Requirements)
 {
 	UStaticMeshComponentToolTarget* Target = NewObject<UStaticMeshComponentToolTarget>();// TODO: Should we set an outer here?
-	Target->Component = Cast<UStaticMeshComponent>(SourceObject);
+	Target->InitializeComponent(Cast<UStaticMeshComponent>(SourceObject));
 	Target->SetEditingLOD(EditingLOD);
-	check(Target->Component && Requirements.AreSatisfiedBy(Target));
+	checkSlow(Target->Component.Get() && Requirements.AreSatisfiedBy(Target));
 
 	return Target;
 }
diff --git a/Engine/Plugins/Runtime/MeshModelingToolset/Source/ModelingComponentsEditorOnly/Private/ToolTargets/StaticMeshToolTarget.cpp b/Engine/Plugins/Runtime/MeshModelingToolset/Source/ModelingComponentsEditorOnly/Private/ToolTargets/StaticMeshToolTarget.cpp
--- a/Engine/Plugins/Runtime/MeshModelingToolset/Source/ModelingComponentsEditorOnly/Private/ToolTargets/StaticMeshToolTarget.cpp
+++ b/Engine/Plugins/Runtime/MeshModelingToolset/Source/ModelingComponentsEditorOnly/Private/ToolTargets/StaticMeshToolTarget.cpp
@@ -34,7 +34,7 @@ namespace StaticMeshToolTargetLocals
 
 void UStaticMeshToolTarget::SetEditingLOD(EMeshLODIdentifier RequestedEditingLOD)
 {
-	EditingLOD = GetValidEditingLOD(StaticMesh, RequestedEditingLOD);
+	EditingLOD = GetValidEditingLOD(StaticMesh.Get(), RequestedEditingLOD);
 }
 
 EMeshLODIdentifier UStaticMeshToolTarget::GetValidEditingLOD(const UStaticMesh* StaticMeshIn, 
@@ -75,7 +75,7 @@ EMeshLODIdentifier UStaticMeshToolTarget::GetValidEditingLOD(const UStaticMesh*
 
 TArray<EMeshLODIdentifier> UStaticMeshToolTarget::GetAvailableLODs(bool bSkipAutoGenerated) const
 {
-	return GetAvailableLODs(StaticMesh, bSkipAutoGenerated);
+	return GetAvailableLODs(StaticMesh.Get(), bSkipAutoGenerated);
 }
 
 TArray<EMeshLODIdentifier> UStaticMeshToolTarget::GetAvailableLODs(const UStaticMesh* InStaticMesh, bool bSkipAutoGenerated)
@@ -111,7 +111,9 @@ TArray<EMeshLODIdentifier> UStaticMeshToolTarget::GetAvailableLODs(const UStatic
 
 bool UStaticMeshToolTarget::IsValid() const
 {
-	return IsValid(StaticMesh, EditingLOD);
+	// The StaticMesh.Get() call here will give a nullptr if object is not alive, which 
+	// HasNonGeneratedLOD checks for.
+	return HasNonGeneratedLOD(StaticMesh.Get(), EditingLOD);
 }
 
 bool UStaticMeshToolTarget::IsValid(const UStaticMesh* StaticMeshIn, EMeshLODIdentifier EditingLODIn)
@@ -184,7 +186,7 @@ void UStaticMeshToolTarget::GetMaterialSet(FComponentMaterialSet& MaterialSetOut
 {
 	if (!ensure(IsValid())) return;
 
-	GetMaterialSet(StaticMesh, MaterialSetOut, bPreferAssetMaterials);
+	GetMaterialSet(StaticMesh.Get(), MaterialSetOut, bPreferAssetMaterials);
 }
 
 void UStaticMeshToolTarget::GetMaterialSet(const UStaticMesh* StaticMeshIn, 
@@ -202,7 +204,7 @@ bool UStaticMeshToolTarget::CommitMaterialSetUpdate(const FComponentMaterialSet&
 {
 	if (!ensure(IsValid())) return false;
 
-	return CommitMaterialSetUpdate(StaticMesh, MaterialSet, bApplyToAsset);
+	return CommitMaterialSetUpdate(StaticMesh.Get(), MaterialSet, bApplyToAsset);
 }
 
 bool UStaticMeshToolTarget::CommitMaterialSetUpdate(UStaticMesh* StaticMeshIn,
@@ -272,7 +274,7 @@ const FMeshDescription* UStaticMeshToolTarget::GetMeshDescription(const FGetMesh
 		EMeshLODIdentifier UseLOD = EditingLOD;
 		if (GetMeshParams.bHaveRequestLOD)
 		{
-			UseLOD = UStaticMeshToolTarget::GetValidEditingLOD(StaticMesh, GetMeshParams.RequestLOD);
+			UseLOD = UStaticMeshToolTarget::GetValidEditingLOD(StaticMesh.Get(), GetMeshParams.RequestLOD);
 			ensure(UseLOD == GetMeshParams.RequestLOD);		// probably a bug somewhere if this is not true
 		}
 
@@ -298,7 +300,7 @@ void UStaticMeshToolTarget::CommitMeshDescription(const FCommitter& Committer, c
 
 	EMeshLODIdentifier WriteToLOD = (CommitParams.bHaveTargetLOD && CommitParams.TargetLOD != EMeshLODIdentifier::Default) ? CommitParams.TargetLOD : EditingLOD;
 
-	CommitMeshDescription(StaticMesh, Committer, WriteToLOD);
+	CommitMeshDescription(StaticMesh.Get(), Committer, WriteToLOD);
 }
 
 void UStaticMeshToolTarget::CommitMeshDescription(UStaticMesh* StaticMeshIn, const FCommitter& Committer, EMeshLODIdentifier EditingLODIn)
@@ -400,7 +402,7 @@ void UStaticMeshToolTarget::CommitDynamicMesh(const FDynamicMesh3& Mesh, const F
 
 UStaticMesh* UStaticMeshToolTarget::GetStaticMesh() const
 {
-	return IsValid() ? StaticMesh : nullptr;
+	return IsValid() ? StaticMesh.Get() : nullptr;
 }
 
 
@@ -420,7 +422,7 @@ UToolTarget* UStaticMeshToolTargetFactory::BuildTarget(UObject* SourceObject, co
 	UStaticMeshToolTarget* Target = NewObject<UStaticMeshToolTarget>();// TODO: Should we set an outer here?
 	Target->StaticMesh = Cast<UStaticMesh>(SourceObject);
 	Target->SetEditingLOD(EditingLOD);
-	check(Target->StaticMesh && Requirements.AreSatisfiedBy(Target));
+	checkSlow(Target->StaticMesh.IsValid() && Requirements.AreSatisfiedBy(Target));
 
 	return Target;
 }
diff --git a/Engine/Plugins/Runtime/MeshModelingToolset/Source/ModelingComponentsEditorOnly/Private/ToolTargets/VolumeComponentToolTarget.cpp b/Engine/Plugins/Runtime/MeshModelingToolset/Source/ModelingComponentsEditorOnly/Private/ToolTargets/VolumeComponentToolTarget.cpp
--- a/Engine/Plugins/Runtime/MeshModelingToolset/Source/ModelingComponentsEditorOnly/Private/ToolTargets/VolumeComponentToolTarget.cpp
+++ b/Engine/Plugins/Runtime/MeshModelingToolset/Source/ModelingComponentsEditorOnly/Private/ToolTargets/VolumeComponentToolTarget.cpp
@@ -187,9 +187,9 @@ bool UVolumeComponentToolTargetFactory::CanBuildTarget(UObject* SourceObject, co
 UToolTarget* UVolumeComponentToolTargetFactory::BuildTarget(UObject* SourceObject, const FToolTargetTypeRequirements& Requirements)
 {
 	UVolumeComponentToolTarget* Target = NewObject<UVolumeComponentToolTarget>();
-	Target->Component = Cast<UBrushComponent>(SourceObject);
+	Target->InitializeComponent(Cast<UBrushComponent>(SourceObject));
 
-	check(Target->Component && Requirements.AreSatisfiedBy(Target)
+	checkSlow(Target->Component.IsValid() && Requirements.AreSatisfiedBy(Target)
 		&& Cast<AVolume>(Cast<UBrushComponent>(SourceObject)->GetOwner()));
 
 	return Target;
diff --git a/Engine/Plugins/Runtime/MeshModelingToolset/Source/ModelingComponentsEditorOnly/Public/ToolTargets/SkeletalMeshToolTarget.h b/Engine/Plugins/Runtime/MeshModelingToolset/Source/ModelingComponentsEditorOnly/Public/ToolTargets/SkeletalMeshToolTarget.h
--- a/Engine/Plugins/Runtime/MeshModelingToolset/Source/ModelingComponentsEditorOnly/Public/ToolTargets/SkeletalMeshToolTarget.h
+++ b/Engine/Plugins/Runtime/MeshModelingToolset/Source/ModelingComponentsEditorOnly/Public/ToolTargets/SkeletalMeshToolTarget.h
@@ -54,13 +54,14 @@ public:
 	virtual USkeletalMesh* GetSkeletalMesh() const override;
 
 protected:
-	USkeletalMesh* SkeletalMesh = nullptr;
+	TWeakObjectPtr<USkeletalMesh> SkeletalMesh = nullptr;
 
 	// So that the tool target factory can poke into Component.
 	friend class USkeletalMeshReadOnlyToolTargetFactory;
 	friend class USkeletalMeshComponentReadOnlyToolTarget;
 	friend class USkeletalMeshComponentToolTarget;
 
+	UE_DEPRECATED(5.3, "This function doesn't check anything skeletal mesh specific, only (badly) checks to see if the object is alive")
 	static bool IsValid(const USkeletalMesh* USkeletalMesh);
 
 	static void GetMeshDescription(const USkeletalMesh* SkeletalMesh, FMeshDescription& MeshDescriptionOut);
diff --git a/Engine/Plugins/Runtime/MeshModelingToolset/Source/ModelingComponentsEditorOnly/Public/ToolTargets/StaticMeshToolTarget.h b/Engine/Plugins/Runtime/MeshModelingToolset/Source/ModelingComponentsEditorOnly/Public/ToolTargets/StaticMeshToolTarget.h
--- a/Engine/Plugins/Runtime/MeshModelingToolset/Source/ModelingComponentsEditorOnly/Public/ToolTargets/StaticMeshToolTarget.h
+++ b/Engine/Plugins/Runtime/MeshModelingToolset/Source/ModelingComponentsEditorOnly/Public/ToolTargets/StaticMeshToolTarget.h
@@ -78,7 +78,7 @@ public:
 	// Rest provided by parent class
 
 protected:
-	UStaticMesh* StaticMesh = nullptr;
+	TWeakObjectPtr<UStaticMesh> StaticMesh = nullptr;
 
 	EMeshLODIdentifier EditingLOD = EMeshLODIdentifier::LOD0;
 
@@ -88,7 +88,11 @@ protected:
 
 	friend class UStaticMeshComponentToolTargetFactory;
 
+	UE_DEPRECATED(5.3, "Compared to using a TWeakObjectPtr, this function performs both redundant and insufficient "
+		"checks of whether the UStaticMesh is alive, and the LOD check does not guard against a generated LOD (which "
+		"tools should not edit) as HasNonGeneratedLOD does.")
 	static bool IsValid(const UStaticMesh* StaticMesh, EMeshLODIdentifier EditingLOD);
+
 	// returns true if the mesh has a non-generated LOD of the requested level.  Interprets EMeshLODIdentifier::Default as LOD0
 	static bool HasNonGeneratedLOD(const UStaticMesh* StaticMesh, EMeshLODIdentifier MeshLOD);
 	static EMeshLODIdentifier GetValidEditingLOD(const UStaticMesh* StaticMesh, 
diff --git a/Engine/Source/Runtime/InteractiveToolsFramework/Private/ToolTargets/PrimitiveComponentToolTarget.cpp b/Engine/Source/Runtime/InteractiveToolsFramework/Private/ToolTargets/PrimitiveComponentToolTarget.cpp
--- a/Engine/Source/Runtime/InteractiveToolsFramework/Private/ToolTargets/PrimitiveComponentToolTarget.cpp
+++ b/Engine/Source/Runtime/InteractiveToolsFramework/Private/ToolTargets/PrimitiveComponentToolTarget.cpp
@@ -8,12 +8,14 @@
 
 bool UPrimitiveComponentToolTarget::IsValid() const
 {
-	return Component && ::IsValid(Component) && !Component->IsUnreachable() && Component->IsValidLowLevel();
+	return Component.IsValid();
 }
 
 UPrimitiveComponent* UPrimitiveComponentToolTarget::GetOwnerComponent() const
 {
-	return IsValid() ? Component : nullptr;
+	// Note that we don't just return Component.Get() because we want to call the virtual IsValid
+	// for derived classes.
+	return IsValid() ? Component.Get() : nullptr;
 }
 
 AActor* UPrimitiveComponentToolTarget::GetOwnerActor() const
@@ -44,6 +46,46 @@ bool UPrimitiveComponentToolTarget::HitTestComponent(const FRay& WorldRay, FHitR
 	return false;
 }
 
+void UPrimitiveComponentToolTarget::InitializeComponent(UPrimitiveComponent* ComponentIn)
+{
+	Component = ComponentIn;
+
+	if (ensure(Component.IsValid()))
+	{
+#if WITH_EDITOR
+		// See comment in OnObjectsReplaced
+		FCoreUObjectDelegates::OnObjectsReplaced.AddUObject(this, &UPrimitiveComponentToolTarget::OnObjectsReplaced);
+#endif
+	}
+}
+
+void UPrimitiveComponentToolTarget::BeginDestroy()
+{
+#if WITH_EDITOR
+	FCoreUObjectDelegates::OnObjectsReplaced.RemoveAll(this);
+#endif
+
+	Super::BeginDestroy();
+}
+
+#if WITH_EDITOR
+void UPrimitiveComponentToolTarget::OnObjectsReplaced(const FCoreUObjectDelegates::FReplacementObjectMap& Map)
+{
+	// Components frequently get destroyed and recreated when they are part of blueprint actors that 
+	// get modified. For the most part, we don't need to worry about supporting these cases, but keeping
+	// a consistent reference here allows us to avoid getting into some bad states. For instance, we often
+	// hide the source component and unhide at tool end, and if we lose the reference to the component
+	// while the tool is running, we are unable to unhide it later. The user is unlikely to understand 
+	// why their object disappeared in that case or know to fix it via the component visibility property.
+
+	UObject* const * MappedObject = Map.Find(Component.Get());
+	if (MappedObject)
+	{
+		Component = Cast<UPrimitiveComponent>(*MappedObject);
+	}
+}
+#endif
+
 
 // Factory
 
@@ -60,8 +102,8 @@ bool UPrimitiveComponentToolTargetFactory::CanBuildTarget(UObject* SourceObject,
 UToolTarget* UPrimitiveComponentToolTargetFactory::BuildTarget(UObject* SourceObject, const FToolTargetTypeRequirements& Requirements)
 {
 	UPrimitiveComponentToolTarget* Target = NewObject<UPrimitiveComponentToolTarget>();
-	Target->Component = Cast<UPrimitiveComponent>(SourceObject);
-	check(Target->Component && Requirements.AreSatisfiedBy(Target));
+	Target->InitializeComponent(Cast<UPrimitiveComponent>(SourceObject));
+	checkSlow(Target->Component.IsValid() && Requirements.AreSatisfiedBy(Target));
 	return Target;
 }
 
diff --git a/Engine/Source/Runtime/InteractiveToolsFramework/Public/ToolTargets/PrimitiveComponentToolTarget.h b/Engine/Source/Runtime/InteractiveToolsFramework/Public/ToolTargets/PrimitiveComponentToolTarget.h
--- a/Engine/Source/Runtime/InteractiveToolsFramework/Public/ToolTargets/PrimitiveComponentToolTarget.h
+++ b/Engine/Source/Runtime/InteractiveToolsFramework/Public/ToolTargets/PrimitiveComponentToolTarget.h
@@ -5,7 +5,6 @@
 #include "CoreMinimal.h"
 
 #include "TargetInterfaces/PrimitiveComponentBackedTarget.h"
-#include "Engine/HitResult.h"
 #include "ToolTargets/ToolTarget.h"
 
 #include "PrimitiveComponentToolTarget.generated.h"
@@ -19,6 +18,8 @@ class INTERACTIVETOOLSFRAMEWORK_API UPrimitiveComponentToolTarget : public UTool
 {
 	GENERATED_BODY()
 public:
+
+	// UToolTarget
 	virtual bool IsValid() const override;
 
 	// IPrimitiveComponentBackedTarget implementation
@@ -28,10 +29,20 @@ public:
 	virtual FTransform GetWorldTransform() const override;
 	virtual bool HitTestComponent(const FRay& WorldRay, FHitResult& OutHit) const override;
 
+	// UObject
+	void BeginDestroy() override;
+
 protected:
 	friend class UPrimitiveComponentToolTargetFactory;
 	
-	UPrimitiveComponent* Component;
+	void InitializeComponent(UPrimitiveComponent* ComponentIn);
+
+	TWeakObjectPtr<UPrimitiveComponent> Component;
+
+private:
+#if WITH_EDITOR
+	void OnObjectsReplaced(const FCoreUObjectDelegates::FReplacementObjectMap& Map);
+#endif
 };
 
 /**
