Make a larger UI range for height offset in revolve properties

diff --git a/Engine/Plugins/Runtime/MeshModelingToolset/Source/MeshModelingTools/Public/Properties/RevolveProperties.h b/Engine/Plugins/Runtime/MeshModelingToolset/Source/MeshModelingTools/Public/Properties/RevolveProperties.h
--- a/Engine/Plugins/Runtime/MeshModelingToolset/Source/MeshModelingTools/Public/Properties/RevolveProperties.h
+++ b/Engine/Plugins/Runtime/MeshModelingToolset/Source/MeshModelingTools/Public/Properties/RevolveProperties.h
@@ -109,7 +109,7 @@ public:
 	int NumExplicitSteps = 24;
 
 	/** How far to move each step along the revolution axis per degree. Non-zero values are useful for creating spirals. */
-	UPROPERTY(EditAnywhere, Category = Revolve, meta = (DisplayName = "Height Offset", UIMin = "-1", UIMax = "1", ClampMin = "-100000", ClampMax = "100000"))
+	UPROPERTY(EditAnywhere, Category = Revolve, meta = (DisplayName = "Height Offset", UIMin = "-10", UIMax = "10", ClampMin = "-100000", ClampMax = "100000"))
 	double HeightOffsetPerDegree = 0;
 
 	/** By default, revolution is done counterclockwise if looking down the revolution axis. This reverses the revolution direction to clockwise.*/
