Fixing compile error from 25482982

diff --git a/Engine/Plugins/Runtime/MeshModelingToolset/Source/ModelingComponents/Private/Selection/MeshTopologySelectionMechanic.cpp b/Engine/Plugins/Runtime/MeshModelingToolset/Source/ModelingComponents/Private/Selection/MeshTopologySelectionMechanic.cpp
--- a/Engine/Plugins/Runtime/MeshModelingToolset/Source/ModelingComponents/Private/Selection/MeshTopologySelectionMechanic.cpp
+++ b/Engine/Plugins/Runtime/MeshModelingToolset/Source/ModelingComponents/Private/Selection/MeshTopologySelectionMechanic.cpp
@@ -166,19 +166,22 @@ void UMeshTopologySelectionMechanic::Render(IToolsContextRenderAPI* RenderAPI)
 	FViewCameraState RenderCameraState = RenderAPI->GetCameraState();
 
 	const FDynamicMesh3* TargetMesh = this->Mesh;
-	FTransform Transform = (FTransform)TargetTransform;
-
+	const FTransform Transform = static_cast<FTransform>(TargetTransform);
+	
 	const FTopologyProvider* TopologyProvider = TopoSelector->GetTopologyProvider();
 
 	PolyEdgesRenderer.BeginFrame(RenderAPI, RenderCameraState);
 	PolyEdgesRenderer.SetTransform(Transform);
-	for (int EdgeID = 0; EdgeID < TopologyProvider->GetNumEdges(); ++EdgeID)
+	if (bShowEdges)
 	{
-		FVector3d A, B;
-		for (int eid : TopologyProvider->GetGroupEdgeEdges(EdgeID))
+		for (int EdgeID = 0; EdgeID < TopologyProvider->GetNumEdges(); ++EdgeID)
 		{
-			TargetMesh->GetEdgeV(eid, A, B);
-			PolyEdgesRenderer.DrawLine(A, B);
+			FVector3d A, B;
+			for (int eid : TopologyProvider->GetGroupEdgeEdges(EdgeID))
+			{
+				TargetMesh->GetEdgeV(eid, A, B);
+				PolyEdgesRenderer.DrawLine(A, B);
+			}
 		}
 	}
 	if (bShowSelectableCorners)
diff --git a/Engine/Plugins/Runtime/MeshModelingToolset/Source/ModelingComponents/Public/Selection/MeshTopologySelectionMechanic.h b/Engine/Plugins/Runtime/MeshModelingToolset/Source/ModelingComponents/Public/Selection/MeshTopologySelectionMechanic.h
--- a/Engine/Plugins/Runtime/MeshModelingToolset/Source/ModelingComponents/Public/Selection/MeshTopologySelectionMechanic.h
+++ b/Engine/Plugins/Runtime/MeshModelingToolset/Source/ModelingComponents/Public/Selection/MeshTopologySelectionMechanic.h
@@ -343,6 +343,11 @@ public:
 	* Render only the MarqueeMechanic, without rendering the current selection
 	*/
 	void RenderMarquee(IToolsContextRenderAPI* RenderAPI);
+
+	/**
+	* Toggle rendering of edges
+	*/
+	void SetShowEdges(const bool bRenderEdges) { bShowEdges = bRenderEdges; };
 	
 
 	// TODO: Is it worth issuing separate callbacks in normal selection changes and in FChange ones, to
@@ -433,6 +438,7 @@ protected:
 	static const int32 CtrlModifierID = 2;
 
 	bool bShowSelectableCorners = true;
+	bool bShowEdges = true;
 public:
 	FToolDataVisualizer PolyEdgesRenderer;
 	FToolDataVisualizer HilightRenderer;
