ModelingTools: Add IMeshDescriptionProvider::GetMeshDescriptionLOD() which returns current LOD for a ToolTarget. Implement in StaticMesh Tool Targets. Add UE::ToolTarget::GetTargetMeshDescriptionLOD() which wraps this function

diff --git a/Engine/Plugins/Runtime/MeshModelingToolset/Source/ModelingComponents/Private/ModelingToolTargetUtil.cpp b/Engine/Plugins/Runtime/MeshModelingToolset/Source/ModelingComponents/Private/ModelingToolTargetUtil.cpp
--- a/Engine/Plugins/Runtime/MeshModelingToolset/Source/ModelingComponents/Private/ModelingToolTargetUtil.cpp
+++ b/Engine/Plugins/Runtime/MeshModelingToolset/Source/ModelingComponents/Private/ModelingToolTargetUtil.cpp
@@ -162,6 +162,21 @@ TArray<EMeshLODIdentifier> UE::ToolTarget::GetMeshDescriptionLODs(UToolTarget* T
 	return DefaultArray;
 }
 
+
+EMeshLODIdentifier UE::ToolTarget::GetTargetMeshDescriptionLOD(
+	UToolTarget* Target,
+	bool& bOutTargetSupportsLODs)
+{
+	if (IMeshDescriptionProvider* MeshDescriptionProvider = Cast<IMeshDescriptionProvider>(Target))
+	{
+		bOutTargetSupportsLODs = MeshDescriptionProvider->SupportsLODs();
+		return MeshDescriptionProvider->GetMeshDescriptionLOD();
+	}
+
+	bOutTargetSupportsLODs = false;
+	return EMeshLODIdentifier::Default;
+}
+
 const FMeshDescription* UE::ToolTarget::GetMeshDescription(UToolTarget* Target, const FGetMeshParameters& GetMeshParams)
 {
 	static FMeshDescription EmptyMeshDescription;
diff --git a/Engine/Plugins/Runtime/MeshModelingToolset/Source/ModelingComponents/Public/ModelingToolTargetUtil.h b/Engine/Plugins/Runtime/MeshModelingToolset/Source/ModelingComponents/Public/ModelingToolTargetUtil.h
--- a/Engine/Plugins/Runtime/MeshModelingToolset/Source/ModelingComponents/Public/ModelingToolTargetUtil.h
+++ b/Engine/Plugins/Runtime/MeshModelingToolset/Source/ModelingComponents/Public/ModelingToolTargetUtil.h
@@ -103,6 +103,16 @@ MODELINGCOMPONENTS_API TArray<EMeshLODIdentifier> GetMeshDescriptionLODs(
 	bool bOnlyReturnDefaultLOD = false,
 	bool bExcludeAutoGeneratedLODs = true);
 
+/**
+ * Return identifier for the mesh LOD that a ToolTarget represents, if defined for the given ToolTarget
+ * @param Target					The tool target to query for LODs
+ * @param bOutTargetSupportsLODs	Will be set to 'false' if the target does not support querying MeshDescriptions with LODs
+ * @return LODIdentifier for current LOD that would be provided/edited by the ToolTarget, or the Default LODIdentifier if the target doesn't support LODs or bOnlyReturnDefaultLOD is true
+ */
+MODELINGCOMPONENTS_API EMeshLODIdentifier GetTargetMeshDescriptionLOD(
+	UToolTarget* Target,
+	bool& bOutTargetSupportsLODs);
+
 
 /**
  * @return the MeshDescription underlying a ToolTarget, if it has such a mesh. May be generated internally by the ToolTarget. May be nullptr if the Target does not have a mesh.
diff --git a/Engine/Plugins/Runtime/MeshModelingToolset/Source/ModelingComponentsEditorOnly/Public/ToolTargets/StaticMeshComponentToolTarget.h b/Engine/Plugins/Runtime/MeshModelingToolset/Source/ModelingComponentsEditorOnly/Public/ToolTargets/StaticMeshComponentToolTarget.h
--- a/Engine/Plugins/Runtime/MeshModelingToolset/Source/ModelingComponentsEditorOnly/Public/ToolTargets/StaticMeshComponentToolTarget.h
+++ b/Engine/Plugins/Runtime/MeshModelingToolset/Source/ModelingComponentsEditorOnly/Public/ToolTargets/StaticMeshComponentToolTarget.h
@@ -50,6 +50,7 @@ public:
 	virtual FMeshDescription GetEmptyMeshDescription() override;
 	virtual bool SupportsLODs() const override { return true; }
 	virtual TArray<EMeshLODIdentifier> GetAvailableLODs(bool bSkipAutoGenerated = true) const override;
+	virtual EMeshLODIdentifier GetMeshDescriptionLOD() const override { return EditingLOD; }
 
 	// IMeshDescritpionCommitter implementation
 	virtual void CommitMeshDescription(const FCommitter& Committer, const FCommitMeshParameters& CommitParams = FCommitMeshParameters()) override;
diff --git a/Engine/Plugins/Runtime/MeshModelingToolset/Source/ModelingComponentsEditorOnly/Public/ToolTargets/StaticMeshToolTarget.h b/Engine/Plugins/Runtime/MeshModelingToolset/Source/ModelingComponentsEditorOnly/Public/ToolTargets/StaticMeshToolTarget.h
--- a/Engine/Plugins/Runtime/MeshModelingToolset/Source/ModelingComponentsEditorOnly/Public/ToolTargets/StaticMeshToolTarget.h
+++ b/Engine/Plugins/Runtime/MeshModelingToolset/Source/ModelingComponentsEditorOnly/Public/ToolTargets/StaticMeshToolTarget.h
@@ -54,6 +54,7 @@ public:
 	virtual FMeshDescription GetEmptyMeshDescription() override;
 	virtual bool SupportsLODs() const override { return true; }
 	virtual TArray<EMeshLODIdentifier> GetAvailableLODs(bool bSkipAutoGenerated = true) const override;
+	virtual EMeshLODIdentifier GetMeshDescriptionLOD() const override { return EditingLOD; }
 
 	// IMeshDescritpionCommitter implementation
 	virtual void CommitMeshDescription(const FCommitter& Committer, const FCommitMeshParameters& CommitParams = FCommitMeshParameters()) override;
diff --git a/Engine/Source/Runtime/InteractiveToolsFramework/Public/TargetInterfaces/MeshDescriptionProvider.h b/Engine/Source/Runtime/InteractiveToolsFramework/Public/TargetInterfaces/MeshDescriptionProvider.h
--- a/Engine/Source/Runtime/InteractiveToolsFramework/Public/TargetInterfaces/MeshDescriptionProvider.h
+++ b/Engine/Source/Runtime/InteractiveToolsFramework/Public/TargetInterfaces/MeshDescriptionProvider.h
@@ -58,7 +58,15 @@ public:
 	}
 
 	/**
-	 * @return true if the provider is able to provide different Mesh Descriptions for different requested LODs.  If so, GetValidLODs() will return the available LODs.
+	 * For providers that have LODs (i.e. if SupportsLODs returns true), returns the LOD that GetMeshDescription/etc will return. Otherwise, returns Default LOD.
+	 */
+	virtual EMeshLODIdentifier GetMeshDescriptionLOD() const
+	{
+		return EMeshLODIdentifier::Default;
+	}
+
+	/**
+	 * @return true if the provider is able to provide different Mesh Descriptions for different requested LODs.  If so, GetAvailableLODs() will return the available LODs.
 	 */
 	virtual bool SupportsLODs() const
 	{
