add plane projection and boolean/offset options to the mesh splines tool

diff --git a/Engine/Plugins/Experimental/MeshModelingToolsetExp/Source/MeshModelingToolsExp/Private/TriangulateSplinesTool.cpp b/Engine/Plugins/Experimental/MeshModelingToolsetExp/Source/MeshModelingToolsExp/Private/TriangulateSplinesTool.cpp
--- a/Engine/Plugins/Experimental/MeshModelingToolsetExp/Source/MeshModelingToolsExp/Private/TriangulateSplinesTool.cpp
+++ b/Engine/Plugins/Experimental/MeshModelingToolsetExp/Source/MeshModelingToolsExp/Private/TriangulateSplinesTool.cpp
@@ -2,7 +2,6 @@
 
 #include "TriangulateSplinesTool.h"
 
-#include "CurveOps/TriangulateCurvesOp.h"
 #include "InteractiveToolManager.h"
 #include "ToolBuilderUtil.h"
 #include "ToolSetupUtil.h"
@@ -201,6 +200,24 @@ TUniquePtr<FDynamicMeshOperator> UTriangulateSplinesTool::MakeNewOperator()
 	TUniquePtr<FTriangulateCurvesOp> Op = MakeUnique<FTriangulateCurvesOp>();
 
 	Op->Thickness = TriangulateProperties->Thickness;
+	Op->bFlipResult = TriangulateProperties->bFlipResult;
+	Op->CombineMethod = TriangulateProperties->CombineMethod;
+	Op->FlattenMethod = TriangulateProperties->FlattenMethod;
+	Op->CurveOffset = TriangulateProperties->CurveOffset;
+	if (TriangulateProperties->CurveOffset == 0.0)
+	{
+		Op->OffsetClosedMethod = EOffsetClosedCurvesMethod::DoNotOffset;
+	}
+	else
+	{
+		Op->OffsetClosedMethod = TriangulateProperties->OffsetClosedCurves;
+	}
+	Op->OffsetOpenMethod = TriangulateProperties->OpenCurves;
+	Op->OffsetJoinMethod = TriangulateProperties->JoinMethod;
+	Op->OpenEndShape = TriangulateProperties->EndShapes;
+	Op->MiterLimit = TriangulateProperties->MiterLimit;
+	
+	Op->bFlipResult = TriangulateProperties->bFlipResult;
 
 	TArray<TObjectPtr<USplineComponent>, TInlineAllocator<8>> SplineComponents;
 	for (int32 ActorIdx = 0; ActorIdx < ActorsWithSplines.Num(); ++ActorIdx)
diff --git a/Engine/Plugins/Experimental/MeshModelingToolsetExp/Source/MeshModelingToolsExp/Public/TriangulateSplinesTool.h b/Engine/Plugins/Experimental/MeshModelingToolsetExp/Source/MeshModelingToolsExp/Public/TriangulateSplinesTool.h
--- a/Engine/Plugins/Experimental/MeshModelingToolsetExp/Source/MeshModelingToolsExp/Public/TriangulateSplinesTool.h
+++ b/Engine/Plugins/Experimental/MeshModelingToolsetExp/Source/MeshModelingToolsExp/Public/TriangulateSplinesTool.h
@@ -12,6 +12,8 @@
 #include "InteractiveToolQueryInterfaces.h"
 #include "PropertySets/CreateMeshObjectTypeProperties.h"
 
+#include "CurveOps/TriangulateCurvesOp.h"
+
 
 #include "TriangulateSplinesTool.generated.h"
 
@@ -31,11 +33,47 @@ public:
 	// How far to allow the triangulation boundary can deviate from the spline curve before we add more vertices
 	UPROPERTY(EditAnywhere, Category = Spline, meta = (ClampMin = 0.001))
 	double ErrorTolerance = 1.0;
+	
+	// Whether and how to flatten the curves. If curves are flattened, they can also be offset and combined
+	UPROPERTY(EditAnywhere, Category = Spline)
+	EFlattenCurveMethod FlattenMethod = EFlattenCurveMethod::DoNotFlatten;
+
+	// Note: Combining and offsetting curves only works when curves are flattened; curves will be left separate and non-offset if FlattenMethod is DoNotFlatten
+	UPROPERTY(EditAnywhere, Category = Spline, meta = (EditCondition = "FlattenMethod != EFlattenCurveMethod::DoNotFlatten", EditConditionHides))
+	ECombineCurvesMethod CombineMethod = ECombineCurvesMethod::LeaveSeparate;
 
 	// If > 0, Extrude the triangulation by this amount
-	UPROPERTY(EditAnywhere, Category = Spline, meta = (ClampMin = 0.0))
+	UPROPERTY(EditAnywhere, Category = Mesh, meta = (ClampMin = 0.0))
 	double Thickness = 0.0;
-	
+
+	// Whether to flip the facing direction of the generated mesh
+	UPROPERTY(EditAnywhere, Category = Mesh)
+	bool bFlipResult = false;
+
+	// How to handle open curves: Either offset them, or treat them as closed curves
+	UPROPERTY(EditAnywhere, Category = Offset, meta = (EditCondition = "FlattenMethod != EFlattenCurveMethod::DoNotFlatten", EditConditionHides))
+	EOffsetOpenCurvesMethod OpenCurves = EOffsetOpenCurvesMethod::Offset;
+
+	// How much offset to apply to curves
+	UPROPERTY(EditAnywhere, Category = Offset, meta = (EditCondition = "FlattenMethod != EFlattenCurveMethod::DoNotFlatten", EditConditionHides))
+	double CurveOffset = 0.0;
+
+	// Whether and how to apply offset to closed curves
+	UPROPERTY(EditAnywhere, Category = Offset, meta = (EditCondition = "FlattenMethod != EFlattenCurveMethod::DoNotFlatten && CurveOffset != 0", EditConditionHides))
+	EOffsetClosedCurvesMethod OffsetClosedCurves = EOffsetClosedCurvesMethod::OffsetOuterSide;
+
+	// The shape of the ends of offset curves
+	UPROPERTY(EditAnywhere, Category = Offset, meta = (EditCondition = "FlattenMethod != EFlattenCurveMethod::DoNotFlatten && OpenCurves != EOffsetOpenCurvesMethod::TreatAsClosed && CurveOffset != 0", EditConditionHides))
+	EOpenCurveEndShapes EndShapes = EOpenCurveEndShapes::Square;
+
+	// The shape of joins between segments of an offset curve
+	UPROPERTY(EditAnywhere, Category = Offset, meta = (EditCondition = "FlattenMethod != EFlattenCurveMethod::DoNotFlatten && CurveOffset != 0", EditConditionHides))
+	EOffsetJoinMethod JoinMethod = EOffsetJoinMethod::Square;
+
+	// How far a miter join can extend before it is replaced by a square join
+	UPROPERTY(EditAnywhere, Category = Offset, meta = (ClampMin = 0.0, EditCondition = "FlattenMethod != EFlattenCurveMethod::DoNotFlatten && CurveOffset != 0 && JoinMethod == EOffsetJoinMethod::Miter", EditConditionHides))
+	double MiterLimit = 1.0;
+
 };
 
 /**
diff --git a/Engine/Plugins/Runtime/MeshModelingToolset/Source/ModelingOperators/Private/CurveOps/TriangulateCurvesOp.cpp b/Engine/Plugins/Runtime/MeshModelingToolset/Source/ModelingOperators/Private/CurveOps/TriangulateCurvesOp.cpp
--- a/Engine/Plugins/Runtime/MeshModelingToolset/Source/ModelingOperators/Private/CurveOps/TriangulateCurvesOp.cpp
+++ b/Engine/Plugins/Runtime/MeshModelingToolset/Source/ModelingOperators/Private/CurveOps/TriangulateCurvesOp.cpp
@@ -5,8 +5,13 @@
 #include "Util/ProgressCancel.h"
 #include "BoxTypes.h"
 #include "CompGeom/PolygonTriangulation.h"
+#include "CompGeom/Delaunay2.h"
 #include "Operations/ExtrudeMesh.h"
 #include "DynamicMesh/MeshNormals.h"
+#include "Curve/GeneralPolygon2.h"
+
+#include "Curve/PolygonOffsetUtils.h"
+#include "Curve/PolygonIntersectionUtils.h"
 
 #include "Components/SplineComponent.h"
 
@@ -44,7 +49,7 @@ void FTriangulateCurvesOp::CalculateResult(FProgressCancel* Progress)
 	}
 	double UseUVScaleFactor = UVScaleFactor / InputBounds.MaxDim();
 
-	auto AppendTriangles = [this, UseUVScaleFactor](FDynamicMesh3& Mesh, const TArray<FVector3d>& Vertices, const TArray<FIndex3i>& Tris, int32 GroupID, FVector3d PlaneOrigin, FVector3d Normal) -> void
+	auto AppendTriangles = [this, UseUVScaleFactor](FDynamicMesh3& Mesh, const TArray<FVector3d>& Vertices, const TArray<FIndex3i>& Tris, int32 GroupID, FVector3d PlaneOrigin, FVector3d Normal, bool bFlip) -> void
 	{
 		if (Tris.IsEmpty())
 		{
@@ -69,28 +74,236 @@ void FTriangulateCurvesOp::CalculateResult(FProgressCancel* Progress)
 		for (const FIndex3i& T : Tris)
 		{
 			FIndex3i NewTri(T.A + VertStart, T.B + VertStart, T.C + VertStart);
+			if (bFlip)
+			{
+				Swap(NewTri.B, NewTri.C);
+			}
 			int32 NewTID = Mesh.AppendTriangle(NewTri, GroupID);
 			Mesh.Attributes()->PrimaryNormals()->SetTriangle(NewTID, NewTri);
 			Mesh.Attributes()->PrimaryUV()->SetTriangle(NewTID, NewTri);
 		}
 	};
 
-	for (int32 PathIdx = 0; PathIdx < Paths.Num(); ++PathIdx)
+	// Non-flat case must go through the simple triangulation path
+	if (FlattenMethod == EFlattenCurveMethod::DoNotFlatten)
 	{
-		FVector3d Normal, PlanePoint;
-		PolygonTriangulation::ComputePolygonPlane<double>(Paths[PathIdx].Vertices, Normal, PlanePoint);
-		TArray<FIndex3i> Tris;
-		PolygonTriangulation::TriangulateSimplePolygon<double>(Paths[PathIdx].Vertices, Tris, false);
-		AppendTriangles(*ResultMesh, Paths[PathIdx].Vertices, Tris, PathIdx, PlanePoint, Normal);
+		for (int32 PathIdx = 0; PathIdx < Paths.Num(); ++PathIdx)
+		{
+			FVector3d Normal, PlanePoint;
+			PolygonTriangulation::ComputePolygonPlane<double>(Paths[PathIdx].Vertices, Normal, PlanePoint);
+			if (bFlipResult)
+			{
+				Normal = -Normal;
+			}
+			TArray<FIndex3i> Tris;
+			PolygonTriangulation::TriangulateSimplePolygon<double>(Paths[PathIdx].Vertices, Tris, false);
+			AppendTriangles(*ResultMesh, Paths[PathIdx].Vertices, Tris, PathIdx, PlanePoint, Normal, bFlipResult);
+		}
+
+		ApplyThickness(UseUVScaleFactor);
+		return;
 	}
 
-	if (Thickness > 0)
+	FVector3d SharedFrameOrigin = InputBounds.Center();
+	FVector3d SharedFrameNormal = FVector::UnitZ();
+	if (FlattenMethod == EFlattenCurveMethod::AlongX)
 	{
-		FExtrudeMesh ExtrudeMesh(ResultMesh.Get());
-		ExtrudeMesh.DefaultExtrudeDistance = Thickness;
-		ExtrudeMesh.UVScaleFactor = UseUVScaleFactor;
-		ExtrudeMesh.Apply();
+		SharedFrameNormal = FVector::UnitX();
+	}
+	else if (FlattenMethod == EFlattenCurveMethod::AlongY)
+	{
+		SharedFrameNormal = FVector::UnitY();
 	}
+	else if (FlattenMethod == EFlattenCurveMethod::ToBestFitPlane && CombineMethod != ECombineCurvesMethod::LeaveSeparate)
+	{
+		// Take an unweighted average of the polygon normals as the shared plane normal
+		// TODO: consider area-weighting these or instead doing a plane fit directly to the combined vertices
+		FVector3d PlanePoint;
+		FVector3d AccumulatedNormal;
+		PolygonTriangulation::ComputePolygonPlane<double>(Paths[0].Vertices, AccumulatedNormal, PlanePoint);
+		for (int32 PathIdx = 1; PathIdx < Paths.Num(); ++PathIdx)
+		{
+			FVector3d PolygonNormal;
+			PolygonTriangulation::ComputePolygonPlane<double>(Paths[0].Vertices, PolygonNormal, PlanePoint);
+			if (PolygonNormal.Dot(AccumulatedNormal) < 0)
+			{
+				PolygonNormal = -PolygonNormal;
+			}
+			AccumulatedNormal += PolygonNormal;
+		}
+		SharedFrameNormal = Normalized(AccumulatedNormal);
+	}
+
+	auto ConvertOffsetOpenEndType = [](EOpenCurveEndShapes OpenEnd) -> EPolygonOffsetEndType
+	{
+		switch (OpenEnd)
+		{
+		case EOpenCurveEndShapes::Square:
+			return EPolygonOffsetEndType::Square;
+		case EOpenCurveEndShapes::Round:
+			return EPolygonOffsetEndType::Round;
+		case EOpenCurveEndShapes::Butt:
+			return EPolygonOffsetEndType::Butt;
+		}
+		checkNoEntry();
+		return EPolygonOffsetEndType::Square;
+	};
+
+	auto ConvertOffsetJoinType = [](EOffsetJoinMethod OffsetJoin) -> EPolygonOffsetJoinType
+	{
+		switch (OffsetJoin)
+		{
+		case EOffsetJoinMethod::Square:
+			return EPolygonOffsetJoinType::Square;
+		case EOffsetJoinMethod::Miter:
+			return EPolygonOffsetJoinType::Miter;
+		case EOffsetJoinMethod::Round:
+			return EPolygonOffsetJoinType::Round;
+		}
+		checkNoEntry();
+		return EPolygonOffsetJoinType::Square;
+	};
+
+	auto ProcessPaths = [&/*TODO FIX*/](FVector3d Normal, FVector3d PlanePoint, TArrayView<const FCurvePath> ToProcess, int32 GroupID) -> int32
+	{
+		if (bFlipResult)
+		{
+			Normal = -Normal;
+		}
+
+		TArray<UE::Geometry::FGeneralPolygon2d> PolygonsGroup[2]; // Second group for Intersect/Difference/ExclusiveOr
+		const bool bNeedsSecondGroup = CombineMethod == ECombineCurvesMethod::Intersect || CombineMethod == ECombineCurvesMethod::Difference || CombineMethod == ECombineCurvesMethod::ExclusiveOr;
+		auto GetPolygonsGroup = [&PolygonsGroup, bNeedsSecondGroup](int32 CurveIdx) ->TArray<UE::Geometry::FGeneralPolygon2d>&
+		{
+			return bNeedsSecondGroup && CurveIdx > 0 ? PolygonsGroup[1] : PolygonsGroup[0];
+		};
+
+		// Array of projected points (to be reused)
+		TArray<FVector2D> Projected;
+		const FFrame3d Frame(PlanePoint, Normal);
+		FVector3d FrameX = Frame.GetAxis(0);
+		FVector3d FrameY = Frame.GetAxis(1);
+		FVector3d FrameZ = Frame.GetAxis(2);
+		auto SetProjected = [Frame, &Projected](const FCurvePath& CurvePath) -> void
+		{
+			Projected.Reset(CurvePath.Vertices.Num());
+			for (const FVector& Pt : CurvePath.Vertices)
+			{
+				Projected.Add(Frame.ToPlaneUV(Pt));
+			}
+		};
+		// Array of general polygons (to be reused);
+		TArray<UE::Geometry::FGeneralPolygon2d> PolygonsToAdd;
+
+		// If open paths are to be offset, use this to convert them to polygons
+		bool bHasProcessedOpenCurves = EOffsetOpenCurvesMethod::Offset == OffsetOpenMethod;
+		// Note: If curve offset is 0, the open curves do not contribute polygons to the output
+		if (bHasProcessedOpenCurves && CurveOffset != 0.0)
+		{
+			for (int32 CurveIdx = 0; CurveIdx < ToProcess.Num(); ++CurveIdx)
+			{
+				const FCurvePath& CurvePath = ToProcess[CurveIdx];
+				if (!CurvePath.bClosed)
+				{
+					SetProjected(CurvePath);
+					UE::Geometry::FOffsetPolygon2d OffsetPolygon;
+					OffsetPolygon.Polygons.Add(Projected);
+					OffsetPolygon.Offset = CurveOffset;
+					OffsetPolygon.MiterLimit = MiterLimit;
+					OffsetPolygon.JoinType = ConvertOffsetJoinType(OffsetJoinMethod);
+					OffsetPolygon.EndType = ConvertOffsetOpenEndType(OpenEndShape);
+					if (OffsetPolygon.ComputeResult())
+					{
+						GetPolygonsGroup(CurveIdx).Append(OffsetPolygon.Result);
+					}
+				}
+			}
+		}
+
+		// Closed-path handling
+		for (int32 CurveIdx = 0; CurveIdx < ToProcess.Num(); ++CurveIdx)
+		{
+			const FCurvePath& CurvePath = ToProcess[CurveIdx];
+			if (bHasProcessedOpenCurves && !CurvePath.bClosed)
+			{
+				continue;
+			}
+			checkSlow(CurvePath.bClosed || OffsetOpenMethod == EOffsetOpenCurvesMethod::TreatAsClosed);
+
+			SetProjected(CurvePath);
+			UE::Geometry::FGeneralPolygon2d PolyToProcess(Projected);
+			TArrayView<const FGeneralPolygon2d> InitialPolygons(&PolyToProcess, 1);
+			// Always pre-process with a union to make orientation consistent
+			PolygonsToAdd.Reset();
+			PolygonsUnion(InitialPolygons, PolygonsToAdd, true);
+			if (OffsetClosedMethod != EOffsetClosedCurvesMethod::DoNotOffset)
+			{
+				PolygonsOffset(
+					CurveOffset,
+					PolygonsToAdd, PolygonsToAdd,
+					true,
+					MiterLimit,
+					EPolygonOffsetJoinType::Square, // TODO: convert this from inputs
+					OffsetClosedMethod == EOffsetClosedCurvesMethod::OffsetOuterSide ? EPolygonOffsetEndType::Polygon : EPolygonOffsetEndType::Joined);
+			}
+			GetPolygonsGroup(CurveIdx).Append(PolygonsToAdd);
+		}
+
+		if (CombineMethod == ECombineCurvesMethod::Union)
+		{
+			checkSlow(!bNeedsSecondGroup);
+			PolygonsUnion(PolygonsGroup[0], PolygonsGroup[0], true);
+		}
+		else if (CombineMethod == ECombineCurvesMethod::Intersect)
+		{
+			PolygonsIntersection(PolygonsGroup[0], PolygonsGroup[1], PolygonsGroup[0]);
+		}
+		else if (CombineMethod == ECombineCurvesMethod::Difference)
+		{
+			PolygonsDifference(PolygonsGroup[0], PolygonsGroup[1], PolygonsGroup[0]);
+		}
+		else if (CombineMethod == ECombineCurvesMethod::ExclusiveOr)
+		{
+			PolygonsExclusiveOr(PolygonsGroup[0], PolygonsGroup[1], PolygonsGroup[0]);
+		}
+
+		// Triangulate and append each polygon separately
+		for (int32 PolyIdx = 0; PolyIdx < PolygonsGroup[0].Num(); ++PolyIdx)
+		{
+			UE::Geometry::FDelaunay2 Delaunay;
+			Delaunay.bAutomaticallyFixEdgesToDuplicateVertices = true;
+			TArray<FIndex3i> Triangles;
+			TArray<FVector2d> Vertices;
+			TArray<FVector3d> Vertices3d;
+			Delaunay.Triangulate(PolygonsGroup[0][PolyIdx], &Triangles, &Vertices, true);
+			Vertices3d.SetNum(Vertices.Num());
+			for (int32 VIdx = 0; VIdx < Vertices.Num(); ++VIdx)
+			{
+				Vertices3d[VIdx] = Frame.FromPlaneUV(Vertices[VIdx]);
+			}
+			AppendTriangles(*ResultMesh, Vertices3d, Triangles, GroupID++, PlanePoint, Normal, true);
+		}
+
+		return GroupID;
+	};
+
+	if (CombineMethod == ECombineCurvesMethod::LeaveSeparate)
+	{
+		int32 GroupIdx = 1;
+		for (int32 CurveIdx = 0; CurveIdx < Paths.Num(); ++CurveIdx)
+		{
+			FVector3d PlanePoint;
+			FVector3d Normal;
+			PolygonTriangulation::ComputePolygonPlane<double>(Paths[CurveIdx].Vertices, Normal, PlanePoint);
+			GroupIdx = ProcessPaths(Normal, SharedFrameOrigin, TArrayView<const FCurvePath>(&Paths[CurveIdx], 1), GroupIdx);
+		}
+	}
+	else
+	{
+		ProcessPaths(SharedFrameNormal, SharedFrameOrigin, Paths, 1);
+	}
+
+	ApplyThickness(UseUVScaleFactor);
 
 	if (Progress && Progress->Cancelled())
 	{
@@ -98,3 +311,13 @@ void FTriangulateCurvesOp::CalculateResult(FProgressCancel* Progress)
 	}
 }
 
+void FTriangulateCurvesOp::ApplyThickness(double UseUVScaleFactor)
+{
+	if (Thickness > 0)
+	{
+		FExtrudeMesh ExtrudeMesh(ResultMesh.Get());
+		ExtrudeMesh.DefaultExtrudeDistance = Thickness;
+		ExtrudeMesh.UVScaleFactor = UseUVScaleFactor;
+		ExtrudeMesh.Apply();
+	}
+}
\ No newline at end of file
diff --git a/Engine/Plugins/Runtime/MeshModelingToolset/Source/ModelingOperators/Public/CurveOps/TriangulateCurvesOp.h b/Engine/Plugins/Runtime/MeshModelingToolset/Source/ModelingOperators/Public/CurveOps/TriangulateCurvesOp.h
--- a/Engine/Plugins/Runtime/MeshModelingToolset/Source/ModelingOperators/Public/CurveOps/TriangulateCurvesOp.h
+++ b/Engine/Plugins/Runtime/MeshModelingToolset/Source/ModelingOperators/Public/CurveOps/TriangulateCurvesOp.h
@@ -5,8 +5,60 @@
 #include "CoreMinimal.h"
 #include "ModelingOperators.h"
 
+#include "TriangulateCurvesOp.generated.h"
+
 class USplineComponent;
 
+UENUM()
+enum class EFlattenCurveMethod : uint8
+{
+	DoNotFlatten,
+	ToBestFitPlane,
+	AlongX,
+	AlongY,
+	AlongZ
+};
+
+UENUM()
+enum class ECombineCurvesMethod : uint8
+{
+	LeaveSeparate,
+	Union,
+	Intersect,
+	Difference,
+	ExclusiveOr
+};
+
+UENUM()
+enum class EOffsetClosedCurvesMethod : uint8
+{
+	DoNotOffset,
+	OffsetOuterSide,
+	OffsetBothSides
+};
+
+UENUM()
+enum class EOffsetOpenCurvesMethod : uint8
+{
+	TreatAsClosed,
+	Offset
+};
+
+UENUM()
+enum class EOffsetJoinMethod : uint8
+{
+	Square,
+	Miter,
+	Round
+};
+
+UENUM()
+enum class EOpenCurveEndShapes : uint8
+{
+	Square,
+	Round,
+	Butt
+};
 
 namespace UE {
 namespace Geometry {
@@ -39,6 +91,21 @@ public:
 	// If > 0, thicken the result mesh to make a solid
 	double Thickness = 0.0;
 
+	EFlattenCurveMethod FlattenMethod = EFlattenCurveMethod::DoNotFlatten;
+
+	// Note: Combining and offsetting curves only works when curves are flattened; curves will be left separate and non-offset if FlattenMethod is DoNotFlatten
+
+	ECombineCurvesMethod CombineMethod = ECombineCurvesMethod::LeaveSeparate;
+
+	EOffsetClosedCurvesMethod OffsetClosedMethod = EOffsetClosedCurvesMethod::DoNotOffset;
+	EOffsetOpenCurvesMethod OffsetOpenMethod = EOffsetOpenCurvesMethod::TreatAsClosed;
+	EOffsetJoinMethod OffsetJoinMethod = EOffsetJoinMethod::Square;
+	EOpenCurveEndShapes OpenEndShape = EOpenCurveEndShapes::Square;
+	double MiterLimit = 1.0;
+	double CurveOffset = 1.0;
+
+	bool bFlipResult = false;
+
 	//
 	// FDynamicMeshOperator interface 
 	//
@@ -58,6 +125,7 @@ private:
 	// Local to World transform of the first path
 	FTransform FirstPathTransform;
 
+	void ApplyThickness(double UseUVScaleFactor);
 };
 
 }} // end UE::Geometry
