Fix triangulate curves op using the wrong index when computing average polygon normal

diff --git a/Engine/Plugins/Runtime/MeshModelingToolset/Source/ModelingOperators/Private/CurveOps/TriangulateCurvesOp.cpp b/Engine/Plugins/Runtime/MeshModelingToolset/Source/ModelingOperators/Private/CurveOps/TriangulateCurvesOp.cpp
--- a/Engine/Plugins/Runtime/MeshModelingToolset/Source/ModelingOperators/Private/CurveOps/TriangulateCurvesOp.cpp
+++ b/Engine/Plugins/Runtime/MeshModelingToolset/Source/ModelingOperators/Private/CurveOps/TriangulateCurvesOp.cpp
@@ -124,7 +124,7 @@ void FTriangulateCurvesOp::CalculateResult(FProgressCancel* Progress)
 		for (int32 PathIdx = 1; PathIdx < Paths.Num(); ++PathIdx)
 		{
 			FVector3d PolygonNormal;
-			PolygonTriangulation::ComputePolygonPlane<double>(Paths[0].Vertices, PolygonNormal, PlanePoint);
+			PolygonTriangulation::ComputePolygonPlane<double>(Paths[PathIdx].Vertices, PolygonNormal, PlanePoint);
 			if (PolygonNormal.Dot(AccumulatedNormal) < 0)
 			{
 				PolygonNormal = -PolygonNormal;
