MeshModelingTools: Various volume conversion fixes:
- Fix conversion sometimes making degenerate tris and firing an ensure.
- Fix volume conversion sometimes creating cracks in volumes when dealing with non-convex faces (by deleting necessary colinear verts)
- Make simplification preserve group boundaries when asked.
- Remove colinear verts along group boundaries without losing boundaries when preserving groups.
- Break up faces that have holes into convexes rather than emitting per-triangle faces.
- Make mesh to volume tool use same path for conversion as tool targets, and add controls for preserving groups and simplifying
- Revert to ear clipping if delaunay triangulation fails in converting back from volume

#prelfight 6470ccbf5b1dc793b715d66b

diff --git a/Engine/Plugins/Experimental/MeshModelingToolsetExp/Source/MeshModelingToolsEditorOnlyExp/Private/MeshToVolumeTool.cpp b/Engine/Plugins/Experimental/MeshModelingToolsetExp/Source/MeshModelingToolsEditorOnlyExp/Private/MeshToVolumeTool.cpp
--- a/Engine/Plugins/Experimental/MeshModelingToolsetExp/Source/MeshModelingToolsEditorOnlyExp/Private/MeshToVolumeTool.cpp
+++ b/Engine/Plugins/Experimental/MeshModelingToolsetExp/Source/MeshModelingToolsEditorOnlyExp/Private/MeshToVolumeTool.cpp
@@ -66,8 +66,6 @@ public:
 			return;
 		}
 
-		bool bShowTooLargeWarning = false;
-
 		auto SimplifyToMax = [this](FDynamicMesh3& Mesh, int32 TriangleCount)
 		{
 			FVolPresMeshSimplification Simplifier(&Mesh);
@@ -76,34 +74,12 @@ public:
 
 		if (ConversionMode == EMeshToVolumeMode::MinimalPolygons)
 		{
-			// Since this tool is likely to be a sink, there isn't much reason to keep
-			// the group differentiations if they are coplanar.
-			bool bRespectGroupBoundaries = false;
-
-			// Apply minimal-planar simplification to remove extra vertices along straight edges
-			FDynamicMesh3 LocalMesh = *SourceMesh;
-			LocalMesh.DiscardAttributes();
-			FQEMSimplification PlanarSimplifier(&LocalMesh);
-			PlanarSimplifier.SimplifyToMinimalPlanar(0.1);		// angle tolerance in degrees
-
-			UE::Conversion::GetPolygonFaces(LocalMesh, *Result, bRespectGroupBoundaries);
-
-			bTooManyTriangles = (LocalMesh.TriangleCount() > MeshToVolumeOptions.MaxTriangles);
-
-			if (Progress && Progress->Cancelled())
-			{
-				return;
-			}
-
-			if (bTooManyTriangles)
-			{
-				SimplifyToMax(LocalMesh, MeshToVolumeOptions.MaxTriangles);
-				UE::Conversion::GetPolygonFaces(LocalMesh, *Result, bRespectGroupBoundaries);
-			}
+			UE::Conversion::GetPolygonFaces(*SourceMesh, MeshToVolumeOptions, *Result);
 		}
 		else
 		{
-			bTooManyTriangles = (SourceMesh->TriangleCount() > MeshToVolumeOptions.MaxTriangles);
+			bTooManyTriangles = MeshToVolumeOptions.bAutoSimplify 
+				&& (SourceMesh->TriangleCount() > MeshToVolumeOptions.MaxTriangles);
 
 			if (Progress && Progress->Cancelled())
 			{
@@ -165,9 +141,14 @@ void UMeshToVolumeTool::Setup()
 	Settings->RestoreProperties(this);
 	AddToolPropertySource(Settings);
 
-	Settings->WatchProperty(Settings->ConversionMode,
-							[this](EMeshToVolumeMode NewMode)
-							{ Compute->InvalidateResult(); });
+	Settings->WatchProperty(Settings->ConversionMode, 
+		[this](EMeshToVolumeMode NewMode) { Compute->InvalidateResult(); });
+	Settings->WatchProperty(Settings->bPreserveGroupBoundaries, 
+		[this](bool NewValue) { Compute->InvalidateResult(); });
+	Settings->WatchProperty(Settings->bAutoSimplify,
+		[this](bool NewValue) { Compute->InvalidateResult(); });
+	Settings->WatchProperty(Settings->SimplifyMaxTriangles,
+		[this](int32 NewValue) { Compute->InvalidateResult(); });
 
 
 	HandleSourcesProperties = NewObject<UOnAcceptHandleSourcesPropertiesSingle>(this);
@@ -322,6 +303,10 @@ TUniquePtr<UE::Geometry::TGenericDataOperator<FDynamicMeshFaceArray>> UMeshToVol
 	VolumeOp->SourceMesh = InputMesh;
 	VolumeOp->ConversionMode = Settings->ConversionMode;
 
+	VolumeOp->MeshToVolumeOptions.bAutoSimplify = Settings->bAutoSimplify;
+	VolumeOp->MeshToVolumeOptions.MaxTriangles = Settings->SimplifyMaxTriangles;
+	VolumeOp->MeshToVolumeOptions.bRespectGroupBoundaries = Settings->bPreserveGroupBoundaries;
+
 	return VolumeOp;
 }
 
diff --git a/Engine/Plugins/Experimental/MeshModelingToolsetExp/Source/MeshModelingToolsEditorOnlyExp/Public/MeshToVolumeTool.h b/Engine/Plugins/Experimental/MeshModelingToolsetExp/Source/MeshModelingToolsEditorOnlyExp/Public/MeshToVolumeTool.h
--- a/Engine/Plugins/Experimental/MeshModelingToolsetExp/Source/MeshModelingToolsEditorOnlyExp/Public/MeshToVolumeTool.h
+++ b/Engine/Plugins/Experimental/MeshModelingToolsetExp/Source/MeshModelingToolsEditorOnlyExp/Public/MeshToVolumeTool.h
@@ -54,6 +54,25 @@ public:
 	UPROPERTY(EditAnywhere, Category = ConversionOptions)
 	EMeshToVolumeMode ConversionMode = EMeshToVolumeMode::MinimalPolygons;
 
+	/** 
+	 * When true, adjacent coplanar groups will not be merged together into single faces. Not relevant if 
+	 * conversion mode is set to emit all triangles separately.
+	 */
+	UPROPERTY(EditAnywhere, Category = ConversionOptions, meta = (EditCondition = "ConversionMode != EMeshToVolumeMode::TriangulatePolygons"))
+	bool bPreserveGroupBoundaries = false;
+
+	/**
+	 * Determines whether mesh gets auto simplified when its triangle count is too high.
+	 */
+	UPROPERTY(EditAnywhere, Category = ConversionOptions, AdvancedDisplay)
+	bool bAutoSimplify = true;
+
+	/**
+	 * Target triangle count for auto simplification when Auto Simplify is true.
+	 */
+	UPROPERTY(EditAnywhere, Category = ConversionOptions, AdvancedDisplay, meta = (EditCondition = "bAutoSimplify"))
+	int32 SimplifyMaxTriangles = 250;
+
 	/** Type of new Volume to create on Accept */
 	UPROPERTY(EditAnywhere, Category = NewVolume, meta = (EditCondition = "TargetVolume == nullptr") )
 	TSubclassOf<class AVolume> NewVolumeType = ABlockingVolume::StaticClass();
diff --git a/Engine/Plugins/Runtime/MeshModelingToolset/Source/ModelingComponents/Private/ConversionUtils/VolumeToDynamicMesh.cpp b/Engine/Plugins/Runtime/MeshModelingToolset/Source/ModelingComponents/Private/ConversionUtils/VolumeToDynamicMesh.cpp
--- a/Engine/Plugins/Runtime/MeshModelingToolset/Source/ModelingComponents/Private/ConversionUtils/VolumeToDynamicMesh.cpp
+++ b/Engine/Plugins/Runtime/MeshModelingToolset/Source/ModelingComponents/Private/ConversionUtils/VolumeToDynamicMesh.cpp
@@ -2,6 +2,7 @@
 
 #include "ConversionUtils/VolumeToDynamicMesh.h"
 
+#include "CompGeom/PolygonTriangulation.h"
 #include "ConstrainedDelaunay2.h"
 #include "Polygon2.h"
 #include "DynamicMesh/DynamicMesh3.h"
@@ -23,6 +24,28 @@
 
 using namespace UE::Geometry;
 
+namespace VolumeToDynamicMeshLocals
+{
+	/** Attempts to use delaunay triangulation to triangulate the polygon, and falls back to ear clipping if that fails. */
+	TArray<FIndex3i> TriangulatePolygon(const FPolygon2d& ToTriangulate)
+	{
+		// This first part is equivalent to calling  ConstrainedDelaunayTriangulate<double>(ToTriangulate), except
+		// it allows us to properly check for failure.
+		TConstrainedDelaunay2<double> DelaunayTriangulation;
+		DelaunayTriangulation.FillRule = TConstrainedDelaunay2<double>::EFillRule::Positive;
+		DelaunayTriangulation.Add(ToTriangulate);
+		if (DelaunayTriangulation.Triangulate() && DelaunayTriangulation.Triangles.Num() > 0)
+		{
+			return DelaunayTriangulation.Triangles;
+		}
+		
+		// If delaunay triangulation fails, use ear clipping.
+		TArray<FIndex3i> TrianglesOut;
+		PolygonTriangulation::TriangulateSimplePolygon(ToTriangulate.GetVertices(), TrianglesOut);
+		return TrianglesOut;
+	}
+}
+
 namespace UE {
 namespace Conversion {
 
@@ -66,6 +89,8 @@ void BrushComponentToDynamicMesh(UBrushComponent* Component, UE::Geometry::FDyna
 
 void BrushToDynamicMesh(UModel& BrushModel, UE::Geometry::FDynamicMesh3& Mesh, const FVolumeToMeshOptions& Options)
 {
+	using namespace VolumeToDynamicMeshLocals;
+
 	Mesh.Clear();
 	if (Options.bSetGroups)
 	{
@@ -107,7 +132,7 @@ void BrushToDynamicMesh(UModel& BrushModel, UE::Geometry::FDynamicMesh3& Mesh, c
 		// polygon, but the polygons we get are oriented opposite of what we want (they
 		// are clockwise if the normal is towards us), so this ends up giving us the triangle
 		// orientation that we want.
-		TArray<FIndex3i> PolyTriangles = ConstrainedDelaunayTriangulate<double>(ToTriangulate);
+		TArray<FIndex3i> PolyTriangles = TriangulatePolygon(ToTriangulate);
 
 		int32 GroupID = FDynamicMesh3::InvalidID;
 		if (Options.bSetGroups)
@@ -147,7 +172,7 @@ void BrushToDynamicMesh(UModel& BrushModel, UE::Geometry::FDynamicMesh3& Mesh, c
 			// polygon, but the polygons we get are oriented opposite of what we want (they
 			// are clockwise if the normal is towards us), so this ends up giving us the triangle
 			// orientation that we want.
-			TArray<FIndex3i> PolyTriangles = ConstrainedDelaunayTriangulate<double>(ToTriangulate);
+			TArray<FIndex3i> PolyTriangles = TriangulatePolygon(ToTriangulate);
 
 			int32 GroupID = FDynamicMesh3::InvalidID;
 			if (Options.bSetGroups)
diff --git a/Engine/Plugins/Runtime/MeshModelingToolset/Source/ModelingComponentsEditorOnly/Private/ConversionUtils/DynamicMeshToVolume.cpp b/Engine/Plugins/Runtime/MeshModelingToolset/Source/ModelingComponentsEditorOnly/Private/ConversionUtils/DynamicMeshToVolume.cpp
--- a/Engine/Plugins/Runtime/MeshModelingToolset/Source/ModelingComponentsEditorOnly/Private/ConversionUtils/DynamicMeshToVolume.cpp
+++ b/Engine/Plugins/Runtime/MeshModelingToolset/Source/ModelingComponentsEditorOnly/Private/ConversionUtils/DynamicMeshToVolume.cpp
@@ -3,18 +3,428 @@
 #include "ConversionUtils/DynamicMeshToVolume.h"
 
 #include "BSPOps.h"
+#include "CompGeom/PolygonTriangulation.h"
+#include "ConstrainedDelaunay2.h"
+#include "Curve/PlanarComplex.h"
 #include "DynamicMesh/DynamicMesh3.h"
 #include "Engine/Polys.h"
 #include "GameFramework/Volume.h"
 #include "DynamicMesh/MeshNormals.h"
+#include "MeshConstraintsUtil.h"
 #include "MeshRegionBoundaryLoops.h"
 #include "Model.h"
+#include "Operations/PlanarHoleFiller.h"
 #include "Selections/MeshConnectedComponents.h"
+#include "Util/ProgressCancel.h"
 
 #include "MeshSimplification.h"
 
 using namespace UE::Geometry;
 
+namespace DynamicMeshToVolumeLocals
+{
+	// This is meant as a replacement for the FPoly Triangulate method, which occasionally
+	// removes colinear verts yet still sometimes creates degenerate triangles that trigger an
+	// ensure in that function.
+	bool TriangulatePolygon(FPoly& PolygonIn, TArray<FPoly>& PolygonsOut)
+	{
+		PolygonsOut.Empty();
+		if (PolygonIn.Vertices.Num() < 3)
+		{
+			return false;
+		}
+
+		if (PolygonIn.Vertices.Num() == 3)
+		{
+			if (!PolygonIn.IsConvex())
+			{
+				// This means that the triangle was given to us with wrong orientation (according to legacy code). This can
+				// actually happen as a result of removing verts in certain near-degenerate cases, so no ensure here.
+				Swap(PolygonIn.Vertices[1], PolygonIn.Vertices[2]);
+			}
+			PolygonsOut.Add(PolygonIn);
+			return true;
+		}
+
+		FFrame3d PolygonFrame((FVector3d)PolygonIn.Base, (FVector3d)PolygonIn.Normal);
+		TArray<FVector2D> Polygon2DCoordinates;
+		for (const FVector3f& Vert : PolygonIn.Vertices)
+		{
+			Polygon2DCoordinates.Add(PolygonFrame.ToPlaneUV((FVector3d)Vert));
+		}
+		FPolygon2d Polygon(Polygon2DCoordinates);
+
+		FConstrainedDelaunay2d Delaunay;
+		Delaunay.Add(Polygon);
+		// Output needs to be wound the same way as input
+		Delaunay.bOutputCCW = true;
+
+		TArray<FIndex3i>* TrianglesToUse = &Delaunay.Triangles;
+		TArray<FIndex3i> PolyTriangles;
+		if (!Delaunay.Triangulate() || Delaunay.Triangles.Num() == 0)
+		{
+			PolygonTriangulation::TriangulateSimplePolygon(Polygon.GetVertices(), PolyTriangles, false);
+			TrianglesToUse = &PolyTriangles;
+		}
+
+		for (const FIndex3i& Triangle : *TrianglesToUse)
+		{
+			FPoly PolyToEmit;
+			PolyToEmit.Init();
+
+			for (int i = 0; i < 3; ++i)
+			{
+				// Note: the ordering between PolygonIn.Vertices and Polygon2DCoordinates should be identical so that
+				// we can index back to get the original location easily.
+				PolyToEmit.Vertices.Add(PolygonIn.Vertices[Triangle[i]]);
+			}
+			PolyToEmit.Base = PolyToEmit.Vertices[0];
+
+			// At this point we could call PolyToEmit.Finalize() but it has some questionable
+			// error handling, in particular where some degenerate triangles get kept despite
+			// triggering an ensure and having their normal set to zero. We'll just do it ourselves.
+
+			if (FVector3f::PointsAreSame(PolyToEmit.Vertices[0], PolyToEmit.Vertices[1])
+				|| FVector3f::PointsAreSame(PolyToEmit.Vertices[0], PolyToEmit.Vertices[2])
+				|| FVector3f::PointsAreSame(PolyToEmit.Vertices[1], PolyToEmit.Vertices[2]))
+			{
+				// Don't emit the degenerate poly
+				continue;
+			}
+			
+			// Note that we use the vertices in backwards order because FPoly stores vertices clockwise if the normal is towards you.
+			PolyToEmit.Normal = VectorUtil::Normal(PolyToEmit.Vertices[0], PolyToEmit.Vertices[2], PolyToEmit.Vertices[1]);
+
+			if (!PolyToEmit.Normal.Equals(PolygonIn.Normal))
+			{
+				// We shouldn't really have non-coplanar triangles, but the original handled non-coplanarity, so we do too.
+				FQuaternionf Rotation(PolygonIn.Normal, PolyToEmit.Normal);
+				PolyToEmit.TextureU = Rotation * PolygonIn.TextureU;
+				PolyToEmit.TextureV = Rotation * PolygonIn.TextureV;
+			}
+			else
+			{
+				PolyToEmit.TextureU = PolygonIn.TextureU;
+				PolyToEmit.TextureV = PolygonIn.TextureV;
+			}
+
+			if (!ensure(PolyToEmit.IsConvex()))
+			{
+				// This means that the polygon we got had vertices in the wrong order
+				Swap(PolyToEmit.Vertices[0], PolyToEmit.Vertices[1]);
+
+				if (!ensure(PolyToEmit.IsConvex()))
+				{
+					// Not sure whether this could happen, but don't emit in this case
+					continue;
+				}
+			}
+
+			PolygonsOut.Add(PolyToEmit);
+		}
+
+		return PolygonsOut.Num() > 0;
+	}
+
+	/**
+	 * This is a way to triangulate polygons that have multiple boundary loops (i.e. polygons with
+	 * holes in them) into nicer triangles/convexes than just outputting all triangles separately.
+	 * (We could actually use this for any polygon, but for normal polygons we just emit the whole
+	 * boundary and let the legacy code retriangulate/merge if they are not convex).
+	 */
+	bool TriangulateLoopsIntoFaces(
+		const FDynamicMesh3& InputMesh, const TArray<int32>& Tids,
+		FVector3d PlaneNormal, const TArray<TArray<int32>>& VidLoopsIn,
+		TArray<UE::Conversion::FDynamicMeshFace>& FacesToAppendTo)
+	{
+		// The mesh we have access to is const, and while we could do the retriangulation separately
+		// without a mesh, it helps to have neighbor relationships between the created triangles for
+		// merging them later. So instead of doing a plain retriangulation, we'll create a little
+		// submesh and do a hole fill.
+		FDynamicMesh3 Submesh;
+		TMap<int32, int32> InputVidToSubmeshVid;
+
+		TArray<TArray<int32>> HoleFillVidLoops;
+		for (const TArray<int32>& BaseLoop : VidLoopsIn)
+		{
+			TArray<int32>& HoleFillLoop = HoleFillVidLoops.Emplace_GetRef();
+			// Prep the loops backwards because we found them as boundaries of a component rather than
+			// a hole, and now they are boundaries of a hole.
+			for (int32 i = BaseLoop.Num() - 1; i >= 0; --i)
+			{
+				int32 InputVid = BaseLoop[i];
+				int32* ExistingSubmeshVid = InputVidToSubmeshVid.Find(InputVid);
+				if (!ExistingSubmeshVid)
+				{
+					int32 SubmeshVid = Submesh.AppendVertex(InputMesh.GetVertex(InputVid));
+					ExistingSubmeshVid = &InputVidToSubmeshVid.Add(InputVid, SubmeshVid);
+				}
+				
+				HoleFillLoop.Add(*ExistingSubmeshVid);
+			}
+		}
+
+		// Retriangulation
+		// We use a slightly modified version of UE::Geometry::ConstrainedDelaunayTriangulate<double> for our retriangulation
+		// function in order to know when a retriangulation fails.
+		// TODO: our constrained delaunay triangulation fails in some pathological self intersecting cases that can come
+		// about from degenerate triangles, and it would be nice to be able to fall back to something, but we don't have
+		// an alternative for triangulating with holes.
+		bool bRetriangulationSucceeded = true;
+		auto RetriangulationFunc = [&bRetriangulationSucceeded](const TGeneralPolygon2<double>& GeneralPolygon) -> TArray<FIndex3i>
+		{
+			TConstrainedDelaunay2<double> Triangulation;
+			Triangulation.FillRule = TConstrainedDelaunay2<double>::EFillRule::Positive;
+			Triangulation.Add(GeneralPolygon);
+			if (!Triangulation.Triangulate() || !ensure(Triangulation.Triangles.Num() > 0))
+			{
+				bRetriangulationSucceeded = false;
+			}
+			return Triangulation.Triangles;
+		};
+
+		FPlanarHoleFiller HoleFiller(&Submesh, &HoleFillVidLoops, RetriangulationFunc,
+			Submesh.GetVertex(HoleFillVidLoops[0][0]), PlaneNormal);
+
+		if (!HoleFiller.Fill() || !bRetriangulationSucceeded)
+		{
+			return false;
+		}
+
+		// Now merge the triangles into convexes. 
+		// We could use OptimizeIntoConvexPolys, but we have a bunch of knowledge that lets us do it more
+		// efficiently: we have adjacency information, know that everything is a triangle, and most importantly,
+		// we know that there aren't any interior verts. The lack of interior vertices means that we don't need 
+		// multiple passes (once a triangle is determined to break convexity, another triangle can't fix it because 
+		// that would require a surrounded/interior vertex).
+
+		TSet<int32> ProcessedTids;
+		int32 TidCount = 0; // For sanity check
+		for (int32 Tid : Submesh.TriangleIndicesItr())
+		{
+			if (ProcessedTids.Contains(Tid))
+			{
+				continue;
+			}
+			ProcessedTids.Add(Tid);
+
+			FIndex3i TriVids = Submesh.GetTriangle(Tid);
+
+			// We'll walk around the boundary seeing if we can add neighboring triangles. Top of VidStack is the next
+			// vid to walk to, and we add verts onto our path in front of ourselves as we add triangles.
+			TArray<int32> PolygonVids;
+			TArray<int32> VidStack{ TriVids[2], TriVids[1], TriVids[0] };
+			int32 FromVid = VidStack[0]; // Vert we're walking from (wrap around to bottom of stack)
+			int32 PreviousVid = VidStack[1]; // Vert before FromVid, for convexity check.
+
+			while (!VidStack.IsEmpty())
+			{
+				int32 ToVid = VidStack.Pop();
+				int32 CurrentEid = Submesh.FindEdge(FromVid, ToVid); // edge we're walking along
+
+				// If we determine that the triangle on the other side cannot be added to the polygon, 
+				// then we'll add ToVid to our polygon and move on to the next vid in our stack.
+				auto AdvanceToNext = [&PolygonVids, ToVid, &PreviousVid, &FromVid]()
+				{
+					PolygonVids.Add(ToVid);
+					PreviousVid = FromVid;
+					FromVid = ToVid;
+				};
+
+				// Find the triangle on the other side of the current edge
+				FIndex2i EdgeTids = Submesh.GetEdgeT(CurrentEid);
+				if (EdgeTids.B == IndexConstants::InvalidID) // was a boundary edge
+				{
+					AdvanceToNext();
+					continue;
+				}
+				int32 OtherTid = ProcessedTids.Contains(EdgeTids.A) ? EdgeTids.B : EdgeTids.A;
+				if (ProcessedTids.Contains(OtherTid))
+				{
+					// Both triangles being processed means that the other triangle was already placed into another convex
+					AdvanceToNext();
+					continue;
+				}
+
+				// Look to see whether the two edges of the triangle stay on the correct side of the preceding and following
+				// polygon edges. This test is done similar to the way that it is done in FPoly::IsConvex().
+				FVector3d PreviousVector = (Submesh.GetVertex(FromVid) - Submesh.GetVertex(PreviousVid)).GetSafeNormal();
+				FVector3d PreviousOutVector = PlaneNormal.Cross(PreviousVector); // points out into the disallowed half space
+
+				int32 OtherVid = IndexUtil::FindTriOtherVtxUnsafe(FromVid, ToVid, Submesh.GetTriangle(OtherTid));
+				FVector3d VectorToOther = (Submesh.GetVertex(OtherVid) - Submesh.GetVertex(FromVid)).GetSafeNormal();
+				if (PreviousOutVector.Dot(VectorToOther) > UE_KINDA_SMALL_NUMBER)
+				{
+					AdvanceToNext();
+					continue;
+				}
+
+				int32 VidAfterToVid = VidStack.IsEmpty() ? PolygonVids[0] : VidStack.Top();
+				FVector3d NextVector = (Submesh.GetVertex(VidAfterToVid) - Submesh.GetVertex(ToVid)).GetSafeNormal();
+				FVector3d NextOutVector = PlaneNormal.Cross(NextVector); // points out into the disallowed half space
+				VectorToOther = (Submesh.GetVertex(OtherVid) - Submesh.GetVertex(ToVid)).GetSafeNormal();
+				if (NextOutVector.Dot(VectorToOther) > UE_KINDA_SMALL_NUMBER)
+				{
+					AdvanceToNext();
+					continue;
+				}
+
+				// If we got to here, then we can add in this triangle to the polygon.
+				ProcessedTids.Add(OtherTid);
+
+				// This means that we're not ready to emit ToVid yet, so put it back on the stack, with the new ToVid on top of it.
+				VidStack.Push(ToVid);
+				VidStack.Push(OtherVid);
+
+				// Sanity check
+				if (!ensure(++TidCount <= Submesh.TriangleCount()))
+				{
+					return false;
+				}
+			}
+
+			// Once we've got here, we have a polygon that we can't add more triangles to
+			UE::Conversion::FDynamicMeshFace Face;
+
+			Face.BoundaryLoop.SetNum(PolygonVids.Num());
+			FVector3d PlaneOrigin = FVector3d::Zero();
+			for (int i = 0; i < PolygonVids.Num(); ++i)
+			{
+				Face.BoundaryLoop[i] = Submesh.GetVertex(PolygonVids[i]);
+				PlaneOrigin += Face.BoundaryLoop[i];
+			}
+			PlaneOrigin /= PolygonVids.Num();
+			Face.Plane = FFrame3d(PlaneOrigin, PlaneNormal);
+
+			// Poly vertices are stored backwards
+			Algo::Reverse(Face.BoundaryLoop);
+
+			FacesToAppendTo.Add(Face);
+		}//end emitting faces
+
+		return true;
+	}
+
+	/**
+	 * This filters the passed in FMeshRegionBoundaryLoops to create vid loops that do not include
+	 * unnecessary colinear verts, where "unnecessary" means that the colinear vert is not part of
+	 * more than just the two adjacent faces (i.e., it is not actually a T-junction).
+	 * 
+	 * @param VertCanBeSkippedMap A map used to cache knowledge of whether a vert is skippable, if
+	 *  it has already been examined as part of another loop.
+	 * @param TrianglesShouldBeInSameComponent Function that returns true when the two triangles are
+	 *  in the same face, i.e. the function used in the connected component search.
+	 */
+	void GetVidLoopsWithoutUnneededColinearVerts(const FDynamicMesh3& InputMesh, const FMeshRegionBoundaryLoops& Loops,
+		TMap<int32, bool>& VertCanBeSkippedMap, TFunctionRef<bool(int32, int32)> TrianglesShouldBeInSameComponent,
+		TArray<TArray<int32>>& VidLoopsOut)
+	{
+		auto IsVertMeetingPointOfMultipleBoundaries = [&InputMesh, TrianglesShouldBeInSameComponent](int32 Vid)
+		{
+			int NumAttachedBoundaries = 0;
+			for (int32 Eid : InputMesh.VtxEdgesItr(Vid))
+			{
+				FIndex2i EdgeTids = InputMesh.GetEdgeT(Eid);
+				if (EdgeTids.B == IndexConstants::InvalidID || !TrianglesShouldBeInSameComponent(EdgeTids.A, EdgeTids.B))
+				{
+					++NumAttachedBoundaries;
+				}
+
+				if (NumAttachedBoundaries > 2)
+				{
+					return true;
+				}
+			}
+			return false;
+		};
+
+		auto CanVertBeSkipped = [&InputMesh, &VertCanBeSkippedMap, &IsVertMeetingPointOfMultipleBoundaries]
+			(int32 Vid, int32 NextVid, const FVector3d& LastAddedVertPosition)
+		{
+			// See if we've already dealt with this vert
+			bool* CachedCanSkip = VertCanBeSkippedMap.Find(Vid);
+			if (CachedCanSkip)
+			{
+				return *CachedCanSkip;
+			}
+
+			FVector3d PreviousVector = InputMesh.GetVertex(Vid) - LastAddedVertPosition;
+			FVector3d NextVector = InputMesh.GetVertex(NextVid) - InputMesh.GetVertex(Vid);
+
+			// Skippability depends in part on colinearity
+			bool bCanSkip = !PreviousVector.Normalize() || !NextVector.Normalize()
+				// We use abs here because we don't want to keep backwards folding degenerate "fangs" either. This helps
+				// clean up some pathological degenerate soup cases.
+				|| FMath::Abs(NextVector.Dot(PreviousVector)) >= 1 - KINDA_SMALL_NUMBER;
+
+			// However, even if we can skip based on colinearity, we may need to keep the vert if it's something like a T junction
+			if (bCanSkip)
+			{
+				bCanSkip = !IsVertMeetingPointOfMultipleBoundaries(Vid);
+			}
+			
+			VertCanBeSkippedMap.Add(Vid, bCanSkip);
+			return bCanSkip;
+		};
+
+		VidLoopsOut.Reset();
+		for (const FEdgeLoop& Loop : Loops.Loops)
+		{
+			int32 VidLoopIndex = VidLoopsOut.Emplace();
+			TArray<int32>& VidLoop = VidLoopsOut[VidLoopIndex];
+
+			int32 NumVerts = Loop.Vertices.Num();
+			if (!ensure(NumVerts >= 3))
+			{
+				continue;
+			}
+
+			// The initialization of LastAddedVertPosition might not be accurate if we don't later end up adding the 
+			// last vid. If the edge between last and first vid is not degenerate, the mistake won't matter because
+			// the vector to the true previous vid would be colinear with our initialized one. If that edge is degenerate,
+			// then we can account for the situation when processing the last vid, to make sure we add the final vert if needed.
+			FVector3d LastAddedVertPosition = InputMesh.GetVertex(Loop.Vertices[NumVerts - 1]);
+
+			for (int32 i = 0; i < NumVerts - 1; ++i) // Stop right before the last one
+			{
+				int32 Vid = Loop.Vertices[i];
+				int32 NextVid = Loop.Vertices[(i + 1) % NumVerts];
+
+				if (!CanVertBeSkipped(Vid, NextVid, LastAddedVertPosition))
+				{
+					VidLoop.Add(Vid);
+					LastAddedVertPosition = InputMesh.GetVertex(Vid);
+				}
+			}
+
+			// Process the last vid using knowledge of the next unskipped vert.
+			if (VidLoop.Num() > 0 && !CanVertBeSkipped(Loop.Vertices.Last(), VidLoop[0], LastAddedVertPosition))
+			{
+				VidLoop.Add(Loop.Vertices.Last());
+			}
+
+			if (VidLoop.Num() < 3)
+			{
+				// It's not obvious what we should do with an entirely degenerate loop... We choose to throw it away.
+				// The up side of this choice is that this helps a lot in some cases we run into with high-triangle-count
+				// meshes, where preceding simplification steps occasionally create degenerate tris that still end
+				// up getting a nonzero normal that differs from their surroundings, and which create problematic
+				// cracks and flaps in the volume. Removing degenerate loops here ends up removing these cracks and
+				// often improves the result quite a bit.
+				// 
+				// The down side is that in the extreme case, we could end up throwing away a whole mesh if it is
+				// nothing but degenerates, but it seems that having a volume without faces does not crash anything
+				// (it's possible to delete all faces in brush editing mode, for instance). Another drawback/limitation
+				// of the approach here is that it doesn't deal with all degenerate cases, only ones that can be reduced
+				// to a line or point (so bent ribbon of degenerates would end up staying).
+
+				VidLoopsOut.Pop();
+			}
+		}//end for each loop
+	}//end GetVidLoopsWithoutColinearVerts()
+
+}//end DynamicMeshToVolumeLocals
+
 namespace UE {
 namespace Conversion {
 
@@ -158,6 +568,8 @@ static void OptimizeIntoConvexPolys(ABrush* InOwnerBrush, TArray<FPoly>& InPolyg
  */
 static void ApplyLegacyVolumeCleanup(ABrush* brush)
 {
+	using namespace DynamicMeshToVolumeLocals;
+
 	// Remove invalid polygons from the brush
 	for (int32 x = 0; x < brush->Brush->Polys->Element.Num(); ++x)
 	{
@@ -179,7 +591,11 @@ static void ApplyLegacyVolumeCleanup(ABrush* brush)
 			FPoly WkPoly = *Poly;
 			brush->Brush->Polys->Element.RemoveAt(p);
 			TArray<FPoly> Polygons;
-			if (WkPoly.Triangulate(brush, Polygons) > 0)
+
+			// We use our own triangulation function because the legacy one (WkPoly.Triangulate(brush, Polygons)) is not 
+			// very reliable. It deletes colinear verts (which we don't always want) yet still sometimes creates degenerate
+			// triangles that trigger an ensure inside it.
+			if (TriangulatePolygon(WkPoly, Polygons))
 			{
 				OptimizeIntoConvexPolys(brush, Polygons);
 				for (int32 t = 0; t < Polygons.Num(); ++t)
@@ -211,8 +627,14 @@ static void ApplyLegacyVolumeCleanup(ABrush* brush)
 	}
 }
 
-void DynamicMeshToVolume(const FDynamicMesh3& InputMesh, AVolume* TargetVolume, const FMeshToVolumeOptions& Options)
+void GetPolygonFaces(const FDynamicMesh3& InputMesh, const FMeshToVolumeOptions& Options, 
+	TArray<FDynamicMeshFace>& FacesOut, FProgressCancel* Progress)
 {
+	if (Progress && Progress->Cancelled())
+	{
+		return;
+	}
+
 	FDynamicMesh3 LocalMesh;
 	const FDynamicMesh3* UseMesh = &InputMesh;
 	if (Options.bAutoSimplify && InputMesh.TriangleCount() > Options.MaxTriangles)
@@ -220,10 +642,33 @@ void DynamicMeshToVolume(const FDynamicMesh3& InputMesh, AVolume* TargetVolume,
 		LocalMesh = InputMesh;
 		LocalMesh.DiscardAttributes();
 
+		if (Progress && Progress->Cancelled())
+		{
+			return;
+		}
+
 		// collapse to minimal planar first
 		FQEMSimplification PlanarSimplifier(&LocalMesh);
+
+		if (Options.bRespectGroupBoundaries)
+		{
+			FMeshConstraints Constraints;
+			FMeshConstraintsUtil::SetBoundaryConstraintsWithProjection(
+				Constraints,
+				FMeshConstraintsUtil::EBoundaryType::Group,
+				LocalMesh,
+				/*BoundaryCornerAngleThreshold = */ 30);
+			PlanarSimplifier.SetExternalConstraints(Constraints);
+		}
+
 		PlanarSimplifier.SimplifyToMinimalPlanar(0.1);
 
+		if (Progress && Progress->Cancelled())
+		{
+			return;
+		}
+
+		// If we still have too many triangles, don't bother trying to stay planar or preserve groups.
 		if (LocalMesh.TriangleCount() > Options.MaxTriangles)
 		{
 			FVolPresMeshSimplification Simplifier(&LocalMesh);
@@ -231,14 +676,25 @@ void DynamicMeshToVolume(const FDynamicMesh3& InputMesh, AVolume* TargetVolume,
 		}
 
 		UseMesh = &LocalMesh;
+
+		if (Progress && Progress->Cancelled())
+		{
+			return;
+		}
 	}
 
+	GetPolygonFaces(*UseMesh, FacesOut, Options.bRespectGroupBoundaries);
+}
+
+
+void DynamicMeshToVolume(const FDynamicMesh3& InputMesh, AVolume* TargetVolume, const FMeshToVolumeOptions& Options)
+{
 	TArray<FDynamicMeshFace> Faces;
-	GetPolygonFaces(*UseMesh, Faces, Options.bRespectGroupBoundaries);
-	DynamicMeshToVolume(*UseMesh, Faces, TargetVolume);
+	GetPolygonFaces(InputMesh, Options, Faces);
+	DynamicMeshToVolume(InputMesh, Faces, TargetVolume);
 }
 
-void DynamicMeshToVolume(const FDynamicMesh3& InputMesh, TArray<FDynamicMeshFace>& Faces, AVolume* TargetVolume)
+void DynamicMeshToVolume(const FDynamicMesh3&, TArray<FDynamicMeshFace>& Faces, AVolume* TargetVolume)
 {
 	check(TargetVolume->Brush);
 
@@ -357,6 +813,8 @@ void DynamicMeshToVolume(const FDynamicMesh3& InputMesh, TArray<FDynamicMeshFace
 
 void GetPolygonFaces(const FDynamicMesh3& InputMesh, TArray<FDynamicMeshFace>& Faces, bool bRespectGroupBoundaries)
 {
+	using namespace DynamicMeshToVolumeLocals;
+
 	Faces.SetNum(0);
 
 	// We'll find faces using a connected component search based on normals. Note that we give
@@ -369,50 +827,70 @@ void GetPolygonFaces(const FDynamicMesh3& InputMesh, TArray<FDynamicMeshFace>& F
 
 	double NormalTolerance = FMathf::ZeroTolerance;
 
-	FMeshConnectedComponents Components(&InputMesh);
-	Components.FindConnectedTriangles([&InputMesh, &Normals, NormalTolerance, bRespectGroupBoundaries](int32 Triangle0, int32 Triangle1)
+	auto TrianglesShouldBeInSameComponent = [&InputMesh, &Normals, NormalTolerance, bRespectGroupBoundaries](int32 Triangle0, int32 Triangle1)
 	{
 		return (!bRespectGroupBoundaries || InputMesh.GetTriangleGroup(Triangle0) == InputMesh.GetTriangleGroup(Triangle1))
 
 			// This test is only performed if triangles share an edge, so checking the normal is 
 			// sufficient for coplanarity.
 			&& Normals[Triangle0].Dot(Normals[Triangle1]) >= 1 - NormalTolerance;
-	});
+	};
+
+	FMeshConnectedComponents Components(&InputMesh);
+	Components.FindConnectedTriangles(TrianglesShouldBeInSameComponent);
+
+	// Used for removing colinear verts in the loop ahead
+	TMap<int32, bool> VertCanBeSkippedMap;
 
 	for (const FMeshConnectedComponents::FComponent& Component : Components)
 	{
 		FVector3d FaceNormal = Normals[Component.Indices[0]];
 		FMeshRegionBoundaryLoops Loops(&InputMesh, Component.Indices);
 
-		// The FDynamicMeshFaces we emit, and the FPolys we create from them, do
-		// not support multiple boundary loops (eg a face with a hole). So in this
-		// case we fall back to emitting each triangle separately. 
-		// TODO: convex clustering of these triangles?
-		if (Loops.Num() > 1)
+		// Remove colinear verts from the loop where it is safe to do so
+		TArray<TArray<int32>> VidLoops;
+		GetVidLoopsWithoutUnneededColinearVerts(InputMesh, Loops, VertCanBeSkippedMap,
+			TrianglesShouldBeInSameComponent, VidLoops);
+
+		if (VidLoops.Num() > 1)
 		{
-			for (int32 tid : Component.Indices)
+			int32 PreviousNumFaces = Faces.Num();
+			if (!TriangulateLoopsIntoFaces(InputMesh, Component.Indices, FaceNormal, VidLoops, Faces))
 			{
-				FDynamicMeshFace Face;
-				Face.BoundaryLoop.SetNum(3);
-				InputMesh.GetTriVertices(tid, Face.BoundaryLoop[0], Face.BoundaryLoop[1], Face.BoundaryLoop[2]);
-				Algo::Reverse(Face.BoundaryLoop);
-				Face.Plane = FFrame3d(InputMesh.GetTriCentroid(tid), FaceNormal);
-				Faces.Add(Face);
+				// Sanity check to make sure we didn't add new faces despite a failure
+				if (!ensure(Faces.Num() == PreviousNumFaces))
+				{
+					Faces.SetNum(PreviousNumFaces);
+				}
+
+				// Our triangulation function for multiple loops can fail in some self-intersecting cases possible with degenerate
+				// triangles. Revert to outputting all the triangles individually.
+				// TODO: It would be nice to try to lump the triangles into convexes, but it's hard in this general case,
+				// and possibly not worth it.
+				for (int32 tid : Component.Indices)
+				{
+					FDynamicMeshFace Face;
+					Face.BoundaryLoop.SetNum(3);
+					InputMesh.GetTriVertices(tid, Face.BoundaryLoop[0], Face.BoundaryLoop[1], Face.BoundaryLoop[2]);
+					Algo::Reverse(Face.BoundaryLoop);
+					Face.Plane = FFrame3d(InputMesh.GetTriCentroid(tid), FaceNormal);
+					Faces.Add(Face);
+				}
 			}
 		}
-		else
+		else if (VidLoops.Num() == 1) // if polygon with no holes
 		{
-			const FEdgeLoop& Loop = Loops[0];
+			const TArray<int32>& LoopVids = VidLoops[0];
 			FDynamicMeshFace Face;
 
 			FVector3d AvgPos(0, 0, 0);
-			for (int32 vid : Loop.Vertices)
+			for (int32 vid : LoopVids)
 			{
 				FVector3d Position = InputMesh.GetVertex(vid);
 				Face.BoundaryLoop.Add(Position);
 				AvgPos += Position;
 			}
-			AvgPos /= (double)Loop.Vertices.Num();
+			AvgPos /= (double)LoopVids.Num();
 			Algo::Reverse(Face.BoundaryLoop);
 
 			Face.Plane = FFrame3d(AvgPos, FaceNormal);
diff --git a/Engine/Plugins/Runtime/MeshModelingToolset/Source/ModelingComponentsEditorOnly/Public/ConversionUtils/DynamicMeshToVolume.h b/Engine/Plugins/Runtime/MeshModelingToolset/Source/ModelingComponentsEditorOnly/Public/ConversionUtils/DynamicMeshToVolume.h
--- a/Engine/Plugins/Runtime/MeshModelingToolset/Source/ModelingComponentsEditorOnly/Public/ConversionUtils/DynamicMeshToVolume.h
+++ b/Engine/Plugins/Runtime/MeshModelingToolset/Source/ModelingComponentsEditorOnly/Public/ConversionUtils/DynamicMeshToVolume.h
@@ -12,6 +12,7 @@
 // one file).
 
 class AVolume;
+class FProgressCancel;
 PREDECLARE_USE_GEOMETRY_CLASS(FDynamicMesh3);
 
 namespace UE {
@@ -44,6 +45,13 @@ struct FDynamicMeshFace
 	TArray<FVector3d> BoundaryLoop;
 };
 
+
+/**
+ * Gets an array of face objects that can be used to convert a dynamic mesh to a volume, based on given FMeshToVolumeOptions.
+ */
+void MODELINGCOMPONENTSEDITORONLY_API GetPolygonFaces(const FDynamicMesh3& InputMesh, const FMeshToVolumeOptions& Options, 
+	TArray<FDynamicMeshFace>& FacesOut, FProgressCancel* Progress = nullptr);
+
 /**
  * Gets an array of face objects that can be used to convert a dynamic mesh to a volume. This version tries to
  * merge coplanar triangles into polygons.
