Update modeling components usage of opensubdiv to allow max/unix targets, since our OpenSubdiv pre-built libraries now also support these targets

diff --git a/Engine/Plugins/Runtime/MeshModelingToolset/Source/ModelingComponentsEditorOnly/ModelingComponentsEditorOnly.Build.cs b/Engine/Plugins/Runtime/MeshModelingToolset/Source/ModelingComponentsEditorOnly/ModelingComponentsEditorOnly.Build.cs
--- a/Engine/Plugins/Runtime/MeshModelingToolset/Source/ModelingComponentsEditorOnly/ModelingComponentsEditorOnly.Build.cs
+++ b/Engine/Plugins/Runtime/MeshModelingToolset/Source/ModelingComponentsEditorOnly/ModelingComponentsEditorOnly.Build.cs
@@ -43,12 +43,13 @@ public class ModelingComponentsEditorOnly : ModuleRules
 			}
 			);
 
-		if (Target.Platform == UnrealTargetPlatform.Win64)
+		// OpenSubdiv has Windows, Mac and Unix support
+		if (Target.IsInPlatformGroup(UnrealPlatformGroup.Windows) || Target.Platform == UnrealTargetPlatform.Mac || Target.IsInPlatformGroup(UnrealPlatformGroup.Unix))
 		{
 			PrivateDependencyModuleNames.AddRange(
 				new string[]
 				{
-					"OpenSubdiv", // currently Win64-only
+					"OpenSubdiv",
 				}
 				);
 		}
diff --git a/Engine/Plugins/Runtime/MeshModelingToolset/Source/ModelingComponentsEditorOnly/Private/Operations/SubdividePoly.cpp b/Engine/Plugins/Runtime/MeshModelingToolset/Source/ModelingComponentsEditorOnly/Private/Operations/SubdividePoly.cpp
--- a/Engine/Plugins/Runtime/MeshModelingToolset/Source/ModelingComponentsEditorOnly/Private/Operations/SubdividePoly.cpp
+++ b/Engine/Plugins/Runtime/MeshModelingToolset/Source/ModelingComponentsEditorOnly/Private/Operations/SubdividePoly.cpp
@@ -6,8 +6,8 @@
 
 #include UE_INLINE_GENERATED_CPP_BY_NAME(SubdividePoly)
 
-// OpenSubdiv currently only available on Windows. On other platforms we will make this a no-op
-#if PLATFORM_WINDOWS
+// OpenSubdiv currently only available on Windows, Mac and Unix. On other platforms we will make this a no-op
+#if PLATFORM_WINDOWS || PLATFORM_MAC || PLATFORM_UNIX
 #define HAVE_OPENSUBDIV 1
 #else
 #define HAVE_OPENSUBDIV 0
@@ -28,7 +28,7 @@
 #include "opensubdiv/far/primvarRefiner.h"
 #pragma warning(pop)     
 
-#if LOCAL_M_PI
+#ifdef LOCAL_M_PI
 #undef M_PI
 #endif
 
