ModelingOperators: Prevent ensure from firing when a selection is passed to the operation that doesn't consist of at least one full triangle.

diff --git a/Engine/Plugins/Runtime/MeshModelingToolset/Source/ModelingOperators/Private/ParameterizationOps/RecomputeUVsOp.cpp b/Engine/Plugins/Runtime/MeshModelingToolset/Source/ModelingOperators/Private/ParameterizationOps/RecomputeUVsOp.cpp
--- a/Engine/Plugins/Runtime/MeshModelingToolset/Source/ModelingOperators/Private/ParameterizationOps/RecomputeUVsOp.cpp
+++ b/Engine/Plugins/Runtime/MeshModelingToolset/Source/ModelingOperators/Private/ParameterizationOps/RecomputeUVsOp.cpp
@@ -204,7 +204,11 @@ void FRecomputeUVsOp::CollectIslandComponentsPerTile(const FDynamicMeshUVOverlay
 
 			bUseExistingUVTopology = true;
 		}
-		ComponentsPerTile.Emplace(ConnectedComponents, Tiles[TileIndex]);
+
+		if (!ConnectedComponents.Components.IsEmpty())
+		{
+			ComponentsPerTile.Emplace(ConnectedComponents, Tiles[TileIndex]);
+		}
 	}
 }
 
