Texture2DBuilder shouldn't remove mips - this is unneeded, and in the use cases where the builder texture is reused as the final texture, causes issues.

diff --git a/Engine/Plugins/Runtime/MeshModelingToolset/Source/ModelingComponents/Private/AssetUtils/Texture2DBuilder.cpp b/Engine/Plugins/Runtime/MeshModelingToolset/Source/ModelingComponents/Private/AssetUtils/Texture2DBuilder.cpp
--- a/Engine/Plugins/Runtime/MeshModelingToolset/Source/ModelingComponents/Private/AssetUtils/Texture2DBuilder.cpp
+++ b/Engine/Plugins/Runtime/MeshModelingToolset/Source/ModelingComponents/Private/AssetUtils/Texture2DBuilder.cpp
@@ -158,9 +158,6 @@ bool FTexture2DBuilder::InitializeInternal(ETextureType BuildTypeIn, FImageDimen
 		RawTexture2D->SRGB = false;
 		RawTexture2D->LODGroup = TEXTUREGROUP_WorldNormalMap;
 		//RawTexture2D->bFlipGreenChannel = true;
-#if WITH_EDITOR
-		RawTexture2D->MipGenSettings = TMGS_NoMipmaps;
-#endif
 		RawTexture2D->UpdateResource();
 	}
 
