Fix meshes that have 'recompute tangents' enabled not displaying with correct tangents after conversion to dynamic mesh + enable tangents on sculpt tools so 'original material' mode renders correctly when normal maps are used

diff --git a/Engine/Plugins/Experimental/MeshModelingToolsetExp/Source/MeshModelingToolsExp/Private/Sculpting/MeshSculptToolBase.cpp b/Engine/Plugins/Experimental/MeshModelingToolsetExp/Source/MeshModelingToolsExp/Private/Sculpting/MeshSculptToolBase.cpp
--- a/Engine/Plugins/Experimental/MeshModelingToolsetExp/Source/MeshModelingToolsExp/Private/Sculpting/MeshSculptToolBase.cpp
+++ b/Engine/Plugins/Experimental/MeshModelingToolsetExp/Source/MeshModelingToolsExp/Private/Sculpting/MeshSculptToolBase.cpp
@@ -298,7 +298,7 @@ void UMeshSculptToolBase::InitializeSculptMeshComponent(UBaseDynamicMeshComponen
 	Component->RegisterComponent();
 
 	// initialize from LOD-0 MeshDescription
-	Component->SetMesh(UE::ToolTarget::GetDynamicMeshCopy(Target));
+	Component->SetMesh(UE::ToolTarget::GetDynamicMeshCopy(Target, true));
 	double MaxDimension = Component->GetMesh()->GetBounds(true).MaxDim();
 
 	// bake rotation and scaling into mesh because handling these inside sculpting is a mess
diff --git a/Engine/Plugins/Runtime/MeshModelingToolset/Source/ModelingComponents/Private/AssetUtils/MeshDescriptionUtil.cpp b/Engine/Plugins/Runtime/MeshModelingToolset/Source/ModelingComponents/Private/AssetUtils/MeshDescriptionUtil.cpp
--- a/Engine/Plugins/Runtime/MeshModelingToolset/Source/ModelingComponents/Private/AssetUtils/MeshDescriptionUtil.cpp
+++ b/Engine/Plugins/Runtime/MeshModelingToolset/Source/ModelingComponents/Private/AssetUtils/MeshDescriptionUtil.cpp
@@ -14,32 +14,41 @@ void UE::MeshDescription::InitializeAutoGeneratedAttributes(FMeshDescription& Me
 {
 	check(BuildSettings);
 
-	// if the build settings don't indicate calculated normals/tangents, we will assume they are good
-	if (BuildSettings->bRecomputeNormals || BuildSettings->bRecomputeTangents)
+	bool bHasValidNormals = false;
+	bool bHasValidTangents = false;
+	FStaticMeshConstAttributes Attributes(Mesh);
+	TArrayView<const FVector3f> VertexInstanceNormals = Attributes.GetVertexInstanceNormals().GetRawArray();
+	TArrayView<const FVector3f> VertexInstanceTangents = Attributes.GetVertexInstanceTangents().GetRawArray();
+	for (const FVertexInstanceID VertexInstanceID : Mesh.VertexInstances().GetElementIDs())
 	{
-		// check if any are invalid
-		bool bHasInvalidNormals, bHasInvalidTangents;
-		FStaticMeshOperations::AreNormalsAndTangentsValid(Mesh, bHasInvalidNormals, bHasInvalidTangents);
+		bHasValidNormals |= (!VertexInstanceNormals[VertexInstanceID].IsNearlyZero() && !VertexInstanceNormals[VertexInstanceID].ContainsNaN());
+		bHasValidTangents |= (!VertexInstanceTangents[VertexInstanceID].IsNearlyZero() && !VertexInstanceTangents[VertexInstanceID].ContainsNaN());
+		if (bHasValidNormals && bHasValidTangents)
+		{
+			break;
+		}
+	}
 
-		// if neither are invalid we are not going to recompute
-		if (bHasInvalidNormals || bHasInvalidTangents)
+	// rebuild normals/tangents if there are any invalid ones or the build settings request it
+	if (!bHasValidNormals || !bHasValidTangents || BuildSettings->bRecomputeNormals || BuildSettings->bRecomputeTangents)
+	{
+		bool bRecomputeNormals = !bHasValidNormals || BuildSettings->bRecomputeNormals;
+		bool bRecomputeTangents = !bHasValidTangents || BuildSettings->bRecomputeTangents;
+
+		if (!Attributes.GetTriangleNormals().IsValid() || !Attributes.GetTriangleTangents().IsValid())
 		{
-			FStaticMeshAttributes Attributes(Mesh);
-			if (!Attributes.GetTriangleNormals().IsValid() || !Attributes.GetTriangleTangents().IsValid())
-			{
-				// If these attributes don't exist, create them and compute their values for each triangle
-				FStaticMeshOperations::ComputeTriangleTangentsAndNormals(Mesh);
-			}
-
-			EComputeNTBsFlags ComputeNTBsOptions = EComputeNTBsFlags::BlendOverlappingNormals;
-			ComputeNTBsOptions |= BuildSettings->bRecomputeNormals ? EComputeNTBsFlags::Normals : EComputeNTBsFlags::None;
-			ComputeNTBsOptions |= BuildSettings->bRecomputeTangents ? EComputeNTBsFlags::Tangents : EComputeNTBsFlags::None;
-			ComputeNTBsOptions |= BuildSettings->bUseMikkTSpace ? EComputeNTBsFlags::UseMikkTSpace : EComputeNTBsFlags::None;
-			ComputeNTBsOptions |= BuildSettings->bComputeWeightedNormals ? EComputeNTBsFlags::WeightedNTBs : EComputeNTBsFlags::None;
-			ComputeNTBsOptions |= BuildSettings->bRemoveDegenerates ? EComputeNTBsFlags::IgnoreDegenerateTriangles : EComputeNTBsFlags::None;
-
-			FStaticMeshOperations::ComputeTangentsAndNormals(Mesh, ComputeNTBsOptions);
+			// If these attributes don't exist, create them and compute their values for each triangle
+			FStaticMeshOperations::ComputeTriangleTangentsAndNormals(Mesh);
 		}
+
+		EComputeNTBsFlags ComputeNTBsOptions = EComputeNTBsFlags::BlendOverlappingNormals;
+		ComputeNTBsOptions |= bRecomputeNormals ? EComputeNTBsFlags::Normals : EComputeNTBsFlags::None;
+		ComputeNTBsOptions |= bRecomputeTangents ? EComputeNTBsFlags::Tangents : EComputeNTBsFlags::None;
+		ComputeNTBsOptions |= BuildSettings->bUseMikkTSpace ? EComputeNTBsFlags::UseMikkTSpace : EComputeNTBsFlags::None;
+		ComputeNTBsOptions |= BuildSettings->bComputeWeightedNormals ? EComputeNTBsFlags::WeightedNTBs : EComputeNTBsFlags::None;
+		ComputeNTBsOptions |= BuildSettings->bRemoveDegenerates ? EComputeNTBsFlags::IgnoreDegenerateTriangles : EComputeNTBsFlags::None;
+
+		FStaticMeshOperations::ComputeTangentsAndNormals(Mesh, ComputeNTBsOptions);
 	}
 }
 
