make the 'simple collision editor' tool show hidden lines so that it is usable for collision that is hidden inside solid shapes

diff --git a/Engine/Plugins/Runtime/MeshModelingToolset/Source/ModelingComponents/Private/Mechanics/CollisionPrimitivesMechanic.cpp b/Engine/Plugins/Runtime/MeshModelingToolset/Source/ModelingComponents/Private/Mechanics/CollisionPrimitivesMechanic.cpp
--- a/Engine/Plugins/Runtime/MeshModelingToolset/Source/ModelingComponents/Private/Mechanics/CollisionPrimitivesMechanic.cpp
+++ b/Engine/Plugins/Runtime/MeshModelingToolset/Source/ModelingComponents/Private/Mechanics/CollisionPrimitivesMechanic.cpp
@@ -184,6 +184,7 @@ void UCollisionPrimitivesMechanic::SetWorld(UWorld* World)
 	PreviewGeometry->CreateInWorld(World, FTransform());
 
 	DrawnPrimitiveEdges = PreviewGeometry->AddLineSet("CollisionPrimitives");
+	DrawnPrimitiveEdges->SetLineMaterial(ToolSetupUtil::GetDefaultLineComponentMaterial(GetParentTool()->GetToolManager(), false));
 
 }
 
