Make the SubdividePoly operation not crash on failures to add triangles to the output mesh (which can happen if the output would be non-manifold).

Make the SubdividePolyTool only consider adding extra corner vertices on polygroup edges when the corner vertex would be on a mesh boundary; otherwise, the subdivision would try to create a non-manifold edge.

diff --git a/Engine/Plugins/Experimental/MeshModelingToolsetExp/Source/MeshModelingToolsEditorOnlyExp/Private/SubdividePolyTool.cpp b/Engine/Plugins/Experimental/MeshModelingToolsetExp/Source/MeshModelingToolsEditorOnlyExp/Private/SubdividePolyTool.cpp
--- a/Engine/Plugins/Experimental/MeshModelingToolsetExp/Source/MeshModelingToolsEditorOnlyExp/Private/SubdividePolyTool.cpp
+++ b/Engine/Plugins/Experimental/MeshModelingToolsetExp/Source/MeshModelingToolsEditorOnlyExp/Private/SubdividePolyTool.cpp
@@ -219,7 +219,7 @@ void USubdividePolyTool::Setup()
 	Topology = MakeShared<FGroupTopology>(OriginalMesh.Get(), false);
 	auto ShouldAddExtraCornerAtVert = [this](const FGroupTopology& GroupTopology, int32 Vid, const FIndex2i& AttachedGroupEdgeEids)
 	{
-		return Properties->bAddExtraCorners && FGroupTopology::IsEdgeAngleSharp(GroupTopology.GetMesh(), Vid, AttachedGroupEdgeEids, ExtraCornerDotProductThreshold);
+		return Properties->bAddExtraCorners && GroupTopology.GetMesh()->IsBoundaryVertex(Vid) && FGroupTopology::IsEdgeAngleSharp(GroupTopology.GetMesh(), Vid, AttachedGroupEdgeEids, ExtraCornerDotProductThreshold);
 	};
 	Topology->ShouldAddExtraCornerAtVert = ShouldAddExtraCornerAtVert;
 
diff --git a/Engine/Plugins/Experimental/MeshModelingToolsetExp/Source/MeshModelingToolsEditorOnlyExp/Public/SubdividePolyTool.h b/Engine/Plugins/Experimental/MeshModelingToolsetExp/Source/MeshModelingToolsEditorOnlyExp/Public/SubdividePolyTool.h
--- a/Engine/Plugins/Experimental/MeshModelingToolsetExp/Source/MeshModelingToolsEditorOnlyExp/Public/SubdividePolyTool.h
+++ b/Engine/Plugins/Experimental/MeshModelingToolsetExp/Source/MeshModelingToolsEditorOnlyExp/Public/SubdividePolyTool.h
@@ -67,8 +67,8 @@ public:
 	UPROPERTY(EditAnywhere, Category = Settings)
 	bool bRenderCage = true;
 
-	/** When using the group topology for subdivision, whether to add extra corners at sharp group edge bends. */
-	UPROPERTY(EditAnywhere, Category = TopologyOptions, meta = (EditCondition = "SubdivisionScheme != ESubdivisionScheme::Loop"))
+	/** When using the group topology for subdivision, whether to add extra corners at sharp group edge bends on mesh boundaries. Note: We cannot add extra corners on non-boundary group edges, as this would create non-manifold geometry on subdivision. */
+	UPROPERTY(EditAnywhere, Category = TopologyOptions, meta = (DisplayName = "Add Extra Corners on Boundary", EditCondition = "SubdivisionScheme != ESubdivisionScheme::Loop"))
 	bool bAddExtraCorners = true;
 
 	/** How acute an angle between two edges needs to be to add an extra corner there when Add Extra Corners is true. */
diff --git a/Engine/Plugins/Runtime/MeshModelingToolset/Source/ModelingComponentsEditorOnly/Private/Operations/SubdividePoly.cpp b/Engine/Plugins/Runtime/MeshModelingToolset/Source/ModelingComponentsEditorOnly/Private/Operations/SubdividePoly.cpp
--- a/Engine/Plugins/Runtime/MeshModelingToolset/Source/ModelingComponentsEditorOnly/Private/Operations/SubdividePoly.cpp
+++ b/Engine/Plugins/Runtime/MeshModelingToolset/Source/ModelingComponentsEditorOnly/Private/Operations/SubdividePoly.cpp
@@ -935,28 +935,50 @@ bool FSubdividePoly::ComputeSubdividedMesh(FDynamicMesh3& OutMesh)
 			FIndex3i TriB{ Face[2], Face[3], Face[1] };
 			int TriAIndex = OutMesh.AppendTriangle(TriA, GroupID);
 			int TriBIndex = OutMesh.AppendTriangle(TriB, GroupID);
+			bool bHasValidTriA = TriAIndex >= 0;
+			bool bHasValidTriB = TriBIndex >= 0;
+			if (!bHasValidTriA && !bHasValidTriB)
+			{
+				continue;
+			}
 
 			if (bShouldInterpolateUVs)
 			{
 				OpenSubdiv::Far::ConstIndexArray FaceUVIndices = FinalLevel.GetFaceFVarValues(FaceID);
 				
 				FIndex3i UVTriA{ FaceUVIndices[0], FaceUVIndices[1], FaceUVIndices[3] };
-				AddUVTriangleIfValid(UVTriA);
+				if (bHasValidTriA)
+				{
+					AddUVTriangleIfValid(UVTriA);
+				}
 
 				FIndex3i UVTriB{ FaceUVIndices[2], FaceUVIndices[3], FaceUVIndices[1] };
-				AddUVTriangleIfValid(UVTriB);
+				if (bHasValidTriB)
+				{
+					AddUVTriangleIfValid(UVTriB);
+				}
 			}
 
 			if (bHasMaterialIDs)
 			{
-				OutMesh.Attributes()->GetMaterialID()->SetValue(TriAIndex, RefinedMaterialIDs[FaceID]);
-				OutMesh.Attributes()->GetMaterialID()->SetValue(TriBIndex, RefinedMaterialIDs[FaceID]);
+				if (bHasValidTriA)
+				{
+					OutMesh.Attributes()->GetMaterialID()->SetValue(TriAIndex, RefinedMaterialIDs[FaceID]);
+				}
+				if (bHasValidTriB)
+				{
+					OutMesh.Attributes()->GetMaterialID()->SetValue(TriBIndex, RefinedMaterialIDs[FaceID]);
+				}
 			}
 		}
 		else
 		{
 			check(Face.size() == 3);
 			int TriIndex = OutMesh.AppendTriangle(FIndex3i{ Face[0], Face[1], Face[2] }, GroupID);
+			if (TriIndex < 0)
+			{
+				continue;
+			}
 
 			if (bShouldInterpolateUVs)
 			{
