MeshModelingTools: Make depth bias in our materials less sensitive and adjust for FOV, and clean up the material graphs. Add options for dithered and dimmed occluded lines.

diff --git a/Engine/Plugins/Runtime/MeshModelingToolset/Source/ModelingComponents/Private/Drawing/LineSetComponent.cpp b/Engine/Plugins/Runtime/MeshModelingToolset/Source/ModelingComponents/Private/Drawing/LineSetComponent.cpp
--- a/Engine/Plugins/Runtime/MeshModelingToolset/Source/ModelingComponents/Private/Drawing/LineSetComponent.cpp
+++ b/Engine/Plugins/Runtime/MeshModelingToolset/Source/ModelingComponents/Private/Drawing/LineSetComponent.cpp
@@ -81,6 +81,8 @@ public:
 				VertexBuffers.PositionVertexBuffer.VertexPosition(VertexBufferIndex + 2) = (FVector3f)OverlayLine.End;
 				VertexBuffers.PositionVertexBuffer.VertexPosition(VertexBufferIndex + 3) = (FVector3f)OverlayLine.Start;
 
+				// On each end of the line, we store the line direction pointing in opposite directions so that the verts
+				// are moved in oppositie directions to make the rectangle
 				VertexBuffers.StaticMeshVertexBuffer.SetVertexTangents(VertexBufferIndex + 0, FVector3f::ZeroVector, FVector3f::ZeroVector, (FVector3f)-LineDirection);
 				VertexBuffers.StaticMeshVertexBuffer.SetVertexTangents(VertexBufferIndex + 1, FVector3f::ZeroVector, FVector3f::ZeroVector, (FVector3f)-LineDirection);
 				VertexBuffers.StaticMeshVertexBuffer.SetVertexTangents(VertexBufferIndex + 2, FVector3f::ZeroVector, FVector3f::ZeroVector, (FVector3f)LineDirection);
diff --git a/Engine/Plugins/Runtime/MeshModelingToolset/Source/ModelingComponents/Private/ToolSetupUtil.cpp b/Engine/Plugins/Runtime/MeshModelingToolset/Source/ModelingComponents/Private/ToolSetupUtil.cpp
--- a/Engine/Plugins/Runtime/MeshModelingToolset/Source/ModelingComponents/Private/ToolSetupUtil.cpp
+++ b/Engine/Plugins/Runtime/MeshModelingToolset/Source/ModelingComponents/Private/ToolSetupUtil.cpp
@@ -211,7 +211,7 @@ UMaterialInterface* ToolSetupUtil::GetSelectionMaterial(UInteractiveToolManager*
 }
 
 
-UMaterialInterface* ToolSetupUtil::GetSelectionMaterial(const FLinearColor& UseColor, UInteractiveToolManager* ToolManager, float PercentDepthOffset)
+UMaterialInterface* ToolSetupUtil::GetSelectionMaterial(const FLinearColor& UseColor, UInteractiveToolManager* ToolManager, float DepthBias)
 {
 	UMaterialInterface* Material = LoadObject<UMaterial>(nullptr, TEXT("/MeshModelingToolsetExp/Materials/SelectionMaterial"));
 	if (Material == nullptr && ToolManager != nullptr)
@@ -222,9 +222,9 @@ UMaterialInterface* ToolSetupUtil::GetSelectionMaterial(const FLinearColor& UseC
 	{
 		UMaterialInstanceDynamic* MatInstance = UMaterialInstanceDynamic::Create(Material, ToolManager);
 		MatInstance->SetVectorParameterValue(TEXT("ConstantColor"), UseColor);
-		if (PercentDepthOffset != 0)
+		if (DepthBias != 0)
 		{
-			MatInstance->SetScalarParameterValue(TEXT("PercentDepthOffset"), PercentDepthOffset);
+			MatInstance->SetScalarParameterValue(TEXT("DepthBias"), DepthBias);
 		}
 		return MatInstance;
 	}
@@ -257,55 +257,55 @@ MODELINGCOMPONENTS_API UMaterialInstanceDynamic* ToolSetupUtil::GetSimpleCustomM
 	return nullptr;
 }
 
-UMaterialInstanceDynamic* ToolSetupUtil::GetCustomTwoSidedDepthOffsetMaterial(UInteractiveToolManager* ToolManager, const FLinearColor& Color, float PercentDepthOffset, float Opacity)
+UMaterialInstanceDynamic* ToolSetupUtil::GetCustomTwoSidedDepthOffsetMaterial(UInteractiveToolManager* ToolManager, const FLinearColor& Color, float DepthBias, float Opacity)
 {
 	UMaterialInterface* Material = LoadObject<UMaterial>(nullptr, TEXT("/MeshModelingToolsetExp/Materials/SimpleTwoSidedOffsetMaterial_Transparent"));
 	if (Material != nullptr)
 	{
 		UMaterialInstanceDynamic* MatInstance = UMaterialInstanceDynamic::Create(Material, ToolManager);
 		MatInstance->SetVectorParameterValue(TEXT("Color"), Color);
-		MatInstance->SetScalarParameterValue(TEXT("PercentDepthOffset"), PercentDepthOffset);
+		MatInstance->SetScalarParameterValue(TEXT("DepthBias"), DepthBias);
 		MatInstance->SetScalarParameterValue(TEXT("Opacity"), Opacity);
 		return MatInstance;
 	}
 	return nullptr;
 }
 
-MODELINGCOMPONENTS_API UMaterialInstanceDynamic* ToolSetupUtil::GetCustomTwoSidedDepthOffsetMaterial(UInteractiveToolManager* ToolManager, const FLinearColor& Color, float PercentDepthOffset)
+MODELINGCOMPONENTS_API UMaterialInstanceDynamic* ToolSetupUtil::GetCustomTwoSidedDepthOffsetMaterial(UInteractiveToolManager* ToolManager, const FLinearColor& Color, float DepthBias)
 {
 	UMaterialInterface* Material = LoadObject<UMaterial>(nullptr, TEXT("/MeshModelingToolsetExp/Materials/SimpleTwoSidedOffsetMaterial"));
 	if (Material != nullptr)
 	{
 		UMaterialInstanceDynamic* MatInstance = UMaterialInstanceDynamic::Create(Material, ToolManager);
 		MatInstance->SetVectorParameterValue(TEXT("Color"), Color);
-		MatInstance->SetScalarParameterValue(TEXT("PercentDepthOffset"), PercentDepthOffset);
+		MatInstance->SetScalarParameterValue(TEXT("DepthBias"), DepthBias);
 		return MatInstance;
 	}
 	return nullptr;
 }
 
-UMaterialInstanceDynamic* ToolSetupUtil::GetCustomDepthOffsetMaterial(UInteractiveToolManager* ToolManager, const FLinearColor& Color, float PercentDepthOffset, float Opacity)
+UMaterialInstanceDynamic* ToolSetupUtil::GetCustomDepthOffsetMaterial(UInteractiveToolManager* ToolManager, const FLinearColor& Color, float DepthBias, float Opacity)
 {
 	UMaterialInterface* Material = LoadObject<UMaterial>(nullptr, TEXT("/MeshModelingToolsetExp/Materials/SimpleOffsetMaterial_Transparent"));
 	if (Material != nullptr)
 	{
 		UMaterialInstanceDynamic* MatInstance = UMaterialInstanceDynamic::Create(Material, ToolManager);
 		MatInstance->SetVectorParameterValue(TEXT("Color"), Color);
-		MatInstance->SetScalarParameterValue(TEXT("PercentDepthOffset"), PercentDepthOffset);
+		MatInstance->SetScalarParameterValue(TEXT("DepthBias"), DepthBias);
 		MatInstance->SetScalarParameterValue(TEXT("Opacity"), Opacity);
 		return MatInstance;
 	}
 	return nullptr;
 }
 
-MODELINGCOMPONENTS_API UMaterialInstanceDynamic* ToolSetupUtil::GetCustomDepthOffsetMaterial(UInteractiveToolManager* ToolManager, const FLinearColor& Color, float PercentDepthOffset)
+MODELINGCOMPONENTS_API UMaterialInstanceDynamic* ToolSetupUtil::GetCustomDepthOffsetMaterial(UInteractiveToolManager* ToolManager, const FLinearColor& Color, float DepthBias)
 {
 	UMaterialInterface* Material = LoadObject<UMaterial>(nullptr, TEXT("/MeshModelingToolsetExp/Materials/SimpleOffsetMaterial"));
 	if (Material != nullptr)
 	{
 		UMaterialInstanceDynamic* MatInstance = UMaterialInstanceDynamic::Create(Material, ToolManager);
 		MatInstance->SetVectorParameterValue(TEXT("Color"), Color);
-		MatInstance->SetScalarParameterValue(TEXT("PercentDepthOffset"), PercentDepthOffset);
+		MatInstance->SetScalarParameterValue(TEXT("DepthBias"), DepthBias);
 		return MatInstance;
 	}
 	return nullptr;
diff --git a/Engine/Plugins/Runtime/MeshModelingToolset/Source/ModelingComponents/Public/ToolSetupUtil.h b/Engine/Plugins/Runtime/MeshModelingToolset/Source/ModelingComponents/Public/ToolSetupUtil.h
--- a/Engine/Plugins/Runtime/MeshModelingToolset/Source/ModelingComponents/Public/ToolSetupUtil.h
+++ b/Engine/Plugins/Runtime/MeshModelingToolset/Source/ModelingComponents/Public/ToolSetupUtil.h
@@ -119,9 +119,9 @@ namespace ToolSetupUtil
 	MODELINGCOMPONENTS_API UMaterialInterface* GetSelectionMaterial(UInteractiveToolManager* ToolManager);
 
 	/**
-	 * @return Selection Material 1 with custom color and optional depth offset (depth offset moves vertices towards the camera)
+	 * @return Selection Material 1 with custom color and optional depth bias (depth bias moves vertices towards the camera)
 	 */
-	MODELINGCOMPONENTS_API UMaterialInterface* GetSelectionMaterial(const FLinearColor& UseColor, UInteractiveToolManager* ToolManager, float PercentDepthOffset = 0.0f);
+	MODELINGCOMPONENTS_API UMaterialInterface* GetSelectionMaterial(const FLinearColor& UseColor, UInteractiveToolManager* ToolManager, float DepthBias = 0.0f);
 
 	/**
 	 * @return Simple material with configurable color and opacity. Note that the material 
@@ -136,28 +136,28 @@ namespace ToolSetupUtil
 	MODELINGCOMPONENTS_API UMaterialInstanceDynamic* GetSimpleCustomMaterial(UInteractiveToolManager* ToolManager, const FLinearColor& Color);
 
 	 /**
-	  * @return Simple material with configurable depth offset, color, and opacity. Note that the material
+	  * @return Simple material with configurable depth bias, color, and opacity. Note that the material
 	  *  will have translucent blend mode, which can interact poorly with overlapping translucent
 	  *  objects, so use the other overload if you do not need opacity control.
 	  */
-	MODELINGCOMPONENTS_API UMaterialInstanceDynamic* GetCustomDepthOffsetMaterial(UInteractiveToolManager* ToolManager, const FLinearColor& Color, float PercentDepthOffset, float Opacity);
+	MODELINGCOMPONENTS_API UMaterialInstanceDynamic* GetCustomDepthOffsetMaterial(UInteractiveToolManager* ToolManager, const FLinearColor& Color, float DepthBias, float Opacity);
 
 	/**
-	 * @return Simple material with configurable depth offset and color. The material will have opaque blend mode.
+	 * @return Simple material with configurable depth bias and color. The material will have opaque blend mode.
 	 */
-	MODELINGCOMPONENTS_API UMaterialInstanceDynamic* GetCustomDepthOffsetMaterial(UInteractiveToolManager* ToolManager, const FLinearColor& Color, float PercentDepthOffset);
+	MODELINGCOMPONENTS_API UMaterialInstanceDynamic* GetCustomDepthOffsetMaterial(UInteractiveToolManager* ToolManager, const FLinearColor& Color, float DepthBias);
 	
 	/**
-	 * @return Simple two-sided material with configurable depth offset, color, and opacity. Note that 
+	 * @return Simple two-sided material with configurable depth bias, color, and opacity. Note that 
 	  *  the material will have translucent blend mode, which can interact poorly with overlapping translucent
 	  *  objects, so use the other overload if you do not need opacity control.
 	 */
-	MODELINGCOMPONENTS_API UMaterialInstanceDynamic* GetCustomTwoSidedDepthOffsetMaterial(UInteractiveToolManager* ToolManager, const FLinearColor& Color, float PercentDepthOffset, float Opacity);
+	MODELINGCOMPONENTS_API UMaterialInstanceDynamic* GetCustomTwoSidedDepthOffsetMaterial(UInteractiveToolManager* ToolManager, const FLinearColor& Color, float DepthBias, float Opacity);
 
 	/**
-	 * @return Simple two-sided material with configurable depth offset and color. The material will have opaque blend mode.
+	 * @return Simple two-sided material with configurable depth bias and color. The material will have opaque blend mode.
 	 */
-	MODELINGCOMPONENTS_API UMaterialInstanceDynamic* GetCustomTwoSidedDepthOffsetMaterial(UInteractiveToolManager* ToolManager, const FLinearColor& Color, float PercentDepthOffset);
+	MODELINGCOMPONENTS_API UMaterialInstanceDynamic* GetCustomTwoSidedDepthOffsetMaterial(UInteractiveToolManager* ToolManager, const FLinearColor& Color, float DepthBias);
 
 	/**
 	 * @return Material used when editing AVolume objects using our tools.
@@ -172,7 +172,7 @@ namespace ToolSetupUtil
 	 *  Note that the current implementations of the depth-tested and non-depth-tested modes use opaque and translucent 
 	 *  blend modes, respectively, and so inherit their limitations. Specifically, opaque does not support opacity, 
 	 *  and translucent does not always follow correct draw order relative to other translucent objects, which means
-	 *  that depth offset cannot reliably order lines within a non-depth-tested line set component.
+	 *  that depth bias cannot reliably order lines within a non-depth-tested line set component.
 	 */
 	MODELINGCOMPONENTS_API UMaterialInterface* GetDefaultPointComponentMaterial(UInteractiveToolManager* ToolManager, bool bDepthTested = true);
 
@@ -195,7 +195,7 @@ namespace ToolSetupUtil
 	 *  Note that the current implementations of the depth-tested and non-depth-tested modes use opaque and translucent 
 	 *  blend modes, respectively, and so inherit their limitations. Specifically, opaque does not support opacity, 
 	 *  and translucent does not always follow correct draw order relative to other translucent objects, which means
-	 *  that depth offset cannot reliably order lines within a non-depth-tested line set component.
+	 *  that depth bias cannot reliably order lines within a non-depth-tested line set component.
 	 */
 	MODELINGCOMPONENTS_API UMaterialInterface* GetDefaultLineComponentMaterial(UInteractiveToolManager* ToolManager, bool bDepthTested = true);
 
