Fix crashes in BakeRC tool, improved the scene capture progress bar, improved and simplified the scene capture cancellation behavior

* Fixed a crash that occurred in the baking step when cancelling a scene capture which was triggered by changing an option which invalidated all photo sets (i.e., render capture resolution or capture anti aliasing). This was occurring because the baking step proceeded after the scene capture cancellation and the needed photosets were empty.

* Fixed a crash which can happen if you have no channel enabled, change cleanup tolerance from 0 to 1 and then enable the BaseColor channel. This was fixed by changing the behavior of the Cleanup Threshold property, if its positive the DeviceDepth capture is enabled and its disabled otherwise, previously we also required that another channel was enabled which led to the crash in the baking step. The new semantics are easier to understand in the code and and is probably an improvement to the UX as well since the Cleanup Threshold property behavior does not depend on other properties.

* Improved the scene capture cancellation: Users can now cancel the scene capture when the tool is starting up which is better UX since if they picked slow settings they don't intend to reuse, the capture can be cancelled, also, the code is simpler. Changed the behavior of the BakeRC tool properties so that after cancellation the completed captures with the desired settings are enabled. The old behavior where cancellation was documented to restore old results was broken and cannot be made to work unless we store the old results in a separate scene capture but for high capture resolutions that would require a lot of memory. When the SceneCapture is cancelled the channels which were pending and didn't complete are logged in the Output Log in case the user forgets which ones were enabled.

* Simplified the code in various places: Added `operator[](ERenderCaptureType)` to structs with per-capture data, and use it to simplify many places where we previously checked channels individually. Added conversion functions from tool properties/geometry script parameters to the render capture parameters. Removed old confusing logic that determined if a capture had been updated by tracking the number of photos in each photoset with a new system based on a new per-capture pending state. Deprecated a bunch of functions which are no longer useful or which used the wrong types (e.g., `FRenderCaptureUpdate` was added because I hadn't noticed `FRenderCaptureTypeFlags`). Changed the defaults in `FSceneCapturePhotoSet` so that the DeviceDepth capture is off, its an advanced feature and the defaults are intended to make the class friendly to new developers

diff --git a/Engine/Plugins/Experimental/GeometryScripting/Source/GeometryScriptingCore/Private/MeshBakeFunctions.cpp b/Engine/Plugins/Experimental/GeometryScripting/Source/GeometryScriptingCore/Private/MeshBakeFunctions.cpp
--- a/Engine/Plugins/Experimental/GeometryScripting/Source/GeometryScriptingCore/Private/MeshBakeFunctions.cpp
+++ b/Engine/Plugins/Experimental/GeometryScripting/Source/GeometryScriptingCore/Private/MeshBakeFunctions.cpp
@@ -829,40 +829,36 @@ namespace GeometryScriptBakeLocals
 	}
 
 
-	FRenderCaptureOptions MakeRenderCaptureOptions(FGeometryScriptBakeRenderCaptureOptions BakeOptions)
+	FSceneCaptureConfig GetSceneCaptureConfig(FGeometryScriptBakeRenderCaptureOptions BakeOptions)
 	{
-		FRenderCaptureOptions Options;
-
-		Options.bBakeDeviceDepth = BakeOptions.CleanupTolerance > 0 &&
-			(
-			BakeOptions.bBaseColorMap ||
-			BakeOptions.bNormalMap    ||
-			BakeOptions.bEmissiveMap  ||
-			BakeOptions.bOpacityMap ||
-			BakeOptions.bSubsurfaceColorMap ||
-			BakeOptions.bPackedMRSMap ||
-			BakeOptions.bMetallicMap  ||
-			BakeOptions.bRoughnessMap ||
-			BakeOptions.bSpecularMap
-			);
-
-		Options.RenderCaptureImageSize = GeometryScriptBakeLocals::GetDimensions(BakeOptions.RenderCaptureResolution).GetHeight();
-		Options.bAntiAliasing = BakeOptions.bRenderCaptureAntiAliasing;
-		Options.FieldOfViewDegrees = BakeOptions.FieldOfViewDegrees;
-		Options.NearPlaneDist = BakeOptions.NearPlaneDist;
-		Options.bBakeBaseColor = BakeOptions.bBaseColorMap;
-		Options.bBakeEmissive = BakeOptions.bEmissiveMap;
-		Options.bBakeOpacity = BakeOptions.bOpacityMap;
-		Options.bBakeSubsurfaceColor = BakeOptions.bSubsurfaceColorMap;
-		Options.bBakeNormalMap = BakeOptions.bNormalMap;
+		FSceneCaptureConfig Parameters;
 
-		// Enforce the PackedMRS precondition here
-		Options.bUsePackedMRS =  BakeOptions.bPackedMRSMap;
-		Options.bBakeMetallic =  BakeOptions.bPackedMRSMap ? false : BakeOptions.bMetallicMap;
-		Options.bBakeRoughness = BakeOptions.bPackedMRSMap ? false : BakeOptions.bRoughnessMap;
-		Options.bBakeSpecular =  BakeOptions.bPackedMRSMap ? false : BakeOptions.bSpecularMap;
+		Parameters.Flags.bBaseColor = BakeOptions.bBaseColorMap;
+		Parameters.Flags.bWorldNormal = BakeOptions.bNormalMap;
+		Parameters.Flags.bEmissive = BakeOptions.bEmissiveMap;
+		Parameters.Flags.bOpacity = BakeOptions.bOpacityMap;
+		Parameters.Flags.bSubsurfaceColor = BakeOptions.bSubsurfaceColorMap;
+
+		// TODO Maybe its better UX if we ignore the precondition below, it makes sense to do it in the BakeRC tool but
+		// in GS it is easy to add it manually and users may be surprised to not get the textures they requested
 
-		return Options;
+		// Enforce the PackedMRS precondition here
+		Parameters.Flags.bCombinedMRS = BakeOptions.bPackedMRSMap;
+		Parameters.Flags.bMetallic    = BakeOptions.bPackedMRSMap ? false : BakeOptions.bMetallicMap;
+		Parameters.Flags.bRoughness   = BakeOptions.bPackedMRSMap ? false : BakeOptions.bRoughnessMap;
+		Parameters.Flags.bSpecular    = BakeOptions.bPackedMRSMap ? false : BakeOptions.bSpecularMap;
+
+		// Only compute the device depth if we compute at least one other channel, the DeviceDepth is used to eliminate
+		// occlusion artefacts from the other channels
+		Parameters.Flags.bDeviceDepth =
+			(BakeOptions.CleanupTolerance > 0) && (Parameters.Flags != FRenderCaptureTypeFlags::None());
+
+		Parameters.RenderCaptureImageSize = GetDimensions(BakeOptions.RenderCaptureResolution).GetHeight();
+		Parameters.bAntiAliasing = BakeOptions.bRenderCaptureAntiAliasing;
+		Parameters.FieldOfViewDegrees = BakeOptions.FieldOfViewDegrees;
+		Parameters.NearPlaneDist = BakeOptions.NearPlaneDist;
+
+		return Parameters;
 	}
 } // end namespace GeometryScriptBakeLocals
 
@@ -1115,9 +1111,9 @@ FGeometryScriptRenderCaptureTextures UGeometryScriptLibrary_MeshBakeFunctions::B
 		return {};
 	}
 
-	const FRenderCaptureOptions Options = GeometryScriptBakeLocals::MakeRenderCaptureOptions(BakeOptions);
-	FRenderCaptureUpdate UnusedUpdate;
-	TUniquePtr<FSceneCapturePhotoSet> SceneCapture = CapturePhotoSet(ValidSourceActors, Options, UnusedUpdate, false);
+	const FSceneCaptureConfig Options = GeometryScriptBakeLocals::GetSceneCaptureConfig(BakeOptions);
+	TUniquePtr<FSceneCapturePhotoSet> SceneCapture = MakeUnique<FSceneCapturePhotoSet>();
+	ConfigureSceneCapture(SceneCapture, ValidSourceActors, Options, false);
 
 	const FDynamicMeshAABBTree3 TargetMeshSpatial(TargetMesh->GetMeshPtr());
 	TSharedPtr<FMeshTangentsd, ESPMode::ThreadSafe> TargetMeshTangents = MakeShared<FMeshTangentsd>(TargetMesh->GetMeshPtr());
@@ -1134,7 +1130,7 @@ FGeometryScriptRenderCaptureTextures UGeometryScriptLibrary_MeshBakeFunctions::B
 
 	FRenderCaptureOcclusionHandler OcclusionHandler(GeometryScriptBakeLocals::GetDimensions(BakeOptions.Resolution));
 
-	const FRenderCaptureOptions PendingBake = Options; // All specified channels need baking
+	const FRenderCaptureTypeFlags PendingBake = Options.Flags; // All specified channels need baking
 	TUniquePtr<FMeshMapBaker> Baker = MakeRenderCaptureBaker(
 		TargetMesh->GetMeshPtr(),
 		TargetMeshTangents,
diff --git a/Engine/Plugins/Experimental/MeshModelingToolsetExp/Source/MeshModelingToolsEditorOnlyExp/Private/BakeRenderCaptureTool.cpp b/Engine/Plugins/Experimental/MeshModelingToolsetExp/Source/MeshModelingToolsEditorOnlyExp/Private/BakeRenderCaptureTool.cpp
--- a/Engine/Plugins/Experimental/MeshModelingToolsetExp/Source/MeshModelingToolsEditorOnlyExp/Private/BakeRenderCaptureTool.cpp
+++ b/Engine/Plugins/Experimental/MeshModelingToolsetExp/Source/MeshModelingToolsEditorOnlyExp/Private/BakeRenderCaptureTool.cpp
@@ -44,35 +44,124 @@ static FString NormalTexParamName = TEXT("NormalMap");
 static FString PackedMRSTexParamName = TEXT("PackedMRS");
 
 
-FRenderCaptureOptions
-MakeRenderCaptureOptions(const URenderCaptureProperties& RenderCaptureProperties)
+FSceneCaptureConfig GetSceneCaptureConfig(const URenderCaptureProperties& Properties)
 {
-	FRenderCaptureOptions Options;
-
-	Options.bBakeBaseColor = RenderCaptureProperties.bBaseColorMap;
-	Options.bBakeNormalMap = RenderCaptureProperties.bNormalMap;
-	Options.bBakeEmissive =  RenderCaptureProperties.bEmissiveMap;
-	Options.bBakeOpacity = RenderCaptureProperties.bOpacityMap;
-	Options.bBakeSubsurfaceColor = RenderCaptureProperties.bSubsurfaceColorMap;
-	Options.bBakeDeviceDepth = RenderCaptureProperties.bDeviceDepthMap;
-	
-	// Enforce the PackedMRS precondition here so we don't have to check it at each usage site.  Note: We don't
-	// apply this precondition on the RenderCaptureProperties because we don't want the user to have to re-enable
-	// options which enabling PackedMRS disabled.
-	Options.bUsePackedMRS =  RenderCaptureProperties.bPackedMRSMap;
-	Options.bBakeMetallic =  RenderCaptureProperties.bPackedMRSMap ? false : RenderCaptureProperties.bMetallicMap;
-	Options.bBakeRoughness = RenderCaptureProperties.bPackedMRSMap ? false : RenderCaptureProperties.bRoughnessMap;
-	Options.bBakeSpecular =  RenderCaptureProperties.bPackedMRSMap ? false : RenderCaptureProperties.bSpecularMap;
-
-	Options.RenderCaptureImageSize = static_cast<int32>(RenderCaptureProperties.Resolution);
-	Options.bAntiAliasing = RenderCaptureProperties.bAntiAliasing;
-	Options.FieldOfViewDegrees = RenderCaptureProperties.CaptureFieldOfView;
-	Options.NearPlaneDist = RenderCaptureProperties.NearPlaneDist;
-
-	return Options;
+	FSceneCaptureConfig Config;
+
+	Config.Flags.bBaseColor       = Properties.bBaseColorMap;
+	Config.Flags.bWorldNormal     = Properties.bNormalMap;
+	Config.Flags.bEmissive        = Properties.bEmissiveMap;
+	Config.Flags.bOpacity         = Properties.bOpacityMap;
+	Config.Flags.bSubsurfaceColor = Properties.bSubsurfaceColorMap;
+	Config.Flags.bDeviceDepth     = Properties.bDeviceDepthMap;
+
+	// When we convert Properties (controls the tool UI) to Config (controls what is computed in SceneCapture) we make
+	// sure that the PackedMRS and the separate Metallic/Roughness/Specular captures are mutually exclusive. These
+	// captures are intended to be used in different materials so we assume the user will want one or the other. Also,
+	// for high resolution captures with lots of viewpoints each enabled capture takes up a lot of memory. We could
+	// enforce this mutually exclusive behaviour in Properties but doing it Config improves the UX because enabling the
+	// PackedMRS checkbox can just clear the Metallic/Roughness/Specular textures and photosets and disable the
+	// corresponding checkboxes while leaving the checkbox state unchanged. See :HandlingPackedMRS
+	Config.Flags.bCombinedMRS = Properties.bPackedMRSMap;
+	Config.Flags.bMetallic    = Properties.bPackedMRSMap ? false : Properties.bMetallicMap;
+	Config.Flags.bRoughness   = Properties.bPackedMRSMap ? false : Properties.bRoughnessMap;
+	Config.Flags.bSpecular    = Properties.bPackedMRSMap ? false : Properties.bSpecularMap;
+
+	Config.RenderCaptureImageSize = static_cast<int32>(Properties.Resolution);
+	Config.bAntiAliasing          = Properties.bAntiAliasing;
+	Config.FieldOfViewDegrees     = Properties.CaptureFieldOfView;
+	Config.NearPlaneDist          = Properties.NearPlaneDist;
+
+	return Config;
+}
+
+void SetSceneCaptureConfig(URenderCaptureProperties& Properties, const FSceneCaptureConfig& Config)
+{
+	Properties.bBaseColorMap       = Config.Flags.bBaseColor;
+	Properties.bNormalMap          = Config.Flags.bWorldNormal;
+	Properties.bEmissiveMap        = Config.Flags.bEmissive;
+	Properties.bOpacityMap         = Config.Flags.bOpacity;
+	Properties.bSubsurfaceColorMap = Config.Flags.bSubsurfaceColor;
+	Properties.bDeviceDepthMap     = Config.Flags.bDeviceDepth;
+
+	// When we convert Config (controls what is computed in SceneCapture) to Properties (controls the tool UI) we
+	// can directly copy the MRS booleans since Properties must always reflect the state of the computed SceneCapture
+	// channels. The reverse conversion in GetSceneCaptureConfig (see comment tagged :HandlingPackedMRS) ensures that
+	// the PackedMRS and separate Metallic/Roughness/Specular captures are mutually exclusive.
+	Properties.bPackedMRSMap = Config.Flags.bCombinedMRS;
+	Properties.bMetallicMap  = Config.Flags.bMetallic;
+	Properties.bRoughnessMap = Config.Flags.bRoughness;
+	Properties.bSpecularMap  = Config.Flags.bSpecular;
+
+	Properties.Resolution         = static_cast<EBakeTextureResolution>(Config.RenderCaptureImageSize);
+	Properties.bAntiAliasing      = Config.bAntiAliasing;
+	Properties.CaptureFieldOfView = Config.FieldOfViewDegrees;
+	Properties.NearPlaneDist      = Config.NearPlaneDist;
 }
 
 
+//
+// Tool Properties
+//
+
+bool UBakeRenderCaptureResults::IsEmpty() const
+{
+	bool bEmpty = true;
+	ForEachCaptureType([this, &bEmpty](ERenderCaptureType CaptureType)
+	{
+		bEmpty = bEmpty && (this->operator[](CaptureType) == nullptr);
+	});
+	return bEmpty;
+}
+
+const TObjectPtr<UTexture2D>& UBakeRenderCaptureResults::operator[](ERenderCaptureType CaptureType) const
+{
+	switch (CaptureType)
+	{
+	case ERenderCaptureType::BaseColor:
+		return BaseColorMap;
+	case ERenderCaptureType::WorldNormal:
+		return NormalMap;
+	case ERenderCaptureType::Roughness:
+		return RoughnessMap;
+	case ERenderCaptureType::Metallic:
+		return MetallicMap;
+	case ERenderCaptureType::Specular:
+		return SpecularMap;
+	case ERenderCaptureType::Emissive:
+		return EmissiveMap;
+	case ERenderCaptureType::CombinedMRS:
+		return PackedMRSMap;
+	case ERenderCaptureType::Opacity:
+		return OpacityMap;
+	case ERenderCaptureType::SubsurfaceColor:
+		return SubsurfaceColorMap;
+	case ERenderCaptureType::DeviceDepth:
+		ensure(DeviceDepthMap == nullptr); // DeviceDepth is unused and shouldn't change from the default
+		return DeviceDepthMap;
+	default:
+		ensure(false);
+	}
+	return DeviceDepthMap;
+}
+
+bool URenderCaptureProperties::operator==(const URenderCaptureProperties& Other) const
+{
+	const FSceneCaptureConfig Config = GetSceneCaptureConfig(*this);
+	const FSceneCaptureConfig OtherConfig = GetSceneCaptureConfig(Other);
+	return Config == OtherConfig;
+}
+
+bool URenderCaptureProperties::operator!=(const URenderCaptureProperties& Other) const
+{
+	return !(*this == Other);
+}
+
+
+
+
+
+
 //
 // Tool Operator
 //
@@ -91,7 +180,8 @@ public:
 	FSceneCapturePhotoSet* SceneCapture = nullptr;
 
 	// Used to pass the channels which need baking via the bBakeXXX and bUsePackedMRS members
-	FRenderCaptureOptions PendingBake;
+	// PendingBake allows us to skip baking for computed capture types previously baked
+	FRenderCaptureTypeFlags PendingBake;
 
 	// Begin TGenericDataOperator interface
 	virtual void CalculateResult(FProgressCancel* Progress) override;
@@ -227,21 +317,9 @@ void UBakeRenderCaptureTool::Setup()
 	Settings->WatchProperty(Settings->TextureSize, [this](EBakeTextureResolution) { OpState |= EBakeOpState::Evaluate; });
 	Settings->WatchProperty(Settings->ValidSampleDepthThreshold, [this](float ValidSampleDepthThreshold)
 	{
-		// Only compute the device depth if we compute at least one other channel, the DeviceDepth is used to eliminate
-		// occlusion artefacts from the other channels
-		RenderCaptureProperties->bDeviceDepthMap = (ValidSampleDepthThreshold > 0) &&
-			(
-			RenderCaptureProperties->bBaseColorMap ||
-			RenderCaptureProperties->bNormalMap    ||
-			RenderCaptureProperties->bEmissiveMap  ||
-			RenderCaptureProperties->bOpacityMap   ||
-			RenderCaptureProperties->bSubsurfaceColorMap ||
-			RenderCaptureProperties->bPackedMRSMap ||
-			RenderCaptureProperties->bMetallicMap  ||
-			RenderCaptureProperties->bRoughnessMap ||
-			RenderCaptureProperties->bSpecularMap
-			);
-
+		// The depth capture channel is enabled implicitly when the following parameter is greater than 0.
+		// To disable the depth capture we need to set the value to 0. See :EnableDisableDeviceDepthMap 
+		RenderCaptureProperties->bDeviceDepthMap = (ValidSampleDepthThreshold > 0);
 		OpState |= EBakeOpState::Evaluate;
 	});
 
@@ -303,6 +381,8 @@ void UBakeRenderCaptureTool::Setup()
 		UE::ToolTarget::HideSourceObject(Targets[Idx]);
 	}
 
+	SceneCapture = MakeUnique<FSceneCapturePhotoSet>();
+
 	// Make sure we trigger SceneCapture computation in UpdateResult
 	OpState |= EBakeOpState::Evaluate;
 
@@ -396,13 +476,6 @@ void UBakeRenderCaptureTool::OnShutdown(EToolShutdownType ShutdownType)
 		Compute->Shutdown();
 	}
 
-	// Restore visibility of source meshes
-	const int NumTargets = Targets.Num();
-	for (int Idx = 1; Idx < NumTargets; ++Idx)
-	{
-		UE::ToolTarget::ShowSourceObject(Targets[Idx]);
-	}
-	
 	if (ShutdownType == EToolShutdownType::Accept && ResultSettings->IsEmpty() == false)
 	{
 		// TODO Support skeletal meshes here---see BakeMeshAttributeMapsTool::OnShutdown
@@ -415,7 +488,15 @@ void UBakeRenderCaptureTool::OnShutdown(EToolShutdownType ShutdownType)
 	// Clear actors on shutdown so that their lifetime is not tied to the lifetime of the tool
 	Actors.Empty();
 
+	// Restore visibility of the target mesh
 	UE::ToolTarget::ShowSourceObject(Targets[0]);
+
+	// Restore visibility of source meshes
+	const int NumTargets = Targets.Num();
+	for (int Idx = 1; Idx < NumTargets; ++Idx)
+	{
+		UE::ToolTarget::ShowSourceObject(Targets[Idx]);
+	}
 }
 
 void UBakeRenderCaptureTool::CreateAssets(UWorld* SourceWorld, UObject* SourceAsset)
@@ -644,6 +725,7 @@ void UBakeRenderCaptureTool::CreateAssets(UWorld* SourceWorld, UObject* SourceAs
 
 
 
+// Return false if the user requested a texture but it is not yet baked, or if the tool is in an invalid state
 bool UBakeRenderCaptureTool::CanAccept() const
 {
 	if ((OpState & EBakeOpState::Invalid) == EBakeOpState::Invalid)
@@ -671,11 +753,12 @@ bool UBakeRenderCaptureTool::CanAccept() const
 	{
 		return false;
 	}
-	
-	// We need different code paths based on PackedMRS here because we don't want to uncheck the separate channels
-	// when PackedMRS is enabled to give the user a better UX (they don't have to re-check them after disabling
-	// PackedMRS). In other place we can test the PackedMRS and separate channel booleans in series and avoid the
-	// complexity of nested if statements.
+
+	// For reasons explained in the comment tagged :HandlingPackedMRS the Metallic/Roughness/Specular checkboxes can
+	// be enabled when the correponding photo sets in SceneCapture and textures in ResultSettings are empty. This means
+	// we should only check the checkboxes/textures (RenderCaptureProperties/ResultSettings resp.) if we are not using
+	// the PackedMRS option In other places we can test the PackedMRS and separate channel booleans in series and avoid
+	// the complexity of nested if statements.
 	if (RenderCaptureProperties->bPackedMRSMap)
 	{
 		if (ResultSettings->PackedMRSMap == nullptr)
@@ -704,6 +787,7 @@ bool UBakeRenderCaptureTool::CanAccept() const
 
 
 
+// This function gets called via UBakeRenderCaptureTool::InvalidateCompute
 TUniquePtr<TGenericDataOperator<FMeshMapBaker>> UBakeRenderCaptureTool::MakeNewOperator()
 {
 	// We should not have requested a bake if we don't have a SceneCapture
@@ -720,22 +804,17 @@ TUniquePtr<TGenericDataOperator<FMeshMapBaker>> UBakeRenderCaptureTool::MakeNewO
 	Op->SamplesPerPixel = Settings->SamplesPerPixel;
 	Op->SceneCapture = SceneCapture.Get();
 
-	auto IsPendingBake = [this](UTexture *CurrentResult, ERenderCaptureType CaptureType) -> bool
+	ForEachCaptureType([this, &Op](ERenderCaptureType CaptureType)
 	{
-		const bool bRequested = SceneCapture->GetCaptureTypeEnabled(CaptureType);
-		const bool bBaked = (CurrentResult != nullptr);
-		return (bRequested && !bBaked);
-	};
+		using EStatus = FSceneCapturePhotoSet::ECaptureTypeStatus;
 
-	Op->PendingBake.bBakeBaseColor = IsPendingBake(ResultSettings->BaseColorMap, ERenderCaptureType::BaseColor);
-	Op->PendingBake.bBakeRoughness = IsPendingBake(ResultSettings->RoughnessMap, ERenderCaptureType::Roughness);
-	Op->PendingBake.bBakeMetallic =  IsPendingBake(ResultSettings->MetallicMap,  ERenderCaptureType::Metallic);
-	Op->PendingBake.bBakeSpecular =  IsPendingBake(ResultSettings->SpecularMap,  ERenderCaptureType::Specular);
-	Op->PendingBake.bUsePackedMRS =  IsPendingBake(ResultSettings->PackedMRSMap, ERenderCaptureType::CombinedMRS);
-	Op->PendingBake.bBakeEmissive =  IsPendingBake(ResultSettings->EmissiveMap,  ERenderCaptureType::Emissive);
-	Op->PendingBake.bBakeNormalMap = IsPendingBake(ResultSettings->NormalMap,    ERenderCaptureType::WorldNormal);
-	Op->PendingBake.bBakeOpacity =   IsPendingBake(ResultSettings->OpacityMap,   ERenderCaptureType::Opacity);
-	Op->PendingBake.bBakeSubsurfaceColor = IsPendingBake(ResultSettings->SubsurfaceColorMap, ERenderCaptureType::SubsurfaceColor);
+		const bool bComputed = SceneCapture->GetCaptureTypeStatus(CaptureType) == EStatus::Computed;
+		const bool bBaked = (*ResultSettings)[CaptureType] != nullptr;
+		const bool bPendingBake = bComputed && !bBaked;
+
+		// The BakeRC tool does not bake the DeviceDepth, this channel is queried via the SceneCapture when baking the other channels
+		Op->PendingBake[CaptureType] = (CaptureType == ERenderCaptureType::DeviceDepth ? false : bPendingBake);
+	});
 
 	return Op;
 }
@@ -750,8 +829,7 @@ void UBakeRenderCaptureTool::OnMapsUpdated(const TUniquePtr<FMeshMapBaker>& NewR
 	GetTexturesFromRenderCaptureBaker(NewResult, TexturesOut);
 
 	// The NewResult will contain the newly baked textures so we only update those and not overwrite any already baked
-	// valid results. Note that if results are invalidated by some parameter change they are cleared by the
-	// InvalidateResults function, by comparing/syncing with the enabled capture types in SceneCapture
+	// valid TexturesOut. If a texture is invalidated by some tool property change it will null on entry to this function
 	if (TexturesOut.BaseColorMap)
 	{
 		ResultSettings->BaseColorMap = TexturesOut.BaseColorMap;
@@ -971,66 +1049,7 @@ void UBakeRenderCaptureTool::InvalidateCompute()
 	OpState = EBakeOpState::Clean;
 }
 
-FRenderCaptureUpdate UBakeRenderCaptureTool::UpdateSceneCapture()
-{
-	for (int Idx = 1; Idx < Targets.Num(); ++Idx)
-	{
-		UE::ToolTarget::ShowSourceObject(Targets[Idx]);
-	}
-
-	FRenderCaptureUpdate Update;
-
-	if (!SceneCapture)
-	{
-		const FRenderCaptureOptions Options = MakeRenderCaptureOptions(*RenderCaptureProperties);
-
-		constexpr bool bAllowCancel = false;
-		SceneCapture = CapturePhotoSet(Actors, Options, Update, bAllowCancel);
-	}
-	else
-	{
-		// Get already computed options from the existing SceneCapture so we can restore them if the capture is cancelled
-		const FRenderCaptureOptions ComputedOptions = GetComputedPhotoSetOptions(SceneCapture);
-
-		const FRenderCaptureOptions Options = MakeRenderCaptureOptions(*RenderCaptureProperties);
-
-		constexpr bool bAllowCancel = true;
-		Update = UpdatePhotoSets(SceneCapture, Actors, Options, bAllowCancel);
-
-		if (SceneCapture->Cancelled())
-		{
-			// Restore the settings present before the change that invoked the scene capture recompute
-			// Note UpdatePhotoSets will have restored the SceneCapture settings to their values prior to the cancel
-			RenderCaptureProperties->bBaseColorMap      = ComputedOptions.bBakeBaseColor;
-			RenderCaptureProperties->bNormalMap         = ComputedOptions.bBakeNormalMap;
-			RenderCaptureProperties->bMetallicMap       = ComputedOptions.bBakeMetallic;
-			RenderCaptureProperties->bRoughnessMap      = ComputedOptions.bBakeRoughness;
-			RenderCaptureProperties->bSpecularMap       = ComputedOptions.bBakeSpecular;
-			RenderCaptureProperties->bPackedMRSMap      = ComputedOptions.bUsePackedMRS;
-			RenderCaptureProperties->bEmissiveMap       = ComputedOptions.bBakeEmissive;
-			RenderCaptureProperties->bOpacityMap        = ComputedOptions.bBakeOpacity;
-			RenderCaptureProperties->bSubsurfaceColorMap = ComputedOptions.bBakeSubsurfaceColor;
-			RenderCaptureProperties->bDeviceDepthMap    = ComputedOptions.bBakeDeviceDepth;
-			RenderCaptureProperties->bAntiAliasing      = ComputedOptions.bAntiAliasing;
-			RenderCaptureProperties->CaptureFieldOfView = ComputedOptions.FieldOfViewDegrees;
-			RenderCaptureProperties->NearPlaneDist      = ComputedOptions.NearPlaneDist;
-			RenderCaptureProperties->Resolution         = static_cast<EBakeTextureResolution>(ComputedOptions.RenderCaptureImageSize);
-
-			// Silently make the above updates so we don't overwrite the change to OpState below and call this function again
-			RenderCaptureProperties->SilentUpdateWatched();
-			Settings->SilentUpdateWatched();
-		}
-	}
-
-	for (int Idx = 1; Idx < Targets.Num(); ++Idx)
-	{
-		UE::ToolTarget::HideSourceObject(Targets[Idx]);
-	}
-
-	return Update;
-}
-
-// Process dirty props and update background compute. Called by Render
+// Process dirty props and update background compute. Called by UBakeRenderCaptureTool::Render
 void UBakeRenderCaptureTool::UpdateResult()
 {
 	// Return if the bake is already launched/complete.
@@ -1098,7 +1117,7 @@ void UBakeRenderCaptureTool::UpdateResult()
 	if (bIsInvalid)
 	{
 		// Clear all results and wait for user to fix the invalid tool inputs
-		InvalidateResults(FRenderCaptureUpdate{});
+		InvalidateResults(FRenderCaptureTypeFlags::All(true));
 
 		// Only call UpdateVisualization when we first detect the invalid inputs
 		const bool bWasValid = static_cast<bool>(OpState & EBakeOpState::Invalid) == false;
@@ -1112,27 +1131,69 @@ void UBakeRenderCaptureTool::UpdateResult()
 	}
 	else
 	{
-		FRenderCaptureUpdate Update = UpdateSceneCapture();
+		const FSceneCaptureConfig DesiredConfig = GetSceneCaptureConfig(*RenderCaptureProperties);
+
+		// Update the scene capture and get the capture types that were updated so we can bake them
+		FRenderCaptureTypeFlags UpdatedChannels;
+		{
+			// Show the source meshes. If we don't do this the renderer doesn't see anything and we get a blank capture
+			// TODO FSceneCapturePhotoSet should probably ensure the visibility its Actors so that we don't need to do
+			// this visibility thing here in the tool
+			for (int Idx = 1; Idx < Targets.Num(); ++Idx)
+			{
+				UE::ToolTarget::ShowSourceObject(Targets[Idx]);
+			}
+
+			UpdatedChannels = UpdateSceneCapture(SceneCapture, Actors, DesiredConfig, false);
+
+			// Hide the source meshes after the render capture so they don't occlude the preview
+			for (int Idx = 1; Idx < Targets.Num(); ++Idx)
+			{
+				UE::ToolTarget::HideSourceObject(Targets[Idx]);
+			}
 
-		const bool bInvalidateAll =
-			ComputedTextureSize != Settings->TextureSize ||
-			ComputedSamplesPerPixel != Settings->SamplesPerPixel ||
-			ComputedValidDepthThreshold != Settings->ValidSampleDepthThreshold;
+			const FSceneCaptureConfig AchievedConfig = GetSceneCaptureConfig(SceneCapture);
+			ensure(SceneCapture->Cancelled() == (AchievedConfig != DesiredConfig));
 
-		// Invalidate any results which need re-baking because the SceneCapture was updated, or
-		// Invalidate all results if the any baking parameters were changed
-		InvalidateResults(bInvalidateAll ? FRenderCaptureUpdate{} : Update);
+			// If the scene capture was cancelled make sure the tool properties are consistent with the computed captures
+			if (SceneCapture->Cancelled())
+			{
+				SetSceneCaptureConfig(*RenderCaptureProperties, AchievedConfig);
+				RenderCaptureProperties->SilentUpdateWatched();
 
-		// Update the preview mesh material with the (possibly invalidated) bake results
-		UpdateVisualization();
+				if (AchievedConfig.Flags.bDeviceDepth == false)
+				{
+					// See :EnableDisableDeviceDepthMap
+					Settings->ValidSampleDepthThreshold = 0.f;
+					Settings->SilentUpdateWatched();
+				}
+			}
+		}
 
-		// Start another bake operation
-		InvalidateCompute();
+		// Update/Bake the result textures
+		{
+			const bool bInvalidateAll =
+				ComputedTargetUVLayer != InputMeshSettings->TargetUVLayer ||
+				ComputedTextureSize != Settings->TextureSize ||
+				ComputedSamplesPerPixel != Settings->SamplesPerPixel ||
+				ComputedValidDepthThreshold != Settings->ValidSampleDepthThreshold;
 
-		// Cache computed parameters which are used to determine if results need re-baking
-		ComputedTextureSize = Settings->TextureSize;
-		ComputedSamplesPerPixel = Settings->SamplesPerPixel;
-		ComputedValidDepthThreshold = Settings->ValidSampleDepthThreshold;
+			// Invalidate any results corresponding to SceneCapture Channels that got updated, or
+			// Invalidate all results if the any baking parameters were changed
+			InvalidateResults(bInvalidateAll ? FRenderCaptureTypeFlags::All(true) : UpdatedChannels);
+
+			// Update the preview mesh material with the (possibly invalidated) bake results
+			UpdateVisualization();
+
+			// Start another bake operation, this will bake the computed captures with invalidated/null texture results
+			InvalidateCompute();
+
+			// Cache computed parameters which are used to determine if results need re-baking
+			ComputedTargetUVLayer = InputMeshSettings->TargetUVLayer;
+			ComputedTextureSize = Settings->TextureSize;
+			ComputedSamplesPerPixel = Settings->SamplesPerPixel;
+			ComputedValidDepthThreshold = Settings->ValidSampleDepthThreshold;
+		}
 	}
 }
 
@@ -1225,44 +1286,44 @@ void UBakeRenderCaptureTool::UpdateVisualization()
 }
 
 
-void UBakeRenderCaptureTool::InvalidateResults(FRenderCaptureUpdate Update)
+void UBakeRenderCaptureTool::InvalidateResults(FRenderCaptureTypeFlags Invalidate)
 {
 	// Note that the bake operation, Compute, updates ResultSettings when results are available via the
 	// Compute->OnResultUpdated delegate.
 
-	if (Update.bUpdatedBaseColor)
+	if (Invalidate.bBaseColor)
 	{
 		ResultSettings->BaseColorMap = nullptr;
 	}
-	if (Update.bUpdatedRoughness)
+	if (Invalidate.bRoughness)
 	{
 		ResultSettings->RoughnessMap = nullptr;
 	}
-	if (Update.bUpdatedMetallic)
+	if (Invalidate.bMetallic)
 	{
 		ResultSettings->MetallicMap = nullptr;
 	}
-	if (Update.bUpdatedSpecular)
+	if (Invalidate.bSpecular)
 	{
 		ResultSettings->SpecularMap = nullptr;
 	}
-	if (Update.bUpdatedEmissive)
+	if (Invalidate.bEmissive)
 	{
 		ResultSettings->EmissiveMap = nullptr;
 	}
-	if (Update.bUpdatedNormalMap)
+	if (Invalidate.bWorldNormal)
 	{
 		ResultSettings->NormalMap = nullptr;
 	}
-	if (Update.bUpdatedOpacity)
+	if (Invalidate.bOpacity)
 	{
 		ResultSettings->OpacityMap = nullptr;
 	}
-	if (Update.bUpdatedSubsurfaceColor)
+	if (Invalidate.bSubsurfaceColor)
 	{
 		ResultSettings->SubsurfaceColorMap = nullptr;
 	}
-	if (Update.bUpdatedPackedMRS)
+	if (Invalidate.bCombinedMRS)
 	{
 		ResultSettings->PackedMRSMap = nullptr;
 	}
diff --git a/Engine/Plugins/Experimental/MeshModelingToolsetExp/Source/MeshModelingToolsEditorOnlyExp/Public/BakeRenderCaptureTool.h b/Engine/Plugins/Experimental/MeshModelingToolsetExp/Source/MeshModelingToolsEditorOnlyExp/Public/BakeRenderCaptureTool.h
--- a/Engine/Plugins/Experimental/MeshModelingToolsetExp/Source/MeshModelingToolsEditorOnlyExp/Public/BakeRenderCaptureTool.h
+++ b/Engine/Plugins/Experimental/MeshModelingToolsetExp/Source/MeshModelingToolsEditorOnlyExp/Public/BakeRenderCaptureTool.h
@@ -68,12 +68,14 @@ public:
 	UPROPERTY(VisibleAnywhere, Category = Results, meta = (TransientToolProperty))
 	TObjectPtr<UTexture2D> SubsurfaceColorMap = nullptr;
 
-	bool IsEmpty() const
-	{
-		return !BaseColorMap &&	!NormalMap && !PackedMRSMap &&
-			!MetallicMap && !RoughnessMap && !SpecularMap &&
-			!EmissiveMap && !OpacityMap && !SubsurfaceColorMap;
-	}
+	/** Device depth is currently unused */
+	UPROPERTY()
+	TObjectPtr<UTexture2D> DeviceDepthMap = nullptr;
+
+	/** Return the texture corresponding to the given CaptureType */
+	const TObjectPtr<UTexture2D>& operator[](UE::Geometry::ERenderCaptureType CaptureType) const;
+
+	bool IsEmpty() const;
 };
 
 
@@ -152,6 +154,11 @@ public:
 	UPROPERTY(Category = RenderCaptureOptions, EditAnywhere, AdvancedDisplay)
 	bool bAntiAliasing = false;
 
+	// Whether to generate a texture for the DeviceDepth.  This option is hidden from the user since its set only if the
+	// Bake operation requires a depth map to remove occlusion artefacts, which is determined by ValidSampleDepthThreshold
+	UPROPERTY()
+	bool bDeviceDepthMap = false;
+
 	// These are hidden in the UI right now, we might want to expose them if they turn out to be useful for very large
 	// or very small objects (not tested yet) TODO Figure out if we want to expose these options
 	
@@ -160,32 +167,9 @@ public:
 
 	UPROPERTY(meta = (ClampMin = "0.001", ClampMax = "1000.0"))
 	float NearPlaneDist = 1.0f;
-	
-	UPROPERTY()
-	bool bDeviceDepthMap = false;
-
-	bool operator==(const URenderCaptureProperties& Other) const
-	{
-		return Resolution == Other.Resolution
-			&& bBaseColorMap == Other.bBaseColorMap
-			&& bNormalMap == Other.bNormalMap
-			&& bMetallicMap == Other.bMetallicMap
-			&& bRoughnessMap == Other.bRoughnessMap
-			&& bSpecularMap == Other.bSpecularMap
-			&& bPackedMRSMap == Other.bPackedMRSMap
-			&& bEmissiveMap == Other.bEmissiveMap
-			&& bOpacityMap == Other.bOpacityMap
-			&& bSubsurfaceColorMap == Other.bSubsurfaceColorMap
-			&& bAntiAliasing == Other.bAntiAliasing
-			&& CaptureFieldOfView == Other.CaptureFieldOfView
-			&& NearPlaneDist == Other.NearPlaneDist
-			&& bDeviceDepthMap == Other.bDeviceDepthMap;
-	}
 
-	bool operator != (const URenderCaptureProperties& Other) const
-	{
-		return !(*this == Other);
-	}
+	bool operator==(const URenderCaptureProperties& Other) const;
+	bool operator!=(const URenderCaptureProperties& Other) const;
 };
 
 UCLASS()
@@ -345,7 +329,7 @@ protected:
 
 	void UpdateResult();
 	void UpdateVisualization();
-	void InvalidateResults(UE::Geometry::FRenderCaptureUpdate Update);
+	void InvalidateResults(UE::Geometry::FRenderCaptureTypeFlags ToInvalidate);
 	void InvalidateCompute();
 	void OnMapsUpdated(const TUniquePtr<UE::Geometry::FMeshMapBaker>& NewResult);
 
@@ -424,6 +408,7 @@ protected:
 	float ComputedValidDepthThreshold;
 	EBakeTextureSamplesPerPixel ComputedSamplesPerPixel = EBakeTextureSamplesPerPixel::Sample1;
 	EBakeTextureResolution ComputedTextureSize = EBakeTextureResolution::Resolution512;
+	FString ComputedTargetUVLayer;
 
 	TMap<int, FText> TargetUVLayerToError;
 
@@ -454,7 +439,4 @@ protected:
 	void GatherAnalytics(FBakeAnalytics::FMeshSettings& Data);
 	void GatherAnalytics(const UE::Geometry::FMeshMapBaker& Result);
 	void RecordAnalytics() const;
-
-private:
-	UE::Geometry::FRenderCaptureUpdate UpdateSceneCapture();
 };
\ No newline at end of file
diff --git a/Engine/Plugins/Runtime/MeshModelingToolset/Source/ModelingComponents/Private/Baking/RenderCaptureFunctions.cpp b/Engine/Plugins/Runtime/MeshModelingToolset/Source/ModelingComponents/Private/Baking/RenderCaptureFunctions.cpp
--- a/Engine/Plugins/Runtime/MeshModelingToolset/Source/ModelingComponents/Private/Baking/RenderCaptureFunctions.cpp
+++ b/Engine/Plugins/Runtime/MeshModelingToolset/Source/ModelingComponents/Private/Baking/RenderCaptureFunctions.cpp
@@ -149,7 +149,7 @@ FSceneCapturePhotoSetSampler::FSceneCapturePhotoSetSampler(
 	check(Spatial != nullptr);
 	check(Tangents != nullptr);
 
-	bool bHasDepth = SceneCapture->GetCaptureTypeEnabled(ERenderCaptureType::DeviceDepth);
+	bool bHasDepth = SceneCapture->GetCaptureTypeStatus(ERenderCaptureType::DeviceDepth) == FSceneCapturePhotoSet::ECaptureTypeStatus::Computed;
 	if (bHasDepth)
 	{
 		ensure(ValidSampleDepthThreshold > 0); // We only need the depth capture if this threshold is positive
@@ -205,38 +205,46 @@ bool FSceneCapturePhotoSetSampler::IsValidCorrespondence(const FMeshMapEvaluator
 
 
 
-namespace UE::Geometry
+bool FSceneCaptureConfig::operator==(const FSceneCaptureConfig& Other) const
 {
-namespace
+	return Flags == Other.Flags
+		&& RenderCaptureImageSize == Other.RenderCaptureImageSize
+		&& bAntiAliasing == Other.bAntiAliasing
+		&& FieldOfViewDegrees == Other.FieldOfViewDegrees
+		&& NearPlaneDist == Other.NearPlaneDist;
+}
+
+bool FSceneCaptureConfig::operator!=(const FSceneCaptureConfig& Other) const
 {
+	return !this->operator==(Other);
+}
+
+
 
-void UpdateSceneCaptureSettings(
+
+
+
+void UE::Geometry::ConfigureSceneCapture(
 	const TUniquePtr<FSceneCapturePhotoSet>& SceneCapture,
-	const TArray<AActor*>& Actors,
-	const FRenderCaptureOptions& Options)
+	const TArray<TObjectPtr<AActor>>& Actors,
+	const FSceneCaptureConfig& Options,
+	bool bAllowCancel)
 {
-	SceneCapture->SetCaptureTypeEnabled(ERenderCaptureType::DeviceDepth, Options.bBakeDeviceDepth);
-	SceneCapture->SetCaptureTypeEnabled(ERenderCaptureType::BaseColor,   Options.bBakeBaseColor);
-	SceneCapture->SetCaptureTypeEnabled(ERenderCaptureType::WorldNormal, Options.bBakeNormalMap);
-	SceneCapture->SetCaptureTypeEnabled(ERenderCaptureType::Emissive,    Options.bBakeEmissive);
-	SceneCapture->SetCaptureTypeEnabled(ERenderCaptureType::Opacity,     Options.bBakeOpacity);
-	SceneCapture->SetCaptureTypeEnabled(ERenderCaptureType::SubsurfaceColor, Options.bBakeSubsurfaceColor);
-	SceneCapture->SetCaptureTypeEnabled(ERenderCaptureType::CombinedMRS, Options.bUsePackedMRS);
-	SceneCapture->SetCaptureTypeEnabled(ERenderCaptureType::Metallic,    Options.bBakeMetallic);
-	SceneCapture->SetCaptureTypeEnabled(ERenderCaptureType::Roughness,   Options.bBakeRoughness);
-	SceneCapture->SetCaptureTypeEnabled(ERenderCaptureType::Specular,    Options.bBakeSpecular);
-
-	FRenderCaptureConfig Config;
-	Config.bAntiAliasing = Options.bAntiAliasing;
-	SceneCapture->SetCaptureConfig(ERenderCaptureType::BaseColor,   Config);
-	SceneCapture->SetCaptureConfig(ERenderCaptureType::WorldNormal, Config);
-	SceneCapture->SetCaptureConfig(ERenderCaptureType::CombinedMRS, Config);
-	SceneCapture->SetCaptureConfig(ERenderCaptureType::Metallic,    Config);
-	SceneCapture->SetCaptureConfig(ERenderCaptureType::Roughness,   Config);
-	SceneCapture->SetCaptureConfig(ERenderCaptureType::Specular,    Config);
-	SceneCapture->SetCaptureConfig(ERenderCaptureType::Emissive,    Config);
-	SceneCapture->SetCaptureConfig(ERenderCaptureType::Opacity,     Config);
-	SceneCapture->SetCaptureConfig(ERenderCaptureType::SubsurfaceColor, Config);
+	if (SceneCapture.IsValid() == false)
+	{
+		ensure(false);
+		return;
+	}
+
+	ForEachCaptureType([&SceneCapture, &Options](ERenderCaptureType CaptureType)
+	{
+		const bool bCaptureTypeEnabled = Options.Flags[CaptureType];
+		SceneCapture->SetCaptureTypeEnabled(CaptureType, bCaptureTypeEnabled);
+
+		FRenderCaptureConfig Config;
+		Config.bAntiAliasing = (CaptureType == ERenderCaptureType::DeviceDepth ? false : Options.bAntiAliasing);
+		SceneCapture->SetCaptureConfig(CaptureType, Config);
+	});
 
 	UWorld* World = Actors.IsEmpty() ? nullptr : Actors[0]->GetWorld();
 
@@ -251,13 +259,61 @@ void UpdateSceneCaptureSettings(
 		true, true, true, true, true);
 
 	SceneCapture->SetSpatialPhotoParams(SpatialParams);
+
+	SceneCapture->SetAllowCancel(bAllowCancel);
+}
+
+
+
+UE::Geometry::FRenderCaptureTypeFlags UE::Geometry::UpdateSceneCapture(
+	const TUniquePtr<FSceneCapturePhotoSet>& SceneCapture,
+	const TArray<TObjectPtr<AActor>>& Actors,
+	const FSceneCaptureConfig& DesiredConfig,
+	bool bAllowCancel)
+{
+	check(SceneCapture.IsValid());
+
+	const FSceneCapturePhotoSet::FStatus PreConfigStatus = SceneCapture->GetSceneCaptureStatus();
+	ConfigureSceneCapture(SceneCapture, Actors, DesiredConfig, true);
+	const FSceneCapturePhotoSet::FStatus PreComputeStatus = SceneCapture->GetSceneCaptureStatus();
+	SceneCapture->Compute();
+	const FSceneCapturePhotoSet::FStatus PostComputeStatus = SceneCapture->GetSceneCaptureStatus();
+
+	const FSceneCaptureConfig AchievedConfig = GetSceneCaptureConfig(SceneCapture);
+
+	// DesiredConfig and AchievedConfig mismatch iff the SceneCapture was cancelled. Also, the mismatch should only be in the Flags member
+	ensure(SceneCapture->Cancelled() == (AchievedConfig != DesiredConfig));
+	ensure(AchievedConfig.RenderCaptureImageSize == DesiredConfig.RenderCaptureImageSize);
+	ensure(AchievedConfig.FieldOfViewDegrees == DesiredConfig.FieldOfViewDegrees);
+	ensure(AchievedConfig.bAntiAliasing == DesiredConfig.bAntiAliasing);
+	ensure(AchievedConfig.NearPlaneDist == DesiredConfig.NearPlaneDist);
+
+	// Deduce which captures were changed in this scene capture update
+	FRenderCaptureTypeFlags UpdatedCaptures;
+	ForEachCaptureType([&UpdatedCaptures, &PreConfigStatus, &PreComputeStatus, &PostComputeStatus] (ERenderCaptureType CaptureType)
+	{
+		using EStatus = FSceneCapturePhotoSet::ECaptureTypeStatus;
+
+		// This boolean is true in two scenarios:
+		// 1. Computed -> Disabled occurs when the CaptureType was disabled by the DesiredConfig
+		// 2. Computed -> Pending occurs when the SceneCapture was cancelled while/before the CaptureType was computed
+		bool bCleared =  PreConfigStatus[CaptureType]  == EStatus::Computed && PostComputeStatus[CaptureType] != EStatus::Computed;
+
+		// This boolean is true in two scenarios, in both cases the PreComputeStatus of the given CaptureType is Pending
+		// 1. DesiredConfig enabled a new CaptureType
+		// 2. DesiredConfig requires an already computed CaptureType to be recomputed
+		bool bComputed = PreComputeStatus[CaptureType] == EStatus::Pending  && PostComputeStatus[CaptureType] == EStatus::Computed;
+
+		UpdatedCaptures[CaptureType] = bCleared || bComputed;
+	});
+
+	return UpdatedCaptures;
 }
 
-} // namespace
-} // namespace UE::Geometry
 
 
 
+PRAGMA_DISABLE_DEPRECATION_WARNINGS
 TUniquePtr<FSceneCapturePhotoSet> UE::Geometry::CapturePhotoSet(
 	const TArray<TObjectPtr<AActor>>& Actors,
 	const FRenderCaptureOptions& Options,
@@ -272,99 +328,60 @@ TUniquePtr<FSceneCapturePhotoSet> UE::Geometry::CapturePhotoSet(
 
 	return SceneCapture;
 }
+PRAGMA_ENABLE_DEPRECATION_WARNINGS
 
+PRAGMA_DISABLE_DEPRECATION_WARNINGS
 FRenderCaptureUpdate UE::Geometry::UpdatePhotoSets(
 	const TUniquePtr<FSceneCapturePhotoSet>& SceneCapture,
 	const TArray<TObjectPtr<AActor>>& Actors,
 	const FRenderCaptureOptions& Options,
 	bool bAllowCancel)
 {
-	FScopedSlowTask Progress(0.f, LOCTEXT("CapturingScene", "Capturing Scene..."));
-	Progress.MakeDialog(bAllowCancel);
-
-	// Cache previous SceneCapture settings so these can be restored if the computation is cancelled
-	const FRenderCaptureOptions ComputedOptions = GetComputedPhotoSetOptions(SceneCapture);
-	const TArray<AActor*> ComputedActors = SceneCapture->GetCaptureSceneActors();
-
-	// We track the photo set counts in order to determine which photosets were updated
-	struct FPhotoSetNums
-	{
-		FPhotoSetNums(const TUniquePtr<FSceneCapturePhotoSet>& SceneCapture)
-		{
-			NumBaseColor       = SceneCapture->GetBaseColorPhotoSet().Num();
-			NumRoughness       = SceneCapture->GetRoughnessPhotoSet().Num();
-			NumSpecular        = SceneCapture->GetSpecularPhotoSet().Num();
-			NumMetallic        = SceneCapture->GetMetallicPhotoSet().Num();
-			NumPackedMRS       = SceneCapture->GetPackedMRSPhotoSet().Num();
-			NumNormalMap       = SceneCapture->GetWorldNormalPhotoSet().Num();
-			NumEmissive        = SceneCapture->GetEmissivePhotoSet().Num();
-			NumOpacity         = SceneCapture->GetOpacityPhotoSet().Num();
-			NumSubsurfaceColor = SceneCapture->GetSubsurfaceColorPhotoSet().Num();
-			NumDeviceDepth     = SceneCapture->GetDeviceDepthPhotoSet().Num();
-		}
-		
-		int32 NumBaseColor;
-		int32 NumRoughness;
-		int32 NumSpecular;
-		int32 NumMetallic;
-		int32 NumPackedMRS;
-		int32 NumNormalMap;
-		int32 NumEmissive;
-		int32 NumOpacity;
-		int32 NumSubsurfaceColor;
-		int32 NumDeviceDepth;
-	};
-
-	const FPhotoSetNums Step1(SceneCapture);
-
-	// This will clear any photosets that are disabled or need recomputing
-	UpdateSceneCaptureSettings(SceneCapture, Actors, Options);
-
-	const FPhotoSetNums Step2(SceneCapture);
-
-	SceneCapture->SetAllowCancel(bAllowCancel);
-
-	// This will compute newly requested photosets
-	SceneCapture->Compute();
-
-	if (SceneCapture->Cancelled())
-	{
-		// This will clear any newly requested photosets
-		UpdateSceneCaptureSettings(SceneCapture, ComputedActors, ComputedOptions);
-	}
-
-	const FPhotoSetNums Step3(SceneCapture);
-
-	// If the photo sets were cleared/computed/recomputed we consider them updated, use the counts to figure this out
-	FRenderCaptureUpdate Update;
-	Update.bUpdatedBaseColor   = !(Step1.NumBaseColor   == Step2.NumBaseColor   && Step2.NumBaseColor   == Step3.NumBaseColor);
-	Update.bUpdatedRoughness   = !(Step1.NumRoughness   == Step2.NumRoughness   && Step2.NumRoughness   == Step3.NumRoughness);
-	Update.bUpdatedSpecular    = !(Step1.NumSpecular    == Step2.NumSpecular    && Step2.NumSpecular    == Step3.NumSpecular);
-	Update.bUpdatedMetallic    = !(Step1.NumMetallic    == Step2.NumMetallic    && Step2.NumMetallic    == Step3.NumMetallic);
-	Update.bUpdatedPackedMRS   = !(Step1.NumPackedMRS   == Step2.NumPackedMRS   && Step2.NumPackedMRS   == Step3.NumPackedMRS);
-	Update.bUpdatedNormalMap   = !(Step1.NumNormalMap   == Step2.NumNormalMap   && Step2.NumNormalMap   == Step3.NumNormalMap);
-	Update.bUpdatedEmissive    = !(Step1.NumEmissive    == Step2.NumEmissive    && Step2.NumEmissive    == Step3.NumEmissive);
-	Update.bUpdatedOpacity     = !(Step1.NumOpacity     == Step2.NumOpacity     && Step2.NumOpacity     == Step3.NumOpacity);
-	Update.bUpdatedDeviceDepth = !(Step1.NumDeviceDepth == Step2.NumDeviceDepth && Step2.NumDeviceDepth == Step3.NumDeviceDepth);
-	Update.bUpdatedSubsurfaceColor = !(Step1.NumSubsurfaceColor == Step2.NumSubsurfaceColor && Step2.NumSubsurfaceColor == Step3.NumSubsurfaceColor);
-
-	return Update;
+	FSceneCaptureConfig Config;
+	Config.RenderCaptureImageSize = Options.RenderCaptureImageSize;
+	Config.bAntiAliasing          = Options.bAntiAliasing;
+	Config.FieldOfViewDegrees     = Options.FieldOfViewDegrees;
+	Config.NearPlaneDist          = Options.NearPlaneDist;
+	Config.Flags.bBaseColor       = Options.bBakeBaseColor;
+	Config.Flags.bRoughness       = Options.bBakeRoughness;
+	Config.Flags.bMetallic        = Options.bBakeMetallic;
+	Config.Flags.bSpecular        = Options.bBakeSpecular;
+	Config.Flags.bEmissive        = Options.bBakeEmissive;
+	Config.Flags.bWorldNormal     = Options.bBakeNormalMap;
+	Config.Flags.bOpacity         = Options.bBakeOpacity;
+	Config.Flags.bSubsurfaceColor = Options.bBakeSubsurfaceColor;
+	Config.Flags.bDeviceDepth     = Options.bBakeDeviceDepth;
+	Config.Flags.bCombinedMRS     = Options.bUsePackedMRS;
+
+	FRenderCaptureTypeFlags UpdatedCaptures = UpdateSceneCapture(SceneCapture, Actors, Config, bAllowCancel);
+
+	FRenderCaptureUpdate Updated;
+	Updated.bUpdatedBaseColor   = UpdatedCaptures.bBaseColor;
+	Updated.bUpdatedRoughness   = UpdatedCaptures.bRoughness;
+	Updated.bUpdatedSpecular    = UpdatedCaptures.bSpecular;
+	Updated.bUpdatedMetallic    = UpdatedCaptures.bMetallic;
+	Updated.bUpdatedPackedMRS   = UpdatedCaptures.bCombinedMRS;
+	Updated.bUpdatedNormalMap   = UpdatedCaptures.bWorldNormal;
+	Updated.bUpdatedEmissive    = UpdatedCaptures.bEmissive;
+	Updated.bUpdatedOpacity     = UpdatedCaptures.bOpacity;
+	Updated.bUpdatedDeviceDepth = UpdatedCaptures.bDeviceDepth;
+	Updated.bUpdatedSubsurfaceColor = UpdatedCaptures.bSubsurfaceColor;
+
+	return Updated;
 }
+PRAGMA_ENABLE_DEPRECATION_WARNINGS
 
-FRenderCaptureOptions UE::Geometry::GetComputedPhotoSetOptions(const TUniquePtr<FSceneCapturePhotoSet>& SceneCapture)
+// Return the SceneCapture configuration, the Flags in the returned struct are set if the corresponding capture type is computed
+FSceneCaptureConfig UE::Geometry::GetSceneCaptureConfig(
+	const TUniquePtr<FSceneCapturePhotoSet>& SceneCapture,
+	const FSceneCapturePhotoSet::ECaptureTypeStatus QueryStatus)
 {
-	FRenderCaptureOptions Result;
-
-	Result.bBakeBaseColor       = SceneCapture->GetCaptureTypeEnabled(ERenderCaptureType::BaseColor);
-	Result.bBakeNormalMap       = SceneCapture->GetCaptureTypeEnabled(ERenderCaptureType::WorldNormal);
-	Result.bBakeMetallic        = SceneCapture->GetCaptureTypeEnabled(ERenderCaptureType::Metallic);
-	Result.bBakeRoughness       = SceneCapture->GetCaptureTypeEnabled(ERenderCaptureType::Roughness);
-	Result.bBakeSpecular        = SceneCapture->GetCaptureTypeEnabled(ERenderCaptureType::Specular);
-	Result.bUsePackedMRS        = SceneCapture->GetCaptureTypeEnabled(ERenderCaptureType::CombinedMRS);
-	Result.bBakeEmissive        = SceneCapture->GetCaptureTypeEnabled(ERenderCaptureType::Emissive);
-	Result.bBakeOpacity         = SceneCapture->GetCaptureTypeEnabled(ERenderCaptureType::Opacity);
-	Result.bBakeSubsurfaceColor = SceneCapture->GetCaptureTypeEnabled(ERenderCaptureType::SubsurfaceColor);
-	Result.bBakeDeviceDepth     = SceneCapture->GetCaptureTypeEnabled(ERenderCaptureType::DeviceDepth);
+	FSceneCaptureConfig Result;
+
+	ForEachCaptureType([&SceneCapture, &QueryStatus, &Result](ERenderCaptureType CaptureType)
+	{
+		Result.Flags[CaptureType] = (SceneCapture->GetCaptureTypeStatus(CaptureType) == QueryStatus);
+	});
 
 	const TArray<FSpatialPhotoParams> SpatialParams = SceneCapture->GetSpatialPhotoParams();
 
@@ -393,6 +410,32 @@ FRenderCaptureOptions UE::Geometry::GetComputedPhotoSetOptions(const TUniquePtr<
 	return Result;
 }
 
+PRAGMA_DISABLE_DEPRECATION_WARNINGS
+FRenderCaptureOptions UE::Geometry::GetComputedPhotoSetOptions(const TUniquePtr<FSceneCapturePhotoSet>& SceneCapture)
+{
+	FSceneCaptureConfig Config = GetSceneCaptureConfig(SceneCapture);
+
+	FRenderCaptureOptions Options;
+
+	Options.RenderCaptureImageSize = Config.RenderCaptureImageSize;
+	Options.bAntiAliasing          = Config.bAntiAliasing;
+	Options.FieldOfViewDegrees     = Config.FieldOfViewDegrees;
+	Options.NearPlaneDist          = Config.NearPlaneDist;
+	Options.bBakeBaseColor         = Config.Flags.bBaseColor;
+	Options.bBakeRoughness         = Config.Flags.bRoughness;
+	Options.bBakeMetallic          = Config.Flags.bMetallic;
+	Options.bBakeSpecular          = Config.Flags.bSpecular;
+	Options.bBakeEmissive          = Config.Flags.bEmissive;
+	Options.bBakeNormalMap         = Config.Flags.bWorldNormal;
+	Options.bBakeOpacity           = Config.Flags.bOpacity;
+	Options.bBakeSubsurfaceColor   = Config.Flags.bSubsurfaceColor;
+	Options.bBakeDeviceDepth       = Config.Flags.bDeviceDepth;
+	Options.bUsePackedMRS          = Config.Flags.bCombinedMRS;
+
+	return Options;
+}
+PRAGMA_ENABLE_DEPRECATION_WARNINGS
+
 
 
 template <ERenderCaptureType CaptureType>
@@ -401,6 +444,9 @@ MakeColorEvaluator(
 	const FSceneCapturePhotoSet::FSceneSample& DefaultSample,
 	const FSceneCapturePhotoSet* SceneCapture)
 {
+	// Bake will crash if we haven't computed the photo sets corresponding to the CaptureType 
+	check(SceneCapture->GetCaptureTypeStatus(CaptureType) == FSceneCapturePhotoSet::ECaptureTypeStatus::Computed);
+
 	TSharedPtr<FRenderCaptureMapEvaluator<FVector4f>> Evaluator = MakeShared<FRenderCaptureMapEvaluator<FVector4f>>();
 
 	switch (CaptureType) {
@@ -463,7 +509,7 @@ TUniquePtr<FMeshMapBaker> UE::Geometry::MakeRenderCaptureBaker(
 	TSharedPtr<TArray<int32>, ESPMode::ThreadSafe> BaseMeshUVCharts,
 	FSceneCapturePhotoSet* SceneCapture,
 	FSceneCapturePhotoSetSampler* Sampler,
-	FRenderCaptureOptions PendingBake,
+	FRenderCaptureTypeFlags PendingBake,
 	int32 TargetUVLayer,
 	EBakeTextureResolution TextureImageSize,
 	EBakeTextureSamplesPerPixel SamplesPerPixel,
@@ -515,9 +561,7 @@ TUniquePtr<FMeshMapBaker> UE::Geometry::MakeRenderCaptureBaker(
 	DefaultColorSample.Opacity = InvalidColor.X;
 	DefaultColorSample.SubsurfaceColor = FVector3f(InvalidColor.X, InvalidColor.Y, InvalidColor.Z);
 	DefaultColorSample.WorldNormal = FVector4f((DefaultNormal + FVector3f::One()) * .5f, InvalidColor.W);
-
-	// We use 0 here since this corresponds to the infinite far plane value. Also, we don't preview the depth texture
-	DefaultColorSample.DeviceDepth = 0;
+	DefaultColorSample.DeviceDepth = 0; // We use 0 here since this corresponds to the infinite far plane value. Also, we don't preview the depth texture
 
 	auto AddColorEvaluator = [&Result, OcclusionHandler] (const TSharedPtr<FRenderCaptureMapEvaluator<FVector4f>>& Evaluator)
 	{
@@ -525,44 +569,47 @@ TUniquePtr<FMeshMapBaker> UE::Geometry::MakeRenderCaptureBaker(
 		OcclusionHandler->PushInfillRequired(true);
 	};
 
-	if (SceneCapture->GetCaptureTypeEnabled(ERenderCaptureType::DeviceDepth) && PendingBake.bBakeDeviceDepth)
+	if (PendingBake.bDeviceDepth)
 	{
 		AddColorEvaluator(MakeColorEvaluator<ERenderCaptureType::DeviceDepth>(DefaultColorSample, SceneCapture));
 	}
-	if (SceneCapture->GetCaptureTypeEnabled(ERenderCaptureType::BaseColor) && PendingBake.bBakeBaseColor)
+	if (PendingBake.bBaseColor)
 	{
 		AddColorEvaluator(MakeColorEvaluator<ERenderCaptureType::BaseColor>(DefaultColorSample, SceneCapture));
 	}
-	if (SceneCapture->GetCaptureTypeEnabled(ERenderCaptureType::CombinedMRS) && PendingBake.bUsePackedMRS)
+	if (PendingBake.bCombinedMRS)
 	{
 		AddColorEvaluator(MakeColorEvaluator<ERenderCaptureType::CombinedMRS>(DefaultColorSample, SceneCapture));
 	}
-	if (SceneCapture->GetCaptureTypeEnabled(ERenderCaptureType::Roughness) && PendingBake.bBakeRoughness)
+	if (PendingBake.bRoughness)
 	{
 		AddColorEvaluator(MakeColorEvaluator<ERenderCaptureType::Roughness>(DefaultColorSample, SceneCapture));
 	}
-	if (SceneCapture->GetCaptureTypeEnabled(ERenderCaptureType::Metallic) && PendingBake.bBakeMetallic)
+	if (PendingBake.bMetallic)
 	{
 		AddColorEvaluator(MakeColorEvaluator<ERenderCaptureType::Metallic>(DefaultColorSample, SceneCapture));
 	}
-	if (SceneCapture->GetCaptureTypeEnabled(ERenderCaptureType::Specular) && PendingBake.bBakeSpecular)
+	if (PendingBake.bSpecular)
 	{
 		AddColorEvaluator(MakeColorEvaluator<ERenderCaptureType::Specular>(DefaultColorSample, SceneCapture));
 	}
-	if (SceneCapture->GetCaptureTypeEnabled(ERenderCaptureType::Emissive) && PendingBake.bBakeEmissive)
+	if (PendingBake.bEmissive)
 	{
 		AddColorEvaluator(MakeColorEvaluator<ERenderCaptureType::Emissive>(DefaultColorSample, SceneCapture));
 	}
-	if (SceneCapture->GetCaptureTypeEnabled(ERenderCaptureType::Opacity) && PendingBake.bBakeOpacity)
+	if (PendingBake.bOpacity)
 	{
 		AddColorEvaluator(MakeColorEvaluator<ERenderCaptureType::Opacity>(DefaultColorSample, SceneCapture));
 	}
-	if (SceneCapture->GetCaptureTypeEnabled(ERenderCaptureType::SubsurfaceColor) && PendingBake.bBakeSubsurfaceColor)
+	if (PendingBake.bSubsurfaceColor)
 	{
 		AddColorEvaluator(MakeColorEvaluator<ERenderCaptureType::SubsurfaceColor>(DefaultColorSample, SceneCapture));
 	}
-	if (SceneCapture->GetCaptureTypeEnabled(ERenderCaptureType::WorldNormal) && PendingBake.bBakeNormalMap)
+	if (PendingBake.bWorldNormal)
 	{
+		// Bake will crash if we haven't computed the WorldNormal photo sets
+		check(SceneCapture->GetCaptureTypeStatus(ERenderCaptureType::WorldNormal) == FSceneCapturePhotoSet::ECaptureTypeStatus::Computed);
+
 		TSharedPtr<FRenderCaptureMapEvaluator<FVector3f>> Evaluator = MakeShared<FRenderCaptureMapEvaluator<FVector3f>>();
 
 		Evaluator->Channel = ERenderCaptureChannel::WorldNormal;
@@ -617,6 +664,46 @@ TUniquePtr<FMeshMapBaker> UE::Geometry::MakeRenderCaptureBaker(
 	return Result;
 }
 
+PRAGMA_DISABLE_DEPRECATION_WARNINGS
+TUniquePtr<FMeshMapBaker> UE::Geometry::MakeRenderCaptureBaker(
+	FDynamicMesh3* BaseMesh,
+	TSharedPtr<UE::Geometry::FMeshTangentsd, ESPMode::ThreadSafe> BaseMeshTangents,
+	TSharedPtr<TArray<int32>, ESPMode::ThreadSafe> BaseMeshUVCharts,
+	FSceneCapturePhotoSet* SceneCapture,
+	FSceneCapturePhotoSetSampler* Sampler,
+	FRenderCaptureOptions Options,
+	int32 TargetUVLayer,
+	EBakeTextureResolution TextureImageSize,
+	EBakeTextureSamplesPerPixel SamplesPerPixel,
+	FRenderCaptureOcclusionHandler* OcclusionHandler)
+{
+	FRenderCaptureTypeFlags Flags;
+
+	Flags.bBaseColor       = Options.bBakeBaseColor;
+	Flags.bRoughness       = Options.bBakeRoughness;
+	Flags.bMetallic        = Options.bBakeMetallic;
+	Flags.bSpecular        = Options.bBakeSpecular;
+	Flags.bEmissive        = Options.bBakeEmissive;
+	Flags.bWorldNormal     = Options.bBakeNormalMap;
+	Flags.bOpacity         = Options.bBakeOpacity;
+	Flags.bSubsurfaceColor = Options.bBakeSubsurfaceColor;
+	Flags.bDeviceDepth     = Options.bBakeDeviceDepth;
+	Flags.bCombinedMRS     = Options.bUsePackedMRS;
+
+	return MakeRenderCaptureBaker(
+		BaseMesh,
+		BaseMeshTangents,
+		BaseMeshUVCharts,
+		SceneCapture,
+		Sampler,
+		Flags,
+		TargetUVLayer,
+		TextureImageSize,
+		SamplesPerPixel,
+		OcclusionHandler);
+}
+PRAGMA_ENABLE_DEPRECATION_WARNINGS
+
 
 void UE::Geometry::GetTexturesFromRenderCaptureBaker(const TUniquePtr<FMeshMapBaker>& Baker, FRenderCaptureTextures& TexturesOut)
 {
@@ -719,6 +806,7 @@ void UE::Geometry::GetTexturesFromRenderCaptureBaker(const TUniquePtr<FMeshMapBa
 				case ERenderCaptureChannel::DeviceDepth:
 
 					// Add a null pointer for this evaluator, the depth capture is used internally and not presented to the user
+					// TODO Implemented this, it could be useful for applications outside BakeRC
 					break;
 
 				case ERenderCaptureChannel::WorldNormal:
diff --git a/Engine/Plugins/Runtime/MeshModelingToolset/Source/ModelingComponents/Private/Scene/SceneCapturePhotoSet.cpp b/Engine/Plugins/Runtime/MeshModelingToolset/Source/ModelingComponents/Private/Scene/SceneCapturePhotoSet.cpp
--- a/Engine/Plugins/Runtime/MeshModelingToolset/Source/ModelingComponents/Private/Scene/SceneCapturePhotoSet.cpp
+++ b/Engine/Plugins/Runtime/MeshModelingToolset/Source/ModelingComponents/Private/Scene/SceneCapturePhotoSet.cpp
@@ -10,22 +10,32 @@ using namespace UE::Geometry;
 
 FSceneCapturePhotoSet::FSceneCapturePhotoSet()
 {
-	BaseColorConfig = GetDefaultRenderCaptureConfig(ERenderCaptureType::BaseColor);
-	WorldNormalConfig = GetDefaultRenderCaptureConfig(ERenderCaptureType::WorldNormal);
-	RoughnessConfig = GetDefaultRenderCaptureConfig(ERenderCaptureType::Roughness);
-	MetallicConfig = GetDefaultRenderCaptureConfig(ERenderCaptureType::Metallic);
-	SpecularConfig = GetDefaultRenderCaptureConfig(ERenderCaptureType::Specular);
-	EmissiveConfig = GetDefaultRenderCaptureConfig(ERenderCaptureType::Emissive);
-	PackedMRSConfig = GetDefaultRenderCaptureConfig(ERenderCaptureType::CombinedMRS);
-	OpacityConfig = GetDefaultRenderCaptureConfig(ERenderCaptureType::Opacity);
-	SubsurfaceColorConfig = GetDefaultRenderCaptureConfig(ERenderCaptureType::SubsurfaceColor);
-	DeviceDepthConfig = GetDefaultRenderCaptureConfig(ERenderCaptureType::DeviceDepth);
+	// These defaults are chosen so an unfamiliar user doesn't need to call the configuration functions to get a result
+	ForEachCaptureType([this](ERenderCaptureType CaptureType)
+	{
+		bool bIsDisabled = (
+			CaptureType == ERenderCaptureType::DeviceDepth ||
+			CaptureType == ERenderCaptureType::Roughness ||
+			CaptureType == ERenderCaptureType::Specular ||
+			CaptureType == ERenderCaptureType::Metallic);
+		PhotoSetStatus[CaptureType] = bIsDisabled ? ECaptureTypeStatus::Disabled : ECaptureTypeStatus::Pending;
+
+		RenderCaptureConfig[CaptureType] = GetDefaultRenderCaptureConfig(CaptureType);
+	});
 }
 
 void FSceneCapturePhotoSet::SetCaptureSceneActors(UWorld* World, const TArray<AActor*>& Actors)
 {
 	if (this->TargetWorld != World || this->VisibleActors != Actors)
 	{
+		ForEachCaptureType([this](ERenderCaptureType CaptureType)
+		{
+			if (PhotoSetStatus[CaptureType] != ECaptureTypeStatus::Disabled)
+			{
+				PhotoSetStatus[CaptureType] = ECaptureTypeStatus::Pending;
+			}
+		});
+
 		// Empty the photo sets because they rendered different actors
 		EmptyAllPhotoSets();
 
@@ -51,9 +61,15 @@ void FSceneCapturePhotoSet::SetSpatialPhotoParams(const TArray<FSpatialPhotoPara
 	// TODO Discard/reset on a per array element level rather than discarding everything when any viewpoint changed
 	if (PhotoSetParams != SpatialParams)
 	{
-		// Empty the photo sets because they were rendered with different viewpoints
-		EmptyAllPhotoSets();
+		ForEachCaptureType([this](ERenderCaptureType CaptureType)
+		{
+			if (PhotoSetStatus[CaptureType] != ECaptureTypeStatus::Disabled)
+			{
+				PhotoSetStatus[CaptureType] = ECaptureTypeStatus::Pending;
+			}
+		});
 
+		EmptyAllPhotoSets();
 		PhotoSetParams = SpatialParams;
 	}
 }
@@ -64,220 +80,81 @@ const TArray<FSpatialPhotoParams>& FSceneCapturePhotoSet::GetSpatialPhotoParams(
 }
 
 
-void FSceneCapturePhotoSet::SetCaptureConfig(ERenderCaptureType CaptureType, const FRenderCaptureConfig& Config)
+void FSceneCapturePhotoSet::SetCaptureConfig(ERenderCaptureType CaptureType, const FRenderCaptureConfig& NewConfig)
 {
-	auto UpdateCaptureConfig1f = [Config](FRenderCaptureConfig& CaptureConfig, FSpatialPhotoSet1f& PhotoSet)
+	if (RenderCaptureConfig[CaptureType] != NewConfig)
 	{
-		if (CaptureConfig.bAntiAliasing != Config.bAntiAliasing)
+		if (PhotoSetStatus[CaptureType] != ECaptureTypeStatus::Disabled)
 		{
-			// If we're changing the AntiAliasing state we need to remove any existing photos
-			PhotoSet.Empty();
+			PhotoSetStatus[CaptureType] = ECaptureTypeStatus::Pending;
 		}
-		CaptureConfig = Config;
-	};
-
-	auto UpdateCaptureConfig3f = [Config](FRenderCaptureConfig& CaptureConfig, FSpatialPhotoSet3f& PhotoSet)
-	{
-		if (CaptureConfig.bAntiAliasing != Config.bAntiAliasing)
-		{
-			// If we're changing the AntiAliasing state we need to remove any existing photos
-			PhotoSet.Empty();
-		}
-		CaptureConfig = Config;
-	};
 
-	switch (CaptureType)
-	{
-	case ERenderCaptureType::BaseColor:
-		UpdateCaptureConfig3f(BaseColorConfig, BaseColorPhotoSet);
-		break;
-	case ERenderCaptureType::WorldNormal:
-		UpdateCaptureConfig3f(WorldNormalConfig, WorldNormalPhotoSet);
-		break;
-	case ERenderCaptureType::Roughness:
-		UpdateCaptureConfig1f(RoughnessConfig, RoughnessPhotoSet);
-		break;
-	case ERenderCaptureType::Metallic:
-		UpdateCaptureConfig1f(MetallicConfig, MetallicPhotoSet);
-		break;
-	case ERenderCaptureType::Specular:
-		UpdateCaptureConfig1f(SpecularConfig, SpecularPhotoSet);
-		break;
-	case ERenderCaptureType::Emissive:
-		UpdateCaptureConfig3f(EmissiveConfig, EmissivePhotoSet);
-		break;
-	case ERenderCaptureType::CombinedMRS:
-		UpdateCaptureConfig3f(PackedMRSConfig, PackedMRSPhotoSet);
-		break;
-	case ERenderCaptureType::Opacity:
-		UpdateCaptureConfig1f(OpacityConfig, OpacityPhotoSet);
-		break;
-	case ERenderCaptureType::SubsurfaceColor:
-		UpdateCaptureConfig3f(SubsurfaceColorConfig, SubsurfaceColorPhotoSet);
-		break;
-	case ERenderCaptureType::DeviceDepth:
-		if (DeviceDepthConfig.bAntiAliasing != Config.bAntiAliasing)
-		{
-			// If we're disabling DeviceDepth we need to remove any existing photos as well as the cached view matrices
-			DeviceDepthPhotoSet.Empty();
-			PhotoViewMatricies.Empty();
-		}
-		DeviceDepthConfig = Config;
-		break;
-	default:
-		ensure(false);
+		EmptyPhotoSet(CaptureType);
+		RenderCaptureConfig[CaptureType] = NewConfig;
 	}
 }
 
 FRenderCaptureConfig FSceneCapturePhotoSet::GetCaptureConfig(ERenderCaptureType CaptureType) const
 {
-	switch (CaptureType)
-	{
-	case ERenderCaptureType::BaseColor:
-			return BaseColorConfig;
-		case ERenderCaptureType::WorldNormal:
-			return WorldNormalConfig;
-		case ERenderCaptureType::Roughness:
-			return RoughnessConfig;
-		case ERenderCaptureType::Metallic:
-			return MetallicConfig;
-		case ERenderCaptureType::Specular:
-			return SpecularConfig;
-		case ERenderCaptureType::Emissive:
-			return EmissiveConfig;
-		case ERenderCaptureType::CombinedMRS:
-			return PackedMRSConfig;
-		case ERenderCaptureType::Opacity:
-			return OpacityConfig;
-		case ERenderCaptureType::SubsurfaceColor:
-			return SubsurfaceColorConfig;
-		case ERenderCaptureType::DeviceDepth:
-			return DeviceDepthConfig;
-		default:
-			ensure(false);
-	}
-	return BaseColorConfig;
+	return RenderCaptureConfig[CaptureType];
 }
 
 void FSceneCapturePhotoSet::DisableAllCaptureTypes()
 {
-	SetCaptureTypeEnabled(ERenderCaptureType::BaseColor, false);
-	SetCaptureTypeEnabled(ERenderCaptureType::WorldNormal, false);
-	SetCaptureTypeEnabled(ERenderCaptureType::Roughness, false);
-	SetCaptureTypeEnabled(ERenderCaptureType::Metallic, false);
-	SetCaptureTypeEnabled(ERenderCaptureType::Specular, false);
-	SetCaptureTypeEnabled(ERenderCaptureType::Emissive, false);
-	SetCaptureTypeEnabled(ERenderCaptureType::Opacity, false);
-	SetCaptureTypeEnabled(ERenderCaptureType::SubsurfaceColor, false);
-	SetCaptureTypeEnabled(ERenderCaptureType::CombinedMRS, false);
-	SetCaptureTypeEnabled(ERenderCaptureType::DeviceDepth, false);
+	ForEachCaptureType([this](ERenderCaptureType CaptureType)
+	{
+		SetCaptureTypeEnabled(CaptureType, false);
+	});
 }
 
 
-void FSceneCapturePhotoSet::SetCaptureTypeEnabled(ERenderCaptureType CaptureType, bool bEnabled)
+void FSceneCapturePhotoSet::SetCaptureTypeEnabled(ERenderCaptureType CaptureType, bool bEnable)
 {
-	auto UpdateCaptureType1f = [bEnabled](bool& bCaptureEnabled, FSpatialPhotoSet1f& PhotoSet)
-	{
-		if (!bEnabled)
-		{
-			// If we're disabling a CaptureType we need to remove any existing photos
-			PhotoSet.Empty();
-		}
-		bCaptureEnabled = bEnabled;
-	};
-
-	auto UpdateCaptureType3f = [bEnabled](bool& bCaptureEnabled, FSpatialPhotoSet3f& PhotoSet)
+	if (bEnable)
 	{
-		if (!bEnabled)
+		if (PhotoSetStatus[CaptureType] == ECaptureTypeStatus::Disabled)
 		{
-			// If we're disabling a CaptureType we need to remove any existing photos
-			PhotoSet.Empty();
+			PhotoSetStatus[CaptureType] = ECaptureTypeStatus::Pending;
 		}
-		bCaptureEnabled = bEnabled;
-	};
-
-	switch (CaptureType)
+	}
+	else
 	{
-	case ERenderCaptureType::BaseColor:
-		UpdateCaptureType3f(bEnableBaseColor, BaseColorPhotoSet);
-		break;
-	case ERenderCaptureType::WorldNormal:
-		UpdateCaptureType3f(bEnableWorldNormal, WorldNormalPhotoSet);
-		break;
-	case ERenderCaptureType::Roughness:
-		UpdateCaptureType1f(bEnableRoughness, RoughnessPhotoSet);
-		break;
-	case ERenderCaptureType::Metallic:
-		UpdateCaptureType1f(bEnableMetallic, MetallicPhotoSet);
-		break;
-	case ERenderCaptureType::Specular:
-		UpdateCaptureType1f(bEnableSpecular, SpecularPhotoSet);
-		break;
-	case ERenderCaptureType::Emissive:
-		UpdateCaptureType3f(bEnableEmissive, EmissivePhotoSet);
-		break;
-	case ERenderCaptureType::CombinedMRS:
-		UpdateCaptureType3f(bEnablePackedMRS, PackedMRSPhotoSet);
-		break;
-	case ERenderCaptureType::Opacity:
-		UpdateCaptureType1f(bEnableOpacity, OpacityPhotoSet);
-		break;
-	case ERenderCaptureType::SubsurfaceColor:
-		UpdateCaptureType3f(bEnableSubsurfaceColor, SubsurfaceColorPhotoSet);
-		break;
-	case ERenderCaptureType::DeviceDepth:
-		if (bEnableDeviceDepth != bEnabled)
-		{
-			// If we're disabling DeviceDepth we need to remove any existing photos as well as the cached view matrices
-			DeviceDepthPhotoSet.Empty();
-			PhotoViewMatricies.Empty();
-		}
-		bEnableDeviceDepth = bEnabled;
-		break;
-	default:
-		ensure(false);
+		PhotoSetStatus[CaptureType] = ECaptureTypeStatus::Disabled;
+		EmptyPhotoSet(CaptureType);
 	}
 }
 
-bool FSceneCapturePhotoSet::GetCaptureTypeEnabled(ERenderCaptureType CaptureType) const
+FSceneCapturePhotoSet::ECaptureTypeStatus FSceneCapturePhotoSet::GetCaptureTypeStatus(ERenderCaptureType CaptureType) const
 {
-	switch (CaptureType)
-	{
-	case ERenderCaptureType::BaseColor:
-			return bEnableBaseColor;
-		case ERenderCaptureType::WorldNormal:
-			return bEnableWorldNormal;
-		case ERenderCaptureType::Roughness:
-			return bEnableRoughness;
-		case ERenderCaptureType::Metallic:
-			return bEnableMetallic;
-		case ERenderCaptureType::Specular:
-			return bEnableSpecular;
-		case ERenderCaptureType::Emissive:
-			return bEnableEmissive;
-		case ERenderCaptureType::Opacity:
-			return bEnableOpacity;
-		case ERenderCaptureType::SubsurfaceColor:
-			return bEnableSubsurfaceColor;
-		case ERenderCaptureType::CombinedMRS:
-			return bEnablePackedMRS;
-		case ERenderCaptureType::DeviceDepth:
-			return bEnableDeviceDepth;
-		default:
-			ensure(false);
-	}
-	return false;
+	return PhotoSetStatus[CaptureType];
+}
+
+FSceneCapturePhotoSet::FStatus FSceneCapturePhotoSet::GetSceneCaptureStatus() const
+{
+	return PhotoSetStatus;
 }
 
 void FSceneCapturePhotoSet::Compute()
 {
-	check(this->TargetWorld != nullptr);
+	int NumPending = 0;
+	ForEachCaptureType([this, &NumPending](ERenderCaptureType CaptureType)
+	{
+		NumPending += static_cast<int>(PhotoSetStatus[CaptureType] == ECaptureTypeStatus::Pending);
+	});
 
-	bWasCancelled = false;
+	if (NumPending == 0)
+	{
+		return;
+	}
 
-	FScopedSlowTask Progress(PhotoSetParams.Num(), LOCTEXT("ComputingViewpoints", "Computing Viewpoints..."));
+	TRACE_CPUPROFILER_EVENT_SCOPE(CapturePhotoSet);
+	FScopedSlowTask Progress(static_cast<float>(NumPending), LOCTEXT("CapturingScene", "Capturing Scene..."));
 	Progress.MakeDialog(bAllowCancel);
 
+	check(this->TargetWorld != nullptr);
+
+	bWasCancelled = false;
+
 	// Unregister all components to remove unwanted proxies from the scene. This was previously the only way to "hide" nanite meshes, now optional.
 	TSet<AActor*> VisibleActorsSet(VisibleActors);
 	TArray<AActor*> ActorsToRegister;
@@ -311,105 +188,84 @@ void FSceneCapturePhotoSet::Compute()
 		RenderCapture.SetEnableWriteDebugImage(true, 0, DebugImagesFolderName);
 	}
 
-	for (const FSpatialPhotoParams& Params : PhotoSetParams)
+	auto CapturePhoto3f = [this, &RenderCapture](ERenderCaptureType CaptureType, const FSpatialPhotoParams& Params)
 	{
-		Progress.EnterProgressFrame(1.f);
-		if (Progress.ShouldCancel())
-		{
-			bWasCancelled = true;
-			return;
-		}
+		FSpatialPhoto3f NewPhoto;
+		NewPhoto.Frame = Params.Frame;
+		NewPhoto.NearPlaneDist = Params.NearPlaneDist;
+		NewPhoto.HorzFOVDegrees = Params.HorzFOVDegrees;
+		NewPhoto.Dimensions = Params.Dimensions;
 
+		// TODO Do something with the success boolean returned by RenderCapture.CaptureFromPosition
+		FImageAdapter Image(&NewPhoto.Image);
+		FRenderCaptureConfig Config = GetCaptureConfig(CaptureType);
 		RenderCapture.SetDimensions(Params.Dimensions);
+		RenderCapture.CaptureFromPosition(CaptureType, NewPhoto.Frame, NewPhoto.HorzFOVDegrees, NewPhoto.NearPlaneDist, Image, Config);
+		GetPhotoSet3f(CaptureType).Add(MoveTemp(NewPhoto));
+	};
 
-		auto CaptureImageTypeFunc_3f = [this, &Progress, &RenderCapture, &Params]
-			(ERenderCaptureType CaptureType, FSpatialPhotoSet3f& PhotoSet)
-		{
-			if (PhotoSet.Num() < PhotoSetParams.Num())
-			{
-				FSpatialPhoto3f NewPhoto;
-				NewPhoto.Frame = Params.Frame;
-				NewPhoto.NearPlaneDist = Params.NearPlaneDist;
-				NewPhoto.HorzFOVDegrees = Params.HorzFOVDegrees;
-				NewPhoto.Dimensions = Params.Dimensions;
-
-				// TODO Do something with the success boolean returned by RenderCapture.CaptureFromPosition
-				FImageAdapter Image(&NewPhoto.Image);
-				FRenderCaptureConfig Config = GetCaptureConfig(CaptureType);
-				RenderCapture.CaptureFromPosition(CaptureType, NewPhoto.Frame, NewPhoto.HorzFOVDegrees, NewPhoto.NearPlaneDist, Image, Config);
-				PhotoSet.Add(MoveTemp(NewPhoto));
-			}
+	auto CapturePhoto1f = [this, &RenderCapture](ERenderCaptureType CaptureType, const FSpatialPhotoParams& Params)
+	{
+		FSpatialPhoto1f NewPhoto;
+		NewPhoto.Frame = Params.Frame;
+		NewPhoto.NearPlaneDist = Params.NearPlaneDist;
+		NewPhoto.HorzFOVDegrees = Params.HorzFOVDegrees;
+		NewPhoto.Dimensions = Params.Dimensions;
 
-			Progress.TickProgress();
-		};
+		// TODO Do something with the success boolean returned by RenderCapture.CaptureFromPosition
+		FImageAdapter Image(&NewPhoto.Image);
+		FRenderCaptureConfig Config = GetCaptureConfig(CaptureType);
+		RenderCapture.SetDimensions(Params.Dimensions);
+		RenderCapture.CaptureFromPosition(CaptureType, NewPhoto.Frame, NewPhoto.HorzFOVDegrees, NewPhoto.NearPlaneDist, Image, Config);
+		GetPhotoSet1f(CaptureType).Add(MoveTemp(NewPhoto));
 
-		auto CaptureImageTypeFunc_1f = [this, &Progress, &RenderCapture, &Params]
-			(ERenderCaptureType CaptureType, FSpatialPhotoSet1f& PhotoSet)
+		if (CaptureType == ERenderCaptureType::DeviceDepth)
 		{
-			// Testing NumDirections is how we currently determine if we need to compute this photo
-			if (PhotoSet.Num() < PhotoSetParams.Num())
-			{
-				FSpatialPhoto1f NewPhoto;
-				NewPhoto.Frame = Params.Frame;
-				NewPhoto.NearPlaneDist = Params.NearPlaneDist;
-				NewPhoto.HorzFOVDegrees = Params.HorzFOVDegrees;
-				NewPhoto.Dimensions = Params.Dimensions;
-
-				// TODO Do something with the success boolean returned by RenderCapture.CaptureFromPosition
-				FImageAdapter Image(&NewPhoto.Image);
-				FRenderCaptureConfig Config = GetCaptureConfig(CaptureType);
-				RenderCapture.CaptureFromPosition(CaptureType, NewPhoto.Frame, NewPhoto.HorzFOVDegrees, NewPhoto.NearPlaneDist, Image, Config);
-				PhotoSet.Add(MoveTemp(NewPhoto));
-			}
-
-			Progress.TickProgress();
-		};
+			PhotoViewMatricies.Add(RenderCapture.GetLastCaptureViewMatrices());
+		}
+	};
 
-		if (bEnableDeviceDepth)
+	// Iterate by channel computing all the photos rather than by photo/viewpoint computing all the channels.
+	// Note: This ordering means that some captures may have computed even if the computation is cancelled.
+	auto CapturePhotoSet = [this, &Progress](
+		TFunctionRef<void(ERenderCaptureType, const FSpatialPhotoParams&)> CapturePhoto,
+		ERenderCaptureType CaptureType,
+		const FText& TaskMessage)
+	{
+		if (PhotoSetStatus[CaptureType] == ECaptureTypeStatus::Pending)
 		{
-			CaptureImageTypeFunc_1f(ERenderCaptureType::DeviceDepth, DeviceDepthPhotoSet);
-			if (PhotoViewMatricies.Num() < PhotoSetParams.Num())
+			Progress.EnterProgressFrame(1.f);
+			FScopedSlowTask PhotoSetProgress(PhotoSetParams.Num(), TaskMessage);
+
+			EmptyPhotoSet(CaptureType);
+			for (const FSpatialPhotoParams& Params : PhotoSetParams)
 			{
-				PhotoViewMatricies.Add(RenderCapture.GetLastCaptureViewMatrices());
+				if (Progress.ShouldCancel())
+				{
+					UE_LOG(LogGeometry, Display, TEXT("FSceneCapturePhotoSet: The pending '%s' step was cancelled"), *TaskMessage.ToString());
+					bWasCancelled = true;
+					return;
+				}
+
+				PhotoSetProgress.EnterProgressFrame(1.f);
+				CapturePhoto(CaptureType, Params);
+				Progress.TickProgress();
 			}
+
+			PhotoSetStatus[CaptureType] = ECaptureTypeStatus::Computed;
 		}
-		if (bEnableBaseColor)
-		{
-			CaptureImageTypeFunc_3f(ERenderCaptureType::BaseColor, BaseColorPhotoSet);
-		}
-		if (bEnableRoughness)
-		{
-			CaptureImageTypeFunc_1f(ERenderCaptureType::Roughness, RoughnessPhotoSet);
-		}
-		if (bEnableSpecular)
-		{
-			CaptureImageTypeFunc_1f(ERenderCaptureType::Specular, SpecularPhotoSet);
-		}
-		if (bEnableMetallic)
-		{
-			CaptureImageTypeFunc_1f(ERenderCaptureType::Metallic, MetallicPhotoSet);
-		}
-		if (bEnablePackedMRS)
-		{
-			CaptureImageTypeFunc_3f(ERenderCaptureType::CombinedMRS, PackedMRSPhotoSet);
-		}
-		if (bEnableWorldNormal)
-		{
-			CaptureImageTypeFunc_3f(ERenderCaptureType::WorldNormal, WorldNormalPhotoSet);
-		}
-		if (bEnableEmissive)
-		{
-			CaptureImageTypeFunc_3f(ERenderCaptureType::Emissive, EmissivePhotoSet);
-		}
-		if (bEnableOpacity)
-		{
-			CaptureImageTypeFunc_1f(ERenderCaptureType::Opacity, OpacityPhotoSet);
-		}
-		if (bEnableSubsurfaceColor)
-		{
-			CaptureImageTypeFunc_3f(ERenderCaptureType::SubsurfaceColor, SubsurfaceColorPhotoSet);
-		}
-	} // end directions loop
+	};
+
+	CapturePhotoSet(CapturePhoto1f, ERenderCaptureType::DeviceDepth,     LOCTEXT("CapturingScene", "Capturing Device Depth"));
+	CapturePhotoSet(CapturePhoto3f, ERenderCaptureType::BaseColor,       LOCTEXT("CapturingScene", "Capturing Base Color"));
+	CapturePhotoSet(CapturePhoto3f, ERenderCaptureType::WorldNormal,     LOCTEXT("CapturingScene", "Capturing World Normal"));
+	CapturePhotoSet(CapturePhoto3f, ERenderCaptureType::CombinedMRS,     LOCTEXT("CapturingScene", "Capturing Packed MRS"));
+	CapturePhotoSet(CapturePhoto1f, ERenderCaptureType::Metallic,        LOCTEXT("CapturingScene", "Capturing Metallic"));
+	CapturePhotoSet(CapturePhoto1f, ERenderCaptureType::Roughness,       LOCTEXT("CapturingScene", "Capturing Roughness"));
+	CapturePhotoSet(CapturePhoto1f, ERenderCaptureType::Specular,        LOCTEXT("CapturingScene", "Capturing Specular"));
+	CapturePhotoSet(CapturePhoto3f, ERenderCaptureType::Emissive,        LOCTEXT("CapturingScene", "Capturing Emissive"));
+	CapturePhotoSet(CapturePhoto1f, ERenderCaptureType::Opacity,         LOCTEXT("CapturingScene", "Capturing Opacity"));
+	CapturePhotoSet(CapturePhoto3f, ERenderCaptureType::SubsurfaceColor, LOCTEXT("CapturingScene", "Capturing Subsurface Color"));
 }
 
 
@@ -584,7 +440,7 @@ void FSceneCapturePhotoSet::AddExteriorCaptures(
 			Progress.TickProgress();
 		};
 
-		if (bEnableDeviceDepth)
+		if (PhotoSetStatus[ERenderCaptureType::DeviceDepth] == ECaptureTypeStatus::Pending)
 		{
 			CaptureImageTypeFunc_1f(ERenderCaptureType::DeviceDepth, DeviceDepthPhotoSet);
 			if (PhotoViewMatricies.Num() < NumDirections)
@@ -592,39 +448,39 @@ void FSceneCapturePhotoSet::AddExteriorCaptures(
 				PhotoViewMatricies.Add(RenderCapture.GetLastCaptureViewMatrices());
 			}
 		}
-		if (bEnableBaseColor)
+		if (PhotoSetStatus[ERenderCaptureType::BaseColor] == ECaptureTypeStatus::Pending)
 		{
 			CaptureImageTypeFunc_3f(ERenderCaptureType::BaseColor, BaseColorPhotoSet);
 		}
-		if (bEnableRoughness)
+		if (PhotoSetStatus[ERenderCaptureType::Roughness] == ECaptureTypeStatus::Pending)
 		{
 			CaptureImageTypeFunc_1f(ERenderCaptureType::Roughness, RoughnessPhotoSet);
 		}
-		if (bEnableSpecular)
+		if (PhotoSetStatus[ERenderCaptureType::Specular] == ECaptureTypeStatus::Pending)
 		{
 			CaptureImageTypeFunc_1f(ERenderCaptureType::Specular, SpecularPhotoSet);
 		}
-		if (bEnableMetallic)
+		if (PhotoSetStatus[ERenderCaptureType::Metallic] == ECaptureTypeStatus::Pending)
 		{
 			CaptureImageTypeFunc_1f(ERenderCaptureType::Metallic, MetallicPhotoSet);
 		}
-		if (bEnablePackedMRS)
+		if (PhotoSetStatus[ERenderCaptureType::CombinedMRS] == ECaptureTypeStatus::Pending)
 		{
 			CaptureImageTypeFunc_3f(ERenderCaptureType::CombinedMRS, PackedMRSPhotoSet);
 		}
-		if (bEnableWorldNormal)
+		if (PhotoSetStatus[ERenderCaptureType::WorldNormal] == ECaptureTypeStatus::Pending)
 		{
 			CaptureImageTypeFunc_3f(ERenderCaptureType::WorldNormal, WorldNormalPhotoSet);
 		}
-		if (bEnableEmissive)
+		if (PhotoSetStatus[ERenderCaptureType::Emissive] == ECaptureTypeStatus::Pending)
 		{
 			CaptureImageTypeFunc_3f(ERenderCaptureType::Emissive, EmissivePhotoSet);
 		}
-		if (bEnableOpacity)
+		if (PhotoSetStatus[ERenderCaptureType::Opacity] == ECaptureTypeStatus::Pending)
 		{
 			CaptureImageTypeFunc_1f(ERenderCaptureType::Opacity, OpacityPhotoSet);
 		}
-		if (bEnableSubsurfaceColor)
+		if (PhotoSetStatus[ERenderCaptureType::SubsurfaceColor] == ECaptureTypeStatus::Pending)
 		{
 			CaptureImageTypeFunc_3f(ERenderCaptureType::SubsurfaceColor, SubsurfaceColorPhotoSet);
 		}
@@ -912,19 +768,10 @@ void FSceneCapturePhotoSet::SetEnableVisibilityByUnregisterMode(bool bEnable)
 
 void FSceneCapturePhotoSet::EmptyAllPhotoSets()
 {
-	BaseColorPhotoSet.Empty();
-	RoughnessPhotoSet.Empty();
-	SpecularPhotoSet.Empty();
-	MetallicPhotoSet.Empty();
-	PackedMRSPhotoSet.Empty();
-	WorldNormalPhotoSet.Empty();
-	EmissivePhotoSet.Empty();
-	OpacityPhotoSet.Empty();
-	SubsurfaceColorPhotoSet.Empty();
-
-	// For the device depth photo set we have two containers to empty
-	DeviceDepthPhotoSet.Empty();
-	PhotoViewMatricies.Empty();
+	ForEachCaptureType([this](ERenderCaptureType CaptureType)
+	{
+		EmptyPhotoSet(CaptureType);
+	});
 }
 
 void FSceneCapturePhotoSet::SetEnableWriteDebugImages(bool bEnable, FString FolderName)
@@ -936,6 +783,89 @@ void FSceneCapturePhotoSet::SetEnableWriteDebugImages(bool bEnable, FString Fold
 	}
 }
 
+FSpatialPhotoSet1f& FSceneCapturePhotoSet::GetPhotoSet1f(ERenderCaptureType CaptureType)
+{
+	switch (CaptureType)
+	{
+	case ERenderCaptureType::Roughness:
+		return RoughnessPhotoSet;
+	case ERenderCaptureType::Metallic:
+		return MetallicPhotoSet;
+	case ERenderCaptureType::Specular:
+		return SpecularPhotoSet;
+	case ERenderCaptureType::Opacity:
+		return OpacityPhotoSet;
+	case ERenderCaptureType::DeviceDepth:
+		return DeviceDepthPhotoSet;
+	default:
+		ensure(false);
+	}
+	return RoughnessPhotoSet;
+}
+
+FSpatialPhotoSet3f& FSceneCapturePhotoSet::GetPhotoSet3f(ERenderCaptureType CaptureType)
+{
+	switch (CaptureType)
+	{
+	case ERenderCaptureType::BaseColor:
+		return BaseColorPhotoSet;
+	case ERenderCaptureType::CombinedMRS:
+		return PackedMRSPhotoSet;
+	case ERenderCaptureType::WorldNormal:
+		return WorldNormalPhotoSet;
+	case ERenderCaptureType::Emissive:
+		return EmissivePhotoSet;
+	case ERenderCaptureType::SubsurfaceColor:
+		return SubsurfaceColorPhotoSet;
+	default:
+		ensure(false);
+	}
+	return BaseColorPhotoSet;
+}
+
+
+void FSceneCapturePhotoSet::EmptyPhotoSet(ERenderCaptureType CaptureType)
+{
+	// Set the functions to empty the photo sets
+	switch (CaptureType)
+	{
+	case ERenderCaptureType::BaseColor:
+		BaseColorPhotoSet.Empty();
+		break;
+	case ERenderCaptureType::WorldNormal:
+		WorldNormalPhotoSet.Empty();
+		break;
+	case ERenderCaptureType::Roughness:
+		RoughnessPhotoSet.Empty();
+		break;
+	case ERenderCaptureType::Metallic:
+		MetallicPhotoSet.Empty();
+		break;
+	case ERenderCaptureType::Specular:
+		SpecularPhotoSet.Empty();
+		break;
+	case ERenderCaptureType::Emissive:
+		EmissivePhotoSet.Empty();
+		break;
+	case ERenderCaptureType::CombinedMRS:
+		PackedMRSPhotoSet.Empty();
+		break;
+	case ERenderCaptureType::Opacity:
+		OpacityPhotoSet.Empty();
+		break;
+	case ERenderCaptureType::SubsurfaceColor:
+		SubsurfaceColorPhotoSet.Empty();
+		break;
+	case ERenderCaptureType::DeviceDepth:
+		// For the device depth photo set we have two containers to empty
+		DeviceDepthPhotoSet.Empty();
+		PhotoViewMatricies.Empty();
+		break;
+	default:
+		ensure(false);
+	}
+}
+
 TArray<FSpatialPhotoParams> UE::Geometry::ComputeStandardExteriorSpatialPhotoParameters(
 	UWorld* World,
 	const TArray<AActor*>& Actors,
diff --git a/Engine/Plugins/Runtime/MeshModelingToolset/Source/ModelingComponents/Private/Scene/WorldRenderCapture.cpp b/Engine/Plugins/Runtime/MeshModelingToolset/Source/ModelingComponents/Private/Scene/WorldRenderCapture.cpp
--- a/Engine/Plugins/Runtime/MeshModelingToolset/Source/ModelingComponents/Private/Scene/WorldRenderCapture.cpp
+++ b/Engine/Plugins/Runtime/MeshModelingToolset/Source/ModelingComponents/Private/Scene/WorldRenderCapture.cpp
@@ -67,41 +67,62 @@ FRenderCaptureTypeFlags FRenderCaptureTypeFlags::Single(ERenderCaptureType Captu
 }
 
 void FRenderCaptureTypeFlags::SetEnabled(ERenderCaptureType CaptureType, bool bEnabled)
+{
+	this->operator[](CaptureType) = bEnabled;
+}
+
+/** @return mutable flag corresponding to the given CaptureType */
+bool& FRenderCaptureTypeFlags::operator[](ERenderCaptureType CaptureType)
 {
 	switch (CaptureType)
 	{
 	case ERenderCaptureType::BaseColor:
-		bBaseColor = bEnabled;
-		break;
+		return bBaseColor;
 	case ERenderCaptureType::WorldNormal:
-		bWorldNormal = bEnabled;
-		break;
+		return bWorldNormal;
 	case ERenderCaptureType::Roughness:
-		bRoughness = bEnabled;
-		break;
+		return bRoughness;
 	case ERenderCaptureType::Metallic:
-		bMetallic = bEnabled;
-		break;
+		return bMetallic;
 	case ERenderCaptureType::Specular:
-		bSpecular = bEnabled;
-		break;
+		return bSpecular;
 	case ERenderCaptureType::Emissive:
-		bEmissive = bEnabled;
-		break;
+		return bEmissive;
 	case ERenderCaptureType::CombinedMRS:
-		bCombinedMRS = bEnabled;
-		break;
+		return bCombinedMRS;
+	case ERenderCaptureType::Opacity:
+		return bOpacity;
 	case ERenderCaptureType::SubsurfaceColor:
-		bSubsurfaceColor = bEnabled;
-		break;
+		return bSubsurfaceColor;
 	case ERenderCaptureType::DeviceDepth:
-		bDeviceDepth = bEnabled;
-		break;
+		return bDeviceDepth;
 	default:
-		check(false);
+		ensure(false);
 	}
+	return bBaseColor;
+}
+
+/** @return constant flag corresponding to the given CaptureType */
+const bool& FRenderCaptureTypeFlags::operator[](ERenderCaptureType CaptureType) const
+{
+	return const_cast<FRenderCaptureTypeFlags*>(this)->operator[](CaptureType);
 }
 
+/** @return true if the flags for this object match the Other object and false otherwise */
+bool FRenderCaptureTypeFlags::operator==(const FRenderCaptureTypeFlags& Other) const
+{
+	bool bEqual = true;
+	ForEachCaptureType([this, &Other, &bEqual](ERenderCaptureType CaptureType)
+	{
+		bEqual = bEqual && (this->operator[](CaptureType) == Other[CaptureType]);
+	});
+	return bEqual;
+}
+
+bool FRenderCaptureTypeFlags::operator!=(const FRenderCaptureTypeFlags& Other) const
+{
+	return !this->operator==(Other);
+}
 
 
 namespace UE
diff --git a/Engine/Plugins/Runtime/MeshModelingToolset/Source/ModelingComponents/Public/Baking/RenderCaptureFunctions.h b/Engine/Plugins/Runtime/MeshModelingToolset/Source/ModelingComponents/Public/Baking/RenderCaptureFunctions.h
--- a/Engine/Plugins/Runtime/MeshModelingToolset/Source/ModelingComponents/Public/Baking/RenderCaptureFunctions.h
+++ b/Engine/Plugins/Runtime/MeshModelingToolset/Source/ModelingComponents/Public/Baking/RenderCaptureFunctions.h
@@ -94,8 +94,89 @@ private:
 
 
 
+struct MODELINGCOMPONENTS_API FSceneCaptureConfig
+{
+	int32 RenderCaptureImageSize = 1024;
+	bool bAntiAliasing = false;
+	double FieldOfViewDegrees = 45.0;
+	double NearPlaneDist = 1.0;
+
+	FRenderCaptureTypeFlags Flags = FRenderCaptureTypeFlags::None();
+
+	bool operator==(const FSceneCaptureConfig&) const;
+	bool operator!=(const FSceneCaptureConfig&) const;
+};
+
+MODELINGCOMPONENTS_API
+void ConfigureSceneCapture(
+	const TUniquePtr<FSceneCapturePhotoSet>& SceneCapture,
+	const TArray<TObjectPtr<AActor>>& Actors,
+	const FSceneCaptureConfig& Config,
+	bool bAllowCancel);
+
+MODELINGCOMPONENTS_API
+FRenderCaptureTypeFlags UpdateSceneCapture(
+	const TUniquePtr<FSceneCapturePhotoSet>& SceneCapture,
+	const TArray<TObjectPtr<AActor>>& Actors,
+	const FSceneCaptureConfig& DesiredConfig,
+	bool bAllowCancel);
+
+MODELINGCOMPONENTS_API
+FSceneCaptureConfig GetSceneCaptureConfig(
+	const TUniquePtr<FSceneCapturePhotoSet>& SceneCapture,
+	FSceneCapturePhotoSet::ECaptureTypeStatus QueryStatus = FSceneCapturePhotoSet::ECaptureTypeStatus::Computed);
+
+// Return a render capture baker, note the lifetime of all arguments such match the lifetime of the returned baker
+MODELINGCOMPONENTS_API
+TUniquePtr<FMeshMapBaker> MakeRenderCaptureBaker(
+	FDynamicMesh3* BaseMesh,
+	TSharedPtr<FMeshTangentsd, ESPMode::ThreadSafe> BaseMeshTangents,
+	TSharedPtr<TArray<int32>, ESPMode::ThreadSafe> BaseMeshUVCharts,
+	FSceneCapturePhotoSet* SceneCapture,
+	FSceneCapturePhotoSetSampler* Sampler,
+	FRenderCaptureTypeFlags PendingBake,
+	int32 TargetUVLayer,
+	EBakeTextureResolution TextureImageSize,
+	EBakeTextureSamplesPerPixel SamplesPerPixel,
+	FRenderCaptureOcclusionHandler* OcclusionHandler);
+
+struct MODELINGCOMPONENTS_API FRenderCaptureTextures
+{
+	UTexture2D* BaseColorMap = nullptr;
+	UTexture2D* NormalMap = nullptr;
+	UTexture2D* PackedMRSMap = nullptr;
+	UTexture2D* MetallicMap = nullptr;
+	UTexture2D* RoughnessMap = nullptr;
+	UTexture2D* SpecularMap = nullptr;
+	UTexture2D* EmissiveMap = nullptr;
+	UTexture2D* OpacityMap = nullptr;
+	UTexture2D* SubsurfaceColorMap = nullptr;
+};
+
+// Note: The source data in the textures is *not* updated by this function
+MODELINGCOMPONENTS_API
+void GetTexturesFromRenderCaptureBaker(
+	const TUniquePtr<FMeshMapBaker>& Baker,
+	FRenderCaptureTextures& TexturesOut);
+
+
+
+
 
-struct MODELINGCOMPONENTS_API FRenderCaptureOptions
+
+
+
+
+
+
+
+////////////////////////////////////////////////////////////////////////////////////////////////////////////////////////
+// The following is all deprecated 
+////////////////////////////////////////////////////////////////////////////////////////////////////////////////////////
+
+PRAGMA_DISABLE_DEPRECATION_WARNINGS
+
+struct UE_DEPRECATED(5.3, "FRenderCaptureOptions is only used by deprecated functions, see those deprecation notes for more info.") MODELINGCOMPONENTS_API FRenderCaptureOptions
 {
 	//
 	// Material approximation settings
@@ -125,7 +206,7 @@ struct MODELINGCOMPONENTS_API FRenderCaptureOptions
 	bool bUsePackedMRS = true;
 };
 
-struct MODELINGCOMPONENTS_API FRenderCaptureUpdate
+struct UE_DEPRECATED(5.3, "FRenderCaptureUpdate is only used by deprecated functions, see those deprecation notes for more info.") MODELINGCOMPONENTS_API FRenderCaptureUpdate
 {
 	// Default to true so that we can use the default value to trigger all update code paths
 	bool bUpdatedBaseColor = true;
@@ -144,6 +225,7 @@ struct MODELINGCOMPONENTS_API FRenderCaptureUpdate
 /**
  * This function computes the SceneCapture for the first time
  */
+UE_DEPRECATED(5.3, "Please use ConfigureSceneCapture and FSceneCapturePhotoSet::Compute instead.")
 MODELINGCOMPONENTS_API
 TUniquePtr<FSceneCapturePhotoSet> CapturePhotoSet(
 	const TArray<TObjectPtr<AActor>>& Actors,
@@ -159,11 +241,8 @@ TUniquePtr<FSceneCapturePhotoSet> CapturePhotoSet(
 * - If the given Options enable a new capture channel then the new photos set are captured and added to the SceneCapture
 * - If the given Actors are different the ones cached in the SceneCapture, or if the given Options changes a parameter
 *   affecting all photo sets (e.g., photo resolution), then all the photo sets are cleared and recomputed
-*   
-* If the SceneCapture computation is cancelled then its state will be identical to how it looked prior to calling this
-* function, which means that queries of the SceneCapture settings (e.g., via the GetXXX() API) will return results which
-* are consistent with the internal photo sets.
 */
+UE_DEPRECATED(5.3, "Please use UpdateSceneCapture instead.")
 MODELINGCOMPONENTS_API
 FRenderCaptureUpdate UpdatePhotoSets(
 	const TUniquePtr<FSceneCapturePhotoSet>& SceneCapture,
@@ -171,11 +250,13 @@ FRenderCaptureUpdate UpdatePhotoSets(
 	const FRenderCaptureOptions& Options,
 	bool bAllowCancel);
 
+UE_DEPRECATED(5.3, "Please use GetSceneCaptureConfig instead.")
 MODELINGCOMPONENTS_API
 FRenderCaptureOptions GetComputedPhotoSetOptions(const TUniquePtr<FSceneCapturePhotoSet>& SceneCapture);
 
 
 // Return a render capture baker, note the lifetime of all arguments such match the lifetime of the returned baker
+UE_DEPRECATED(5.3, "Please use the overload where PendingBake has type FRenderCaptureTypeFlags instead.")
 MODELINGCOMPONENTS_API
 TUniquePtr<FMeshMapBaker> MakeRenderCaptureBaker(
 	FDynamicMesh3* BaseMesh,
@@ -189,30 +270,7 @@ TUniquePtr<FMeshMapBaker> MakeRenderCaptureBaker(
 	EBakeTextureSamplesPerPixel SamplesPerPixel,
 	FRenderCaptureOcclusionHandler* OcclusionHandler);
 
-
-
-
-
-
-struct MODELINGCOMPONENTS_API FRenderCaptureTextures
-{
-	UTexture2D* BaseColorMap = nullptr;
-	UTexture2D* NormalMap = nullptr;
-	UTexture2D* PackedMRSMap = nullptr;
-	UTexture2D* MetallicMap = nullptr;
-	UTexture2D* RoughnessMap = nullptr;
-	UTexture2D* SpecularMap = nullptr;
-	UTexture2D* EmissiveMap = nullptr;
-	UTexture2D* OpacityMap = nullptr;
-	UTexture2D* SubsurfaceColorMap = nullptr;
-};
-
-// Note: The source data in the textures is *not* updated by this function
-MODELINGCOMPONENTS_API
-void GetTexturesFromRenderCaptureBaker(
-	const TUniquePtr<FMeshMapBaker>& Baker,
-	FRenderCaptureTextures& TexturesOut);
-
+PRAGMA_ENABLE_DEPRECATION_WARNINGS
 
 } // end namespace Geometry
 } // end namespace UE
\ No newline at end of file
diff --git a/Engine/Plugins/Runtime/MeshModelingToolset/Source/ModelingComponents/Public/Scene/SceneCapturePhotoSet.h b/Engine/Plugins/Runtime/MeshModelingToolset/Source/ModelingComponents/Public/Scene/SceneCapturePhotoSet.h
--- a/Engine/Plugins/Runtime/MeshModelingToolset/Source/ModelingComponents/Public/Scene/SceneCapturePhotoSet.h
+++ b/Engine/Plugins/Runtime/MeshModelingToolset/Source/ModelingComponents/Public/Scene/SceneCapturePhotoSet.h
@@ -73,7 +73,31 @@ public:
 	 * If bEnabled == false the corresponding photo set will be cleared
 	 */
 	void SetCaptureTypeEnabled(ERenderCaptureType CaptureType, bool bEnabled);
-	bool GetCaptureTypeEnabled(ERenderCaptureType CaptureType) const;
+
+	/** Status of a given capture type */
+	enum class ECaptureTypeStatus {
+		/** The capture type is not enabled */
+		Disabled = 0,
+
+		/** The capture type is enabled and the photo set is computed for the current scene capture parameters */
+		Computed = 1,
+
+		/**
+		 * The capture type is enabled but the photo set is not computed for the current scene capture parameters.
+		 * Note: This state is encountered if 1) a scene capture parameter which invalidates the photo set is changed
+		 * but Compute has not yet been called or 2) the Compute function was cancelled before the photo set was computed
+		 */
+		Pending = 2
+	};
+
+	/** Returns the current status of the given capture type */
+	ECaptureTypeStatus GetCaptureTypeStatus(ERenderCaptureType CaptureType) const;
+
+	/** Status of the overall scene capture i.e., the combined status of all capture types */
+	using FStatus = TRenderCaptureTypeData<ECaptureTypeStatus>;
+
+	/** Returns the current status of all capture types */
+	FStatus GetSceneCaptureStatus() const;
 
 	/**
 	 * Configure the given capture type
@@ -227,16 +251,7 @@ protected:
 
 	bool bEnforceVisibilityViaUnregister = false;
 
-	bool bEnableBaseColor = true;
-	bool bEnableRoughness = false;
-	bool bEnableSpecular = false;
-	bool bEnableMetallic = false;
-	bool bEnablePackedMRS = true;
-	bool bEnableWorldNormal = true;
-	bool bEnableEmissive = true;
-	bool bEnableOpacity = true;
-	bool bEnableSubsurfaceColor = true;
-	bool bEnableDeviceDepth = true;
+	FStatus PhotoSetStatus;
 
 	FSpatialPhotoSet3f BaseColorPhotoSet;
 	FSpatialPhotoSet1f RoughnessPhotoSet;
@@ -249,16 +264,7 @@ protected:
 	FSpatialPhotoSet3f SubsurfaceColorPhotoSet;
 	FSpatialPhotoSet1f DeviceDepthPhotoSet;
 
-	FRenderCaptureConfig BaseColorConfig;
-	FRenderCaptureConfig RoughnessConfig;
-	FRenderCaptureConfig SpecularConfig;
-	FRenderCaptureConfig MetallicConfig;
-	FRenderCaptureConfig PackedMRSConfig;
-	FRenderCaptureConfig WorldNormalConfig;
-	FRenderCaptureConfig EmissiveConfig;
-	FRenderCaptureConfig OpacityConfig;
-	FRenderCaptureConfig SubsurfaceColorConfig;
-	FRenderCaptureConfig DeviceDepthConfig;
+	TRenderCaptureTypeData<FRenderCaptureConfig> RenderCaptureConfig;
 
 	TArray<FSpatialPhotoParams> PhotoSetParams;
 
@@ -272,6 +278,13 @@ protected:
 	bool bWasCancelled = false;
 
 private:
+
+	FSpatialPhotoSet1f& GetPhotoSet1f(ERenderCaptureType CaptureType);
+
+	FSpatialPhotoSet3f& GetPhotoSet3f(ERenderCaptureType CaptureType);
+
+	void EmptyPhotoSet(ERenderCaptureType CaptureType);
+
 	void EmptyAllPhotoSets();
 };
 
diff --git a/Engine/Plugins/Runtime/MeshModelingToolset/Source/ModelingComponents/Public/Scene/WorldRenderCapture.h b/Engine/Plugins/Runtime/MeshModelingToolset/Source/ModelingComponents/Public/Scene/WorldRenderCapture.h
--- a/Engine/Plugins/Runtime/MeshModelingToolset/Source/ModelingComponents/Public/Scene/WorldRenderCapture.h
+++ b/Engine/Plugins/Runtime/MeshModelingToolset/Source/ModelingComponents/Public/Scene/WorldRenderCapture.h
@@ -30,11 +30,37 @@ enum class ERenderCaptureType
 	SubsurfaceColor = 512
 };
 
+// Call the given Function for each capture type/channel
+template <typename FunctionType>
+void ForEachCaptureType(const FunctionType& Function)
+{
+	Function(ERenderCaptureType::BaseColor);
+	Function(ERenderCaptureType::Roughness);
+	Function(ERenderCaptureType::Specular);
+	Function(ERenderCaptureType::Metallic);
+	Function(ERenderCaptureType::CombinedMRS);
+	Function(ERenderCaptureType::WorldNormal);
+	Function(ERenderCaptureType::Emissive);
+	Function(ERenderCaptureType::Opacity);
+	Function(ERenderCaptureType::SubsurfaceColor);
+	Function(ERenderCaptureType::DeviceDepth);
+}
+
 struct MODELINGCOMPONENTS_API FRenderCaptureConfig
 {
 	// You might want to disable this if you're using FMeshMapBaker because it supports multi-sampling in a way which
 	// will avoid blending pixels at different depths. This option is ignored for ERenderCaptureType::DeviceDepth
 	bool bAntiAliasing = true;
+
+	bool operator==(const FRenderCaptureConfig& Other) const
+	{
+		return bAntiAliasing == Other.bAntiAliasing;
+	}
+
+	bool operator!=(const FRenderCaptureConfig& Other) const
+	{
+		return !operator==(Other);
+	}
 };
 
 FRenderCaptureConfig MODELINGCOMPONENTS_API GetDefaultRenderCaptureConfig(ERenderCaptureType CaptureType);
@@ -72,6 +98,71 @@ struct MODELINGCOMPONENTS_API FRenderCaptureTypeFlags
 
 	/** Set the indicated CaptureType to enabled */
 	void SetEnabled(ERenderCaptureType CaptureType, bool bEnabled);
+
+	/** @return mutable flag corresponding to the given CaptureType */
+	bool& operator[](ERenderCaptureType CaptureType);
+
+	/** @return constant flag corresponding to the given CaptureType */
+	const bool& operator[](ERenderCaptureType CaptureType) const;
+
+	/** @return true if the flags for this object match the Other object and false otherwise */
+	bool operator==(const FRenderCaptureTypeFlags& Other) const;
+	bool operator!=(const FRenderCaptureTypeFlags& Other) const;
+};
+
+template <typename T>
+struct MODELINGCOMPONENTS_API TRenderCaptureTypeData
+{
+	// FRenderCaptureTypeFlags uses the correct member naming convention and has additional boolean helpers
+	static_assert(std::is_same_v<T, bool> == false, "Use FRenderCaptureTypeFlags instead");
+
+	T BaseColor;
+	T Roughness;
+	T Specular;
+	T Metallic;
+	T CombinedMRS;
+	T WorldNormal;
+	T Emissive;
+	T Opacity;
+	T SubsurfaceColor;
+	T DeviceDepth;
+
+	/** @return mutable data corresponding to the given CaptureType */
+	FORCEINLINE T& operator[](ERenderCaptureType CaptureType)
+	{
+		switch (CaptureType)
+		{
+		case ERenderCaptureType::BaseColor:
+			return BaseColor;
+		case ERenderCaptureType::WorldNormal:
+			return WorldNormal;
+		case ERenderCaptureType::Roughness:
+			return Roughness;
+		case ERenderCaptureType::Metallic:
+			return Metallic;
+		case ERenderCaptureType::Specular:
+			return Specular;
+		case ERenderCaptureType::Emissive:
+			return Emissive;
+		case ERenderCaptureType::CombinedMRS:
+			return CombinedMRS;
+		case ERenderCaptureType::Opacity:
+			return Opacity;
+		case ERenderCaptureType::SubsurfaceColor:
+			return SubsurfaceColor;
+		case ERenderCaptureType::DeviceDepth:
+			return DeviceDepth;
+		default:
+			ensure(false);
+		}
+		return BaseColor;
+	}
+
+	/** @return constant data corresponding to the given CaptureType */
+	FORCEINLINE const T& operator[](ERenderCaptureType CaptureType) const
+	{
+		return const_cast<TRenderCaptureTypeData<T>*>(this)->operator[](CaptureType);
+	}
 };
 
 
