EditorModelingObjectsCreationAPI: Fix crash due to nullptr dereference on accepting the Pattern Tool.

diff --git a/Engine/Plugins/Runtime/MeshModelingToolset/Source/ModelingComponentsEditorOnly/Private/EditorModelingObjectsCreationAPI.cpp b/Engine/Plugins/Runtime/MeshModelingToolset/Source/ModelingComponentsEditorOnly/Private/EditorModelingObjectsCreationAPI.cpp
--- a/Engine/Plugins/Runtime/MeshModelingToolset/Source/ModelingComponentsEditorOnly/Private/EditorModelingObjectsCreationAPI.cpp
+++ b/Engine/Plugins/Runtime/MeshModelingToolset/Source/ModelingComponentsEditorOnly/Private/EditorModelingObjectsCreationAPI.cpp
@@ -504,12 +504,16 @@ FCreateActorResult UEditorModelingObjectsCreationAPI::CreateNewActor(FCreateActo
 {
 	// create new Actor
 	AActor* NewActor = nullptr;
-	if (GEditor)
+	
+	if (GEditor && CreateActorParams.TemplateActor)
 	{
-		if (UActorFactory* const ActorFactory = GEditor->FindActorFactoryForActorClass(CreateActorParams.TemplateActor->GetClass()))
+		if (UClass* ActorClass = CreateActorParams.TemplateActor->GetClass())
 		{
-			NewActor = ActorFactory->CreateActor(nullptr, CreateActorParams.TargetWorld->GetCurrentLevel(), CreateActorParams.Transform);
-			FActorLabelUtilities::SetActorLabelUnique(NewActor, CreateActorParams.BaseName);
+			if (UActorFactory* const ActorFactory = GEditor->FindActorFactoryForActorClass(ActorClass))
+			{
+				NewActor = ActorFactory->CreateActor(ActorClass, CreateActorParams.TargetWorld->GetCurrentLevel(), CreateActorParams.Transform);
+				FActorLabelUtilities::SetActorLabelUnique(NewActor, CreateActorParams.BaseName);
+			}
 		}
 	}
 
