Fix horde issue 337708 related to text conflicts in LOCTEXT

#horde 337708

diff --git a/Engine/Plugins/Runtime/MeshModelingToolset/Source/ModelingComponents/Private/Scene/SceneCapturePhotoSet.cpp b/Engine/Plugins/Runtime/MeshModelingToolset/Source/ModelingComponents/Private/Scene/SceneCapturePhotoSet.cpp
--- a/Engine/Plugins/Runtime/MeshModelingToolset/Source/ModelingComponents/Private/Scene/SceneCapturePhotoSet.cpp
+++ b/Engine/Plugins/Runtime/MeshModelingToolset/Source/ModelingComponents/Private/Scene/SceneCapturePhotoSet.cpp
@@ -256,16 +256,16 @@ void FSceneCapturePhotoSet::Compute()
 		}
 	};
 
-	CapturePhotoSet(CapturePhoto1f, ERenderCaptureType::DeviceDepth,     LOCTEXT("CapturingScene", "Capturing Device Depth"));
-	CapturePhotoSet(CapturePhoto3f, ERenderCaptureType::BaseColor,       LOCTEXT("CapturingScene", "Capturing Base Color"));
-	CapturePhotoSet(CapturePhoto3f, ERenderCaptureType::WorldNormal,     LOCTEXT("CapturingScene", "Capturing World Normal"));
-	CapturePhotoSet(CapturePhoto3f, ERenderCaptureType::CombinedMRS,     LOCTEXT("CapturingScene", "Capturing Packed MRS"));
-	CapturePhotoSet(CapturePhoto1f, ERenderCaptureType::Metallic,        LOCTEXT("CapturingScene", "Capturing Metallic"));
-	CapturePhotoSet(CapturePhoto1f, ERenderCaptureType::Roughness,       LOCTEXT("CapturingScene", "Capturing Roughness"));
-	CapturePhotoSet(CapturePhoto1f, ERenderCaptureType::Specular,        LOCTEXT("CapturingScene", "Capturing Specular"));
-	CapturePhotoSet(CapturePhoto3f, ERenderCaptureType::Emissive,        LOCTEXT("CapturingScene", "Capturing Emissive"));
-	CapturePhotoSet(CapturePhoto1f, ERenderCaptureType::Opacity,         LOCTEXT("CapturingScene", "Capturing Opacity"));
-	CapturePhotoSet(CapturePhoto3f, ERenderCaptureType::SubsurfaceColor, LOCTEXT("CapturingScene", "Capturing Subsurface Color"));
+	CapturePhotoSet(CapturePhoto1f, ERenderCaptureType::DeviceDepth, LOCTEXT("CapturingScene_DeviceDepth", "Capturing Device Depth"));
+	CapturePhotoSet(CapturePhoto3f, ERenderCaptureType::BaseColor, LOCTEXT("CapturingScene_BaseColor", "Capturing Base Color"));
+	CapturePhotoSet(CapturePhoto3f, ERenderCaptureType::WorldNormal, LOCTEXT("CapturingScene_WorldNormal", "Capturing World Normal"));
+	CapturePhotoSet(CapturePhoto3f, ERenderCaptureType::CombinedMRS, LOCTEXT("CapturingScene_CombinedMRS", "Capturing Packed MRS"));
+	CapturePhotoSet(CapturePhoto1f, ERenderCaptureType::Metallic, LOCTEXT("CapturingScene_Metallic", "Capturing Metallic"));
+	CapturePhotoSet(CapturePhoto1f, ERenderCaptureType::Roughness, LOCTEXT("CapturingScene_Roughness", "Capturing Roughness"));
+	CapturePhotoSet(CapturePhoto1f, ERenderCaptureType::Specular, LOCTEXT("CapturingScene_Specular", "Capturing Specular"));
+	CapturePhotoSet(CapturePhoto3f, ERenderCaptureType::Emissive, LOCTEXT("CapturingScene_Emissive", "Capturing Emissive"));
+	CapturePhotoSet(CapturePhoto1f, ERenderCaptureType::Opacity, LOCTEXT("CapturingScene_Opacity", "Capturing Opacity"));
+	CapturePhotoSet(CapturePhoto3f, ERenderCaptureType::SubsurfaceColor, LOCTEXT("CapturingScene_SubsurfaceColor", "Capturing Subsurface Color"));
 }
 
 
