Skin Bind: fixed geodesic binding crash when the sampled bone's bbox voxels are out of the occupancy grid

this is easy to repro using a single mesh with two bones (one in the center of the mesh and the other one away from the mesh bbox)

#rnx

diff --git a/Engine/Plugins/Runtime/MeshModelingToolset/Source/ModelingOperators/Private/SkinningOps/SkinBindingOp.cpp b/Engine/Plugins/Runtime/MeshModelingToolset/Source/ModelingOperators/Private/SkinningOps/SkinBindingOp.cpp
--- a/Engine/Plugins/Runtime/MeshModelingToolset/Source/ModelingOperators/Private/SkinningOps/SkinBindingOp.cpp
+++ b/Engine/Plugins/Runtime/MeshModelingToolset/Source/ModelingOperators/Private/SkinningOps/SkinBindingOp.cpp
@@ -499,11 +499,14 @@ void FSkinBindingOp::CreateSkinWeights_GeodesicVoxel(
 				for (int32 K = BoneMin.Z; K <= BoneMax.Z; K++)
 				{
 					const FVector3i Candidate(I, J, K);
-					const FBox3d CellBox{Occupancy.GetCellBoxFromIndex(Candidate)};
-					if (Skeleton.GetBoneFanIntersectsBox(BoneIndex, CellBox))
+					if (BoneDistance.IsValidIndex(Candidate))
 					{
-						WorkingSet.Push(Candidate);
-						BoneDistance[Candidate] = 0.0f;
+						const FBox3d CellBox{Occupancy.GetCellBoxFromIndex(Candidate)};
+						if (Skeleton.GetBoneFanIntersectsBox(BoneIndex, CellBox))
+						{
+							WorkingSet.Push(Candidate);
+							BoneDistance[Candidate] = 0.0f;
+						}
 					}
 				}
 			}
