Make the Modeling Tools Rounded Rectangle primitive preserve its Width/Depth by default as the Corner Radius changes. Add an option to enable the old behavior, where the Corner Radius extends beyond the Width/Depth.

diff --git a/Engine/Plugins/Runtime/MeshModelingToolset/Source/MeshModelingTools/Private/AddPrimitiveTool.cpp b/Engine/Plugins/Runtime/MeshModelingToolset/Source/MeshModelingTools/Private/AddPrimitiveTool.cpp
--- a/Engine/Plugins/Runtime/MeshModelingToolset/Source/MeshModelingTools/Private/AddPrimitiveTool.cpp
+++ b/Engine/Plugins/Runtime/MeshModelingToolset/Source/MeshModelingTools/Private/AddPrimitiveTool.cpp
@@ -519,10 +519,25 @@ void UAddRectanglePrimitiveTool::GenerateMesh(FDynamicMesh3* OutMesh) const
 		FRoundedRectangleMeshGenerator RectGen;
 		RectGen.Width = RectangleSettings->Depth;
 		RectGen.Height = RectangleSettings->Width;
+		RectGen.Radius = RectangleSettings->CornerRadius;
+		if (RectangleSettings->bMaintainDimension)
+		{
+			constexpr double MinSide = UE_DOUBLE_KINDA_SMALL_NUMBER;
+			constexpr double MinRadius = MinSide * .5;
+			double SmallSide = FMath::Min(RectGen.Width, RectGen.Height);
+			RectGen.Radius = FMath::Clamp(RectGen.Radius, MinRadius, (SmallSide - MinSide) * .5);
+			RectGen.Width = FMath::Max(MinSide, RectGen.Width - RectGen.Radius * 2);
+			RectGen.Height = FMath::Max(MinSide, RectGen.Height - RectGen.Radius * 2);
+			// Update the Radius in the UI, if it was clamped
+			if (RectGen.Radius != RectangleSettings->CornerRadius)
+			{
+				RectangleSettings->CornerRadius = RectGen.Radius;
+				NotifyOfPropertyChangeByTool(RectangleSettings);
+			}
+		}
 		RectGen.WidthVertexCount = RectangleSettings->DepthSubdivisions + 1;
 		RectGen.HeightVertexCount = RectangleSettings->WidthSubdivisions + 1;
 		RectGen.bSinglePolyGroup = (ShapeSettings->PolygroupMode != EMakeMeshPolygroupMode::PerQuad);
-		RectGen.Radius = RectangleSettings->CornerRadius;
 		RectGen.AngleSamples = RectangleSettings->CornerSlices - 1;
 		RectGen.Generate();
 		OutMesh->Copy(&RectGen);
diff --git a/Engine/Plugins/Runtime/MeshModelingToolset/Source/MeshModelingTools/Public/AddPrimitiveTool.h b/Engine/Plugins/Runtime/MeshModelingToolset/Source/MeshModelingTools/Public/AddPrimitiveTool.h
--- a/Engine/Plugins/Runtime/MeshModelingToolset/Source/MeshModelingTools/Public/AddPrimitiveTool.h
+++ b/Engine/Plugins/Runtime/MeshModelingToolset/Source/MeshModelingTools/Public/AddPrimitiveTool.h
@@ -176,6 +176,10 @@ public:
 	UPROPERTY(EditAnywhere, Category = Shape, meta = (UIMin = "1", UIMax = "100", ClampMin = "1", ClampMax = "500", ProceduralShapeSetting))
 	int DepthSubdivisions = 1;
 
+	/** Whether to preserve the overall Width and Depth for a Rounded Rectangle, or to allow the rounded corners to extend outside those dimensions. */
+	UPROPERTY(EditAnywhere, Category = Shape, meta = (EditCondition = "RectangleType == EProceduralRectType::RoundedRectangle"))
+	bool bMaintainDimension = true;
+
 	/** Radius of rounded corners. This is only available for Rounded Rectangles. */
 	UPROPERTY(EditAnywhere, Category = Shape,
 		meta = (UIMin = "1.0", UIMax = "1000.0", ClampMin = "0.0001", ClampMax = "1000000.0", ProceduralShapeSetting,
