Adding select shell / grow / shrink to vertex selection modes.

diff --git a/Engine/Plugins/Experimental/MeshModelingToolsetExp/Source/MeshModelingToolsEditorOnlyExp/Private/SkeletalMesh/SkinWeightsPaintTool.cpp b/Engine/Plugins/Experimental/MeshModelingToolsetExp/Source/MeshModelingToolsEditorOnlyExp/Private/SkeletalMesh/SkinWeightsPaintTool.cpp
--- a/Engine/Plugins/Experimental/MeshModelingToolsetExp/Source/MeshModelingToolsEditorOnlyExp/Private/SkeletalMesh/SkinWeightsPaintTool.cpp
+++ b/Engine/Plugins/Experimental/MeshModelingToolsetExp/Source/MeshModelingToolsEditorOnlyExp/Private/SkeletalMesh/SkinWeightsPaintTool.cpp
@@ -2033,6 +2033,11 @@ void USkinWeightsPaintTool::ToggleEditingMode()
 	}
 }
 
+TObjectPtr<UPolygonSelectionMechanic> USkinWeightsPaintTool::GetSelectionMechanic()
+{
+	return PolygonSelectionMechanic;
+}
+
 void USkinWeightsPaintTool::GetSelectedVertices(TArray<int32>& OutVertexIndices) const
 {
 	const FGroupTopologySelection& Selection = PolygonSelectionMechanic->GetActiveSelection();
diff --git a/Engine/Plugins/Experimental/MeshModelingToolsetExp/Source/MeshModelingToolsEditorOnlyExp/Public/SkeletalMesh/SkinWeightsPaintTool.h b/Engine/Plugins/Experimental/MeshModelingToolsetExp/Source/MeshModelingToolsEditorOnlyExp/Public/SkeletalMesh/SkinWeightsPaintTool.h
--- a/Engine/Plugins/Experimental/MeshModelingToolsetExp/Source/MeshModelingToolsEditorOnlyExp/Public/SkeletalMesh/SkinWeightsPaintTool.h
+++ b/Engine/Plugins/Experimental/MeshModelingToolsetExp/Source/MeshModelingToolsEditorOnlyExp/Public/SkeletalMesh/SkinWeightsPaintTool.h
@@ -370,6 +370,9 @@ public:
 	// toggle brush / selection mode
 	void ToggleEditingMode();
 
+	// edit selection
+	TObjectPtr<UPolygonSelectionMechanic> GetSelectionMechanic();
+
 	// get a list of currently selected vertices
 	void GetSelectedVertices(TArray<int32>& OutVertexIndices) const;
 
diff --git a/Engine/Plugins/Experimental/MeshModelingToolsetExp/Source/ModelingEditorUI/Private/SkinWeightDetailCustomization.cpp b/Engine/Plugins/Experimental/MeshModelingToolsetExp/Source/ModelingEditorUI/Private/SkinWeightDetailCustomization.cpp
--- a/Engine/Plugins/Experimental/MeshModelingToolsetExp/Source/ModelingEditorUI/Private/SkinWeightDetailCustomization.cpp
+++ b/Engine/Plugins/Experimental/MeshModelingToolsetExp/Source/ModelingEditorUI/Private/SkinWeightDetailCustomization.cpp
@@ -8,6 +8,7 @@
 #include "Widgets/Input/SSegmentedControl.h"
 #include "SkeletalMesh/SkinWeightsPaintTool.h"
 #include "SSkinWeightProfileImportOptions.h"
+#include "Selection/PolygonSelectionMechanic.h"
 #include "UObject/UnrealTypePrivate.h"
 #include "Widgets/Input/SSpinBox.h"
 #include "Widgets/Input/SButton.h"
@@ -306,7 +307,66 @@ void FSkinWeightDetailCustomization::AddBrushUI(IDetailLayoutBuilder& DetailBuil
 void FSkinWeightDetailCustomization::AddSelectionUI(IDetailLayoutBuilder& DetailBuilder)
 {
 	// custom display of weight editing tools
-	IDetailCategoryBuilder& EditWeightsCategory = DetailBuilder.EditCategory("EditWeights", FText::GetEmpty(), ECategoryPriority::Important);
+	IDetailCategoryBuilder& EditSelectionCategory = DetailBuilder.EditCategory("Edit Selection", FText::GetEmpty(), ECategoryPriority::Important);
+	EditSelectionCategory.InitiallyCollapsed(true);
+
+	// GROW/SHRINK/FLOOD Selection category
+	EditSelectionCategory.AddCustomRow(LOCTEXT("EditSelectionRow", "Edit Selection"), false)
+	.WholeRowContent()
+	[
+		SNew(SHorizontalBox)
+
+		+SHorizontalBox::Slot()
+		.Padding(2.f, WeightEditVerticalPadding)
+		[
+			SNew(SButton)
+			.HAlign(HAlign_Center)
+			.VAlign(VAlign_Center)
+			.Text(LOCTEXT("GrowSelectionButtonLabel", "Grow"))
+			.ToolTipText(LOCTEXT("GrowSelectionTooltip",
+					"Grow the current selection by adding connected neighbors to current selection.\n"))
+			.OnClicked_Lambda([this]()
+			{
+				SkinToolSettings->WeightTool->GetSelectionMechanic()->GrowSelection();
+				return FReply::Handled();
+			})
+		]
+
+		+SHorizontalBox::Slot()
+		.Padding(2.f, WeightEditVerticalPadding)
+		[
+			SNew(SButton)
+			.HAlign(HAlign_Center)
+			.VAlign(VAlign_Center)
+			.Text(LOCTEXT("ShrinkSelectionButtonLabel", "Shrink"))
+			.ToolTipText(LOCTEXT("ShrinkSelectionTooltip",
+					"Shrink the current selection by removing vertices on the border of the current selection.\n"))
+			.OnClicked_Lambda([this]()
+			{
+				SkinToolSettings->WeightTool->GetSelectionMechanic()->ShrinkSelection();
+				return FReply::Handled();
+			})
+		]
+
+		+SHorizontalBox::Slot()
+		.Padding(2.f, WeightEditVerticalPadding)
+		[
+			SNew(SButton)
+			.HAlign(HAlign_Center)
+			.VAlign(VAlign_Center)
+			.Text(LOCTEXT("FloodSelectionButtonLabel", "Flood"))
+			.ToolTipText(LOCTEXT("FloodSelectionTooltip",
+					"Flood the current selection by adding all connected vertices to the current selection.\n"))
+			.OnClicked_Lambda([this]()
+			{
+				SkinToolSettings->WeightTool->GetSelectionMechanic()->FloodSelection();
+				return FReply::Handled();
+			})
+		]
+	];
+
+	// custom display of weight editing tools
+	IDetailCategoryBuilder& EditWeightsCategory = DetailBuilder.EditCategory("Edit Weights", FText::GetEmpty(), ECategoryPriority::Important);
 	EditWeightsCategory.InitiallyCollapsed(true);
 
 	// AVERAGE/RELAX/NORMALIZE WEIGHTS category
diff --git a/Engine/Plugins/Runtime/MeshModelingToolset/Source/ModelingComponents/Private/Selection/MeshTopologySelectionMechanic.cpp b/Engine/Plugins/Runtime/MeshModelingToolset/Source/ModelingComponents/Private/Selection/MeshTopologySelectionMechanic.cpp
--- a/Engine/Plugins/Runtime/MeshModelingToolset/Source/ModelingComponents/Private/Selection/MeshTopologySelectionMechanic.cpp
+++ b/Engine/Plugins/Runtime/MeshModelingToolset/Source/ModelingComponents/Private/Selection/MeshTopologySelectionMechanic.cpp
@@ -451,6 +451,123 @@ void UMeshTopologySelectionMechanic::SelectAll()
 	ParentTool->GetToolManager()->EndUndoTransaction();
 }
 
+void UMeshTopologySelectionMechanic::GrowSelection()
+{
+	ParentTool->GetToolManager()->BeginUndoTransaction(LOCTEXT("GrowSelectionChange", "Grow Selection"));
+	BeginChange();
+
+	// add the neighbor vertices to the current selection
+	if (Properties->bSelectVertices)
+	{
+		TSet<int32> VerticesToAdd;
+		for (const int32 VertexIndex : PersistentSelection.SelectedCornerIDs)
+		{
+			for (const int32 NeighborIndex : Mesh->VtxVerticesItr(VertexIndex))
+			{
+				if (!PersistentSelection.SelectedCornerIDs.Contains(NeighborIndex))
+				{
+					VerticesToAdd.Add(NeighborIndex);
+				}
+			}	
+		}
+		
+		PersistentSelection.SelectedCornerIDs.Append(VerticesToAdd);
+	}
+
+	if (Properties->bSelectEdges || Properties->bSelectFaces)
+	{
+		// TODO add support for growing edge/face selections if your tool requires it.
+		// growing edge/face selection not yet supported
+		checkNoEntry();
+	}
+
+	SelectionTimestamp++;
+	OnSelectionChanged.Broadcast();
+	EndChangeAndEmitIfModified();
+	ParentTool->GetToolManager()->EndUndoTransaction();
+}
+
+void UMeshTopologySelectionMechanic::ShrinkSelection()
+{
+	ParentTool->GetToolManager()->BeginUndoTransaction(LOCTEXT("ShrinkSelectionChange", "Shrink Selection"));
+	BeginChange();
+
+	// remove border vertices from the current selection
+	if (Properties->bSelectVertices)
+	{
+		TSet<int32> BorderVertices;
+		for (const int32 VertexIndex : PersistentSelection.SelectedCornerIDs)
+		{
+			for (const int32 NeighborIndex : Mesh->VtxVerticesItr(VertexIndex))
+			{
+				if (!PersistentSelection.SelectedCornerIDs.Contains(NeighborIndex))
+				{
+					BorderVertices.Add(VertexIndex);
+					break;
+				}
+			}	
+		}
+		
+		for (const int32 BorderVertex : BorderVertices)
+		{
+			PersistentSelection.SelectedCornerIDs.Remove(BorderVertex);
+		}
+	}
+
+	if (Properties->bSelectEdges || Properties->bSelectFaces)
+	{
+		// TODO add support for shrinking edge/face selections if your tool requires it.
+		// shrinking edge/face selection not yet supported
+		checkNoEntry();
+	}
+
+	SelectionTimestamp++;
+	OnSelectionChanged.Broadcast();
+	EndChangeAndEmitIfModified();
+	ParentTool->GetToolManager()->EndUndoTransaction();
+}
+
+void UMeshTopologySelectionMechanic::FloodSelection()
+{
+	ParentTool->GetToolManager()->BeginUndoTransaction(LOCTEXT("FloodSelectionChange", "Flood Selection"));
+	BeginChange();
+
+	if (Properties->bSelectVertices)
+	{
+		TSet<int32>& SelectedVertices = PersistentSelection.SelectedCornerIDs;
+		TSet<int32> VerticesAddedPrevIteration = PersistentSelection.SelectedCornerIDs;
+		while(!VerticesAddedPrevIteration.IsEmpty())
+		{
+			TSet<int32> VerticesToAddThisIteration;
+			for (const int32 VertexAdded : VerticesAddedPrevIteration)
+			{
+				for (const int32 NeighborIndex : Mesh->VtxVerticesItr(VertexAdded))
+				{
+					if (!SelectedVertices.Contains(NeighborIndex))
+					{
+						VerticesToAddThisIteration.Add(NeighborIndex);
+					}
+				}
+			}
+			
+			SelectedVertices.Append(VerticesToAddThisIteration);
+			VerticesAddedPrevIteration = VerticesToAddThisIteration;
+		}
+	}
+
+	if (Properties->bSelectEdges || Properties->bSelectFaces)
+	{
+		// TODO add support for flooding edge/face selections if your tool requires it.
+		// flooding edge/face selection not yet supported
+		checkNoEntry();
+	}
+
+	SelectionTimestamp++;
+	OnSelectionChanged.Broadcast();
+	EndChangeAndEmitIfModified();
+	ParentTool->GetToolManager()->EndUndoTransaction();
+}
+
 FInputRayHit UMeshTopologySelectionMechanic::IsHitByClick(const FInputDeviceRay& ClickPos)
 {
 	if (!bIsEnabled)
diff --git a/Engine/Plugins/Runtime/MeshModelingToolset/Source/ModelingComponents/Public/Selection/MeshTopologySelectionMechanic.h b/Engine/Plugins/Runtime/MeshModelingToolset/Source/ModelingComponents/Public/Selection/MeshTopologySelectionMechanic.h
--- a/Engine/Plugins/Runtime/MeshModelingToolset/Source/ModelingComponents/Public/Selection/MeshTopologySelectionMechanic.h
+++ b/Engine/Plugins/Runtime/MeshModelingToolset/Source/ModelingComponents/Public/Selection/MeshTopologySelectionMechanic.h
@@ -256,6 +256,10 @@ public:
 	void InvertSelection();
 	void SelectAll();
 
+	void GrowSelection();
+	void ShrinkSelection();
+	void FloodSelection();
+
 	/** 
 	 * @return true if the current selection is non-empty 
 	 */
