MeshModelingToolset - Render Capture: Use of the fallback material could be captured in some cases, as the first render could potentially trigger shader(s) compilation.
* After the render, test if there are assets pending compilation. If so, wait for compilation to complete and restart the scene capture.

diff --git a/Engine/Plugins/Runtime/MeshModelingToolset/Source/ModelingComponents/Private/Scene/WorldRenderCapture.cpp b/Engine/Plugins/Runtime/MeshModelingToolset/Source/ModelingComponents/Private/Scene/WorldRenderCapture.cpp
--- a/Engine/Plugins/Runtime/MeshModelingToolset/Source/ModelingComponents/Private/Scene/WorldRenderCapture.cpp
+++ b/Engine/Plugins/Runtime/MeshModelingToolset/Source/ModelingComponents/Private/Scene/WorldRenderCapture.cpp
@@ -21,6 +21,7 @@
 
 #include "SceneViewExtension.h"
 #include "RenderCaptureInterface.h" // For debugging with RenderCaptureInterface::FScopedCapture
+#include "AssetCompilingManager.h"
 
 using namespace UE::Geometry;
 
@@ -375,12 +376,26 @@ namespace Internal
  */
 static void PerformSceneRender(FCanvas* Canvas, FSceneViewFamily* ViewFamily)
 {
-	for (int32 i = 0; i < CVarModelingWorldRenderCaptureVTWarmupFrames.GetValueOnGameThread(); i++)
+	bool bCompiledAssets = false;
+
+	do
 	{
-		GetRendererModule().BeginRenderingViewFamily(Canvas, ViewFamily);
-	}
+		int32 NumRender = 1 + CVarModelingWorldRenderCaptureVTWarmupFrames.GetValueOnGameThread();
+		for (int32 i = 0; i < NumRender; i++)
+		{
+			GetRendererModule().BeginRenderingViewFamily(Canvas, ViewFamily);
+		}
 
-	GetRendererModule().BeginRenderingViewFamily(Canvas, ViewFamily);
+		// Flush rendering commands, may queue assets compilation (shaders)
+		FlushRenderingCommands();
+
+		// If there are assets to be compiled, we need to wait for compilation to finish and start render over
+		bCompiledAssets = FAssetCompilingManager::Get().GetNumRemainingAssets() > 0;
+		if (bCompiledAssets)
+		{
+			FAssetCompilingManager::Get().FinishAllCompilation();
+		}
+	} while (bCompiledAssets);
 }
 
 
@@ -635,9 +650,6 @@ bool FWorldRenderCapture::CaptureMRSFromPosition(
 	// Cache the view/projection matricies we used to render the scene
 	LastCaptureViewMatrices = NewView->ViewMatrices;
 
-	// wait for render
-	FlushRenderingCommands();
-
 	// read back image
 	if (PostProcessPassPixelData.IsValid())
 	{
