WorldRenderCapture - Modified capture to use default time (0) and set realtime update off on the view family context
* Makes it clear that the capture isn't using a realtime view, and prevent materials relying on Time from giving different results based on the time elapsed since the app start

diff --git a/Engine/Plugins/Runtime/MeshModelingToolset/Source/ModelingComponents/Private/Scene/WorldRenderCapture.cpp b/Engine/Plugins/Runtime/MeshModelingToolset/Source/ModelingComponents/Private/Scene/WorldRenderCapture.cpp
--- a/Engine/Plugins/Runtime/MeshModelingToolset/Source/ModelingComponents/Private/Scene/WorldRenderCapture.cpp
+++ b/Engine/Plugins/Runtime/MeshModelingToolset/Source/ModelingComponents/Private/Scene/WorldRenderCapture.cpp
@@ -578,7 +578,8 @@ bool FWorldRenderCapture::CaptureMRSFromPosition(
 	FSceneViewFamilyContext ViewFamily(FSceneViewFamily::ConstructionValues(
 		RenderTargetResource, Scene, ShowFlags)
 		.SetTime(FGameTime())
-		.SetRealtimeUpdate(false));
+		.SetRealtimeUpdate(false)
+	);
 
 	// unclear whether these show flags are really necessary, since we are using
 	// a custom postprocess pass and reading it's output buffer directly
@@ -721,7 +722,8 @@ bool FWorldRenderCapture::CaptureEmissiveFromPosition(
 
 	FSceneViewFamilyContext ViewFamily(
 		FSceneViewFamily::ConstructionValues(RenderTargetResource, World->Scene, ShowFlags)
-		.SetTime(FGameTime::GetTimeSinceAppStart())
+		.SetTime(FGameTime())
+		.SetRealtimeUpdate(false)
 	);
 
 	// To enable visualization mode
@@ -769,7 +771,7 @@ bool FWorldRenderCapture::CaptureEmissiveFromPosition(
 	ViewFamily.EngineShowFlags.SetCapsuleShadows(false);
 	ViewFamily.EngineShowFlags.SetContactShadows(false);
 
-	FCanvas Canvas(RenderTargetResource, NULL, FGameTime::GetTimeSinceAppStart(), World->Scene->GetFeatureLevel());
+	FCanvas Canvas(RenderTargetResource, nullptr, FGameTime(), NewView->GetFeatureLevel());
 	Canvas.Clear(FLinearColor::Transparent);
 
 	UMaterialInterface* MaterialInterface = GetBufferVisualizationData().GetMaterial(NewView->CurrentBufferVisualizationMode);
@@ -847,6 +849,7 @@ bool FWorldRenderCapture::CaptureDeviceDepthFromPosition(
 			RenderTargetResource,
 			World->Scene,
 			ShowFlags)
+		.SetTime(FGameTime())
 		.SetRealtimeUpdate(false)
 		.SetResolveScene(false)
 	);
@@ -889,7 +892,7 @@ bool FWorldRenderCapture::CaptureDeviceDepthFromPosition(
 		Extension->SetupView(ViewFamily, *NewView);
 	}
 
-	FCanvas Canvas(RenderTargetResource, NULL, FGameTime::GetTimeSinceAppStart(), World->Scene->GetFeatureLevel());
+	FCanvas Canvas(RenderTargetResource, nullptr, FGameTime(), NewView->GetFeatureLevel());
 	Canvas.Clear(FLinearColor::Transparent);
 
 	// Unlike other capture types, we don't need to cache any shaders before we capture the device depth render
@@ -990,7 +993,8 @@ namespace Internal
 
 		FSceneViewFamilyContext ViewFamily(
 			FSceneViewFamily::ConstructionValues(RenderTargetResource, Scene, FEngineShowFlags(ESFIM_Game))
-			.SetTime(FGameTime::GetTimeSinceAppStart())
+			.SetTime(FGameTime())
+			.SetRealtimeUpdate(false)
 		);
 
 		// To enable visualization mode
@@ -1019,7 +1023,7 @@ namespace Internal
 		ViewFamily.SetScreenPercentageInterface(new FLegacyScreenPercentageDriver(ViewFamily, 1.0f));
 
 		// should we cache the FCanvas?
-		FCanvas Canvas(RenderTargetResource, NULL, FGameTime::GetTimeSinceAppStart(), Scene->GetFeatureLevel());
+		FCanvas Canvas(RenderTargetResource, nullptr, FGameTime(), NewView->GetFeatureLevel());
 		Canvas.Clear(FLinearColor::Transparent);
 
 		UMaterialInterface* MaterialInterface = GetBufferVisualizationData().GetMaterial(NewView->CurrentBufferVisualizationMode);
