WorldRenderCapture - Renamed CVarModelingWorldRenderCaptureVTWarmupFrames to CVarModelingWorldRenderCaptureWarmupFrames
* Warmup is not specific to VT, but also of use for Nanite

diff --git a/Engine/Plugins/Runtime/MeshModelingToolset/Source/ModelingComponents/Private/Scene/WorldRenderCapture.cpp b/Engine/Plugins/Runtime/MeshModelingToolset/Source/ModelingComponents/Private/Scene/WorldRenderCapture.cpp
--- a/Engine/Plugins/Runtime/MeshModelingToolset/Source/ModelingComponents/Private/Scene/WorldRenderCapture.cpp
+++ b/Engine/Plugins/Runtime/MeshModelingToolset/Source/ModelingComponents/Private/Scene/WorldRenderCapture.cpp
@@ -25,10 +25,10 @@
 
 using namespace UE::Geometry;
 
-static TAutoConsoleVariable<int32> CVarModelingWorldRenderCaptureVTWarmupFrames(
+static TAutoConsoleVariable<int32> CVarModelingWorldRenderCaptureWarmupFrames(
 	TEXT("modeling.WorldRenderCapture.VTWarmupFrames"),
 	5,
-	TEXT("Number of frames to render before each capture in order to warmup the VT."));
+	TEXT("Number of frames to render before each capture in order to warmup the renderer."));
 
 FRenderCaptureTypeFlags FRenderCaptureTypeFlags::All(bool bCombinedMRS)
 {
@@ -386,8 +386,7 @@ namespace Internal
 
 /**
  * Render the scene to the provided canvas. Will potentially perform the render multiple times, depending on the value of
- * the CVarModelingWorldRenderCaptureVTWarmupFrames CVar. This is needed to ensure the VT is primed properly before capturing
- * the scene.
+ * the warmup frames CVar. This is needed to ensure the VT/Nanite are primed properly before capturing the scene.
  */
 static void PerformSceneRender(FCanvas* Canvas, FSceneViewFamily* ViewFamily)
 {
@@ -395,7 +394,7 @@ static void PerformSceneRender(FCanvas* Canvas, FSceneViewFamily* ViewFamily)
 
 	do
 	{
-		int32 NumRender = 1 + CVarModelingWorldRenderCaptureVTWarmupFrames.GetValueOnGameThread();
+		int32 NumRender = 1 + CVarModelingWorldRenderCaptureWarmupFrames.GetValueOnGameThread();
 		for (int32 i = 0; i < NumRender; i++)
 		{
 			GetRendererModule().BeginRenderingViewFamily(Canvas, ViewFamily);
