Make the insert edge loop tool preview edges & preview blocking-edge highlights be not depth tested, so it is easier to see where to edges go / where they fail to go

diff --git a/Engine/Plugins/Runtime/MeshModelingToolset/Source/MeshModelingTools/Private/ToolActivities/PolyEditInsertEdgeLoopActivity.cpp b/Engine/Plugins/Runtime/MeshModelingToolset/Source/MeshModelingTools/Private/ToolActivities/PolyEditInsertEdgeLoopActivity.cpp
--- a/Engine/Plugins/Runtime/MeshModelingToolset/Source/MeshModelingTools/Private/ToolActivities/PolyEditInsertEdgeLoopActivity.cpp
+++ b/Engine/Plugins/Runtime/MeshModelingToolset/Source/MeshModelingTools/Private/ToolActivities/PolyEditInsertEdgeLoopActivity.cpp
@@ -91,11 +91,13 @@ void UPolyEditInsertEdgeLoopActivity::Setup(UInteractiveTool* ParentToolIn)
 	// Draws the new group edges that are added
 	PreviewEdgeRenderer.LineColor = FLinearColor::Green;
 	PreviewEdgeRenderer.LineThickness = 2.0;
+	PreviewEdgeRenderer.bDepthTested = false;
 
 	// Highlights non-quad groups that stop the loop;
 	ProblemTopologyRenderer.LineColor = FLinearColor::Red;
 	ProblemTopologyRenderer.LineThickness = 3.0;
 	ProblemTopologyRenderer.DepthBias = 1.0;
+	ProblemTopologyRenderer.bDepthTested = false;
 
 	TopologySelectorSettings.bEnableEdgeHits = true;
 	TopologySelectorSettings.bEnableFaceHits = false;
