Remove CacheShadersForMaterial call in FWorldRenderCapture

diff --git a/Engine/Plugins/Runtime/MeshModelingToolset/Source/ModelingComponents/Private/Scene/WorldRenderCapture.cpp b/Engine/Plugins/Runtime/MeshModelingToolset/Source/ModelingComponents/Private/Scene/WorldRenderCapture.cpp
--- a/Engine/Plugins/Runtime/MeshModelingToolset/Source/ModelingComponents/Private/Scene/WorldRenderCapture.cpp
+++ b/Engine/Plugins/Runtime/MeshModelingToolset/Source/ModelingComponents/Private/Scene/WorldRenderCapture.cpp
@@ -483,37 +483,6 @@ static bool ReadPixelDataToImage(TUniquePtr<FImagePixelData>& PixelData, FImageA
 	}
 }
 
-// This function is used to cache the material shaders used for world render capture (mostly buffer visualization
-// materials, but there are also other ones). If we don't do this, and if the shaders also haven't been otherwise
-// compiled (e.g., by switching to the needed buffer visualization mode in the viewport), then the render capture will
-// use a fallback material the first time around and give incorrect results (which caused #jira UE-146097). Note that
-// although buffer visualization materials are initialized during engine/editor startup, we do not cache needed shaders
-// at that point since doing so only compiles them for the default feature level, GMaxRHIFeatureLevel. The user may
-// change the preview feature level after editor startup, or may have done so in a previous session (the setting is
-// serialized/reused in subsequent sessions via Engine/Saved/Config/WindowsEditor/EditorPerProjectUserSettings.ini),
-// so in order to make the render capture code work in that case we do this caching right here. The CacheShaders
-// function compiles all feature levels returned by GetFeatureLevelsToCompileForAllMaterials (note that
-// SetFeatureLevelToCompile is called when the feature level is switched).
-void CacheShadersForMaterial(UMaterialInterface* MaterialInterface)
-{
-	if (ensure(MaterialInterface))
-	{
-		UMaterial* Material = MaterialInterface->GetMaterial();
-		if (ensure(Material))
-		{
-			// Mark the material as a special engine material so that FMaterial::IsRequiredComplete returns true and we
-			// compile the material, we restore the previous value to minimize surprise for other code
-			bool OldState = Material->bUsedAsSpecialEngineMaterial;
-			Material->bUsedAsSpecialEngineMaterial = true;
-			Material->CacheShaders(EMaterialShaderPrecompileMode::Synchronous);
-			ensureMsgf(Material->IsComplete(),
-				TEXT("Caching shaders for material %s did not complete"),
-				*MaterialInterface->GetName());
-			Material->bUsedAsSpecialEngineMaterial = OldState;
-		}
-	}
-}
-
 } // end namespace Internal
 } // end namespace UE
 
@@ -659,7 +628,6 @@ bool FWorldRenderCapture::CaptureMRSFromPosition(
 	FCanvas Canvas = FCanvas(RenderTargetResource, nullptr, this->World, ERHIFeatureLevel::SM5, FCanvas::CDM_DeferDrawing, 1.0f);
 	Canvas.Clear(FLinearColor::Transparent);
 
-	UE::Internal::CacheShadersForMaterial(PostProcessMaterial);
 	UE::Internal::PerformSceneRender(&Canvas, &ViewFamily);
 
 	// Cache the view/projection matricies we used to render the scene
@@ -773,8 +741,6 @@ bool FWorldRenderCapture::CaptureEmissiveFromPosition(
 	FCanvas Canvas(RenderTargetResource, nullptr, FGameTime(), NewView->GetFeatureLevel());
 	Canvas.Clear(FLinearColor::Transparent);
 
-	UMaterialInterface* MaterialInterface = GetBufferVisualizationData().GetMaterial(NewView->CurrentBufferVisualizationMode);
-	UE::Internal::CacheShadersForMaterial(MaterialInterface);
 	UE::Internal::PerformSceneRender(&Canvas, &ViewFamily);
 
 	// Cache the view/projection matricies we used to render the scene
@@ -1025,8 +991,6 @@ namespace Internal
 		FCanvas Canvas(RenderTargetResource, nullptr, FGameTime(), NewView->GetFeatureLevel());
 		Canvas.Clear(FLinearColor::Transparent);
 
-		UMaterialInterface* MaterialInterface = GetBufferVisualizationData().GetMaterial(NewView->CurrentBufferVisualizationMode);
-		UE::Internal::CacheShadersForMaterial(MaterialInterface);
 		UE::Internal::PerformSceneRender(&Canvas, &ViewFamily);
 
 		// Cache the view/projection matricies we used to render the scene
