WorldRenderCapture: Refactoring to clarify the code and make it easier to rearrange for performance optimization experiments

#rnx

diff --git a/Engine/Plugins/Runtime/MeshModelingToolset/Source/ModelingComponents/Private/Scene/WorldRenderCapture.cpp b/Engine/Plugins/Runtime/MeshModelingToolset/Source/ModelingComponents/Private/Scene/WorldRenderCapture.cpp
--- a/Engine/Plugins/Runtime/MeshModelingToolset/Source/ModelingComponents/Private/Scene/WorldRenderCapture.cpp
+++ b/Engine/Plugins/Runtime/MeshModelingToolset/Source/ModelingComponents/Private/Scene/WorldRenderCapture.cpp
@@ -216,6 +216,8 @@ void FWorldRenderCapture::SetVisibleComponents(const TArray<UActorComponent*>& C
 
 void FWorldRenderCapture::SetVisibleActorsAndComponents(const TArray<AActor*>& Actors, const TArray<UActorComponent*>& Components)
 {
+	TRACE_CPUPROFILER_EVENT_SCOPE(FWorldRenderCapture::SetVisibleActorsAndComponents);
+
 	VisiblePrimitives.Reset();
 
 	FBoxSphereBounds::Builder BoundsBuilder;
@@ -390,6 +392,8 @@ namespace Internal
  */
 static void PerformSceneRender(FCanvas* Canvas, FSceneViewFamily* ViewFamily)
 {
+	TRACE_CPUPROFILER_EVENT_SCOPE(Internal::PerformSceneRender);
+
 	bool bCompiledAssets = false;
 
 	do
@@ -418,6 +422,8 @@ static void PerformSceneRender(FCanvas* Canvas, FSceneViewFamily* ViewFamily)
  */
 static bool ReadPixelDataToImage(TUniquePtr<FImagePixelData>& PixelData, FImageAdapter& ResultImageOut, bool bLinear)
 {
+	TRACE_CPUPROFILER_EVENT_SCOPE(Internal::ReadPixelDataToImage);
+
 	int64 SizeInBytes;
 	const void* OutRawData = nullptr;
 	PixelData->GetRawData(OutRawData, SizeInBytes);
@@ -483,6 +489,80 @@ static bool ReadPixelDataToImage(TUniquePtr<FImagePixelData>& PixelData, FImageA
 	}
 }
 
+// The .ViewFamily member should be set on the return value later
+// VisiblePrimitives: If non-empty, only these primitives are shown
+// HiddenPrimitives: These will be hidden
+FSceneViewInitOptions MakeSceneViewInitOptions(
+	const FFrame3d& Frame,
+	const FImageDimensions& Dimensions,
+	double HorzFOVDegrees,
+	double NearPlaneDist,
+	const TSet<FPrimitiveComponentId>& VisiblePrimitives,
+	const TSet<FPrimitiveComponentId>& HiddenPrimitives = {}
+	)
+{
+	int32 Width = Dimensions.GetWidth();
+	int32 Height = Dimensions.GetHeight();
+
+	// TODO The input frame should have been the render frame the whole time
+	FFrame3d RenderFrame(Frame.Origin, Frame.Y(), Frame.Z(), Frame.X());
+
+	FVector ViewOrigin = RenderFrame.Origin;
+	FMatrix ViewRotationMatrix = FRotationMatrix::Make((FQuat)RenderFrame.Rotation.Inverse());
+
+	const float HalfFOVRadians = FMath::DegreesToRadians<float>(HorzFOVDegrees) * 0.5f;
+	static_assert((int32)ERHIZBuffer::IsInverted != 0, "Check NearPlane and Projection Matrix");
+	FMatrix ProjectionMatrix = FReversedZPerspectiveMatrix(HalfFOVRadians, 1.0f, 1.0f, (float)NearPlaneDist);
+
+	FSceneViewInitOptions ViewInitOptions;
+	ViewInitOptions.SetViewRectangle(FIntRect(0, 0, Width, Height));
+	if (VisiblePrimitives.Num() > 0)
+	{
+		ViewInitOptions.ShowOnlyPrimitives = VisiblePrimitives;
+	}
+	ViewInitOptions.HiddenPrimitives = HiddenPrimitives;
+	ViewInitOptions.ViewOrigin = ViewOrigin;
+	ViewInitOptions.ViewRotationMatrix = ViewRotationMatrix;
+	ViewInitOptions.ProjectionMatrix = ProjectionMatrix;
+
+	return ViewInitOptions;
+}
+
+void SetCommonShowFlags(FEngineShowFlags& ShowFlags, bool bAntiAliasing)
+{
+	ShowFlags.SetAntiAliasing(bAntiAliasing);
+
+	// Disable advanced features
+	ShowFlags.DisableAdvancedFeatures();
+
+	// Disable some more features. These are sorted alphabetically.
+	ShowFlags.LOD = 0;
+	ShowFlags.MotionBlur = 0;
+	ShowFlags.SetAtmosphere(false);
+	ShowFlags.SetBloom(false);
+	ShowFlags.SetCapsuleShadows(false);
+	ShowFlags.SetContactShadows(false);
+	ShowFlags.SetDiffuse(false);
+	ShowFlags.SetDirectionalLights(false);
+	ShowFlags.SetDynamicShadows(false);
+	ShowFlags.SetFog(false);
+	ShowFlags.SetGlobalIllumination(false);
+	ShowFlags.SetLighting(false);
+	ShowFlags.SetMotionBlur(false);
+	ShowFlags.SetPointLights(false);
+	ShowFlags.SetPostProcessing(false);
+	ShowFlags.SetRectLights(false);
+	ShowFlags.SetReflectionEnvironment(false);
+	ShowFlags.SetRefraction(false);
+	ShowFlags.SetSceneColorFringe(false);
+	ShowFlags.SetSkyLighting(false);
+	ShowFlags.SetSpecular(false);
+	ShowFlags.SetSpotLights(false);
+	ShowFlags.SetTonemapper(false);
+	ShowFlags.SetToneCurve(false);
+	ShowFlags.SetTranslucency(false);
+}
+
 } // end namespace Internal
 } // end namespace UE
 
@@ -494,6 +574,8 @@ bool FWorldRenderCapture::CaptureMRSFromPosition(
 	FImageAdapter& ResultImageOut,
 	const FRenderCaptureConfig& Config)
 {
+	TRACE_CPUPROFILER_EVENT_SCOPE(FWorldRenderCapture::CaptureMRSFromPosition);
+
 	// this post-process material renders an image with R=Metallic, G=Roughness, B=Specular, A=AmbientOcclusion
 	FString MRSPostProcessMaterialAssetPath = TEXT("/MeshModelingToolsetExp/Materials/PostProcess_PackedMRSA.PostProcess_PackedMRSA");
 	TSoftObjectPtr<UMaterialInterface> PostProcessMaterialPtr = TSoftObjectPtr<UMaterialInterface>(FSoftObjectPath(MRSPostProcessMaterialAssetPath));
@@ -513,35 +595,9 @@ bool FWorldRenderCapture::CaptureMRSFromPosition(
 	}
 	FRenderTarget* RenderTargetResource = RenderTargetTexture->GameThread_GetRenderTargetResource();
 
-	int32 Width = Dimensions.GetWidth();
-	int32 Height = Dimensions.GetHeight();
-
-	FQuat ViewOrientation = (FQuat)ViewFrame.Rotation;
-	FMatrix ViewRotationMatrix = FInverseRotationMatrix(ViewOrientation.Rotator());
-	FVector ViewOrigin = (FVector)ViewFrame.Origin;
-
-	// convert to rendering coordinate system
-	ViewRotationMatrix = ViewRotationMatrix * FMatrix(
-		FPlane(0, 0, 1, 0),
-		FPlane(1, 0, 0, 0),
-		FPlane(0, 1, 0, 0),
-		FPlane(0, 0, 0, 1));
-
-	const float HalfFOVRadians = FMath::DegreesToRadians<float>(HorzFOVDegrees) * 0.5f;
-	static_assert((int32)ERHIZBuffer::IsInverted != 0, "Check NearPlane and Projection Matrix");
-	FMatrix ProjectionMatrix = FReversedZPerspectiveMatrix(HalfFOVRadians, 1.0f, 1.0f, (float)NearPlaneDist);
-
-	EViewModeIndex ViewModeIndex = EViewModeIndex::VMI_Unlit;		// VMI_Lit, VMI_LightingOnly, VMI_VisualizeBuffer
-
-	FEngineShowFlags ShowFlags = FEngineShowFlags(EShowFlagInitMode::ESFIM_Game);
-	ApplyViewMode(ViewModeIndex, true, ShowFlags);
-
-	// unclear if these flags need to be set before creating ViewFamily
-	ShowFlags.SetAntiAliasing(Config.bAntiAliasing);
-	ShowFlags.SetDepthOfField(false);
-	ShowFlags.SetMotionBlur(false);
-	ShowFlags.SetBloom(false);
-	ShowFlags.SetSceneColorFringe(false);
+	FEngineShowFlags ShowFlags(ESFIM_Game);
+	ApplyViewMode(VMI_Unlit, true, ShowFlags);
+	Internal::SetCommonShowFlags(ShowFlags, Config.bAntiAliasing);
 
 	FSceneViewFamilyContext ViewFamily(FSceneViewFamily::ConstructionValues(
 		RenderTargetResource, Scene, ShowFlags)
@@ -549,50 +605,20 @@ bool FWorldRenderCapture::CaptureMRSFromPosition(
 		.SetRealtimeUpdate(false)
 	);
 
-	// unclear whether these show flags are really necessary, since we are using
-	// a custom postprocess pass and reading it's output buffer directly
-
-	ViewFamily.EngineShowFlags.DisableAdvancedFeatures();
-	ViewFamily.EngineShowFlags.MotionBlur = 0;
-	ViewFamily.EngineShowFlags.LOD = 0;
-
-	ViewFamily.EngineShowFlags.SetTonemapper(false);
-	ViewFamily.EngineShowFlags.SetColorGrading(false);
-	ViewFamily.EngineShowFlags.SetToneCurve(false);
-
-	ViewFamily.EngineShowFlags.SetPostProcessing(false);
-	ViewFamily.EngineShowFlags.SetFog(false);
-	ViewFamily.EngineShowFlags.SetGlobalIllumination(false);
-	ViewFamily.EngineShowFlags.SetEyeAdaptation(false);
-	ViewFamily.EngineShowFlags.SetDirectionalLights(false);
-	ViewFamily.EngineShowFlags.SetPointLights(false);
-	ViewFamily.EngineShowFlags.SetSpotLights(false);
-	ViewFamily.EngineShowFlags.SetRectLights(false);
-
-	ViewFamily.EngineShowFlags.SetDiffuse(false);
-	ViewFamily.EngineShowFlags.SetSpecular(false);
-
-	ViewFamily.EngineShowFlags.SetDynamicShadows(false);
-	ViewFamily.EngineShowFlags.SetCapsuleShadows(false);
-	ViewFamily.EngineShowFlags.SetContactShadows(false);
-
 	//ViewFamily.EngineShowFlags.SetScreenPercentage(false);
 	ViewFamily.SetScreenPercentageInterface(new FLegacyScreenPercentageDriver(ViewFamily, 1.f));
 
 	// This is called in various other places, unclear if we should be doing this too
-	//EngineShowFlagOverride(EShowFlagInitMode::ESFIM_Game, ViewModeIndex, ViewFamily.EngineShowFlags, true);
-
-	FSceneViewInitOptions ViewInitOptions;
-	ViewInitOptions.SetViewRectangle(FIntRect(0, 0, Width, Height));
+	//EngineShowFlagOverride(ESFIM_Game, VMI_Unlit, ViewFamily.EngineShowFlags, true);
+
+	FSceneViewInitOptions ViewInitOptions = Internal::MakeSceneViewInitOptions(
+		ViewFrame,
+		Dimensions,
+		HorzFOVDegrees,
+		NearPlaneDist,
+		VisiblePrimitives);
 	ViewInitOptions.ViewFamily = &ViewFamily;
-	if (VisiblePrimitives.Num() > 0)
-	{
-		ViewInitOptions.ShowOnlyPrimitives = VisiblePrimitives;
-	}
-	ViewInitOptions.ViewOrigin = ViewOrigin;
-	ViewInitOptions.ViewRotationMatrix = ViewRotationMatrix;
-	ViewInitOptions.ProjectionMatrix = ProjectionMatrix;
-	ViewInitOptions.FOV = HorzFOVDegrees;
+	ViewInitOptions.FOV = HorzFOVDegrees; // Probably unecessary (FOV is already accounted for in projection matrix)
 
 	FSceneView* NewView = new FSceneView(ViewInitOptions);
 	ViewFamily.Views.Add(NewView);
@@ -633,7 +659,6 @@ bool FWorldRenderCapture::CaptureMRSFromPosition(
 	// Cache the view/projection matricies we used to render the scene
 	LastCaptureViewMatrices = NewView->ViewMatrices;
 
-	// read back image
 	if (PostProcessPassPixelData.IsValid())
 	{
 		UE::Internal::ReadPixelDataToImage(PostProcessPassPixelData, ResultImageOut, true );
@@ -654,6 +679,8 @@ bool FWorldRenderCapture::CaptureEmissiveFromPosition(
 	FImageAdapter& ResultImageOut,
 	const FRenderCaptureConfig& Config)
 {
+	TRACE_CPUPROFILER_EVENT_SCOPE(FWorldRenderCapture::CaptureEmissiveFromPosition);
+
 	UTextureRenderTarget2D* RenderTargetTexture = GetRenderTexture(true);
 	if (ensure(RenderTargetTexture) == false)
 	{
@@ -661,31 +688,11 @@ bool FWorldRenderCapture::CaptureEmissiveFromPosition(
 	}
 	FTextureRenderTargetResource* RenderTargetResource = RenderTargetTexture->GameThread_GetRenderTargetResource();
 
-	int32 Width = Dimensions.GetWidth();
-	int32 Height = Dimensions.GetHeight();
-
-	FQuat ViewOrientation = (FQuat)Frame.Rotation;
-	FMatrix ViewRotationMatrix = FInverseRotationMatrix(ViewOrientation.Rotator());
-	FVector ViewOrigin = (FVector)Frame.Origin;
-
-	// convert to rendering coordinate system
-	ViewRotationMatrix = ViewRotationMatrix * FMatrix(
-		FPlane(0, 0, 1, 0),
-		FPlane(1, 0, 0, 0),
-		FPlane(0, 1, 0, 0),
-		FPlane(0, 0, 0, 1));
-
-	const float HalfFOVRadians = FMath::DegreesToRadians<float>(HorzFOVDegrees) * 0.5f;
-	static_assert((int32)ERHIZBuffer::IsInverted != 0, "Check NearPlane and Projection Matrix");
-	FMatrix ProjectionMatrix = FReversedZPerspectiveMatrix(HalfFOVRadians, 1.0f, 1.0f, (float)NearPlaneDist);
-
-	// unclear if these flags need to be set before creating ViewFamily
 	FEngineShowFlags ShowFlags(ESFIM_Game);
-	ShowFlags.SetAntiAliasing(Config.bAntiAliasing);
-	ShowFlags.SetDepthOfField(false);
-	ShowFlags.SetMotionBlur(false);
-	ShowFlags.SetBloom(false);
-	ShowFlags.SetSceneColorFringe(false);
+	Internal::SetCommonShowFlags(ShowFlags, Config.bAntiAliasing);
+	ShowFlags.SetPostProcessing(true);
+	ShowFlags.SetVisualizeBuffer(true);
+	ShowFlags.SetLighting(true); // Need this because Emissive is lighting I guess
 
 	FSceneViewFamilyContext ViewFamily(
 		FSceneViewFamily::ConstructionValues(RenderTargetResource, World->Scene, ShowFlags)
@@ -693,22 +700,13 @@ bool FWorldRenderCapture::CaptureEmissiveFromPosition(
 		.SetRealtimeUpdate(false)
 	);
 
-	// To enable visualization mode
-	ViewFamily.EngineShowFlags.SetPostProcessing(true);
-	ViewFamily.EngineShowFlags.SetVisualizeBuffer(true);
-	ViewFamily.EngineShowFlags.SetTonemapper(false);
-	ViewFamily.EngineShowFlags.SetScreenPercentage(false);
-
-	FSceneViewInitOptions ViewInitOptions;
-	ViewInitOptions.SetViewRectangle(FIntRect(0, 0, Width, Height));
+	FSceneViewInitOptions ViewInitOptions = Internal::MakeSceneViewInitOptions(
+		Frame,
+		Dimensions,
+		HorzFOVDegrees,
+		NearPlaneDist,
+		VisiblePrimitives);
 	ViewInitOptions.ViewFamily = &ViewFamily;
-	if (VisiblePrimitives.Num() > 0)
-	{
-		ViewInitOptions.ShowOnlyPrimitives = VisiblePrimitives;
-	}
-	ViewInitOptions.ViewOrigin = ViewOrigin;
-	ViewInitOptions.ViewRotationMatrix = ViewRotationMatrix;
-	ViewInitOptions.ProjectionMatrix = ProjectionMatrix;
 
 	FSceneView* NewView = new FSceneView(ViewInitOptions);
 	NewView->CurrentBufferVisualizationMode = FName("PreTonemapHDRColor");
@@ -716,28 +714,6 @@ bool FWorldRenderCapture::CaptureEmissiveFromPosition(
 
 	ViewFamily.SetScreenPercentageInterface(new FLegacyScreenPercentageDriver(ViewFamily, 1.0f));
 
-	ViewFamily.EngineShowFlags.DisableAdvancedFeatures();
-	ViewFamily.EngineShowFlags.LOD = 0;
-
-	ViewFamily.EngineShowFlags.SetFog(false);
-	ViewFamily.EngineShowFlags.SetGlobalIllumination(false);
-	ViewFamily.EngineShowFlags.SetDirectionalLights(false);
-	ViewFamily.EngineShowFlags.SetPointLights(false);
-	ViewFamily.EngineShowFlags.SetSpotLights(false);
-	ViewFamily.EngineShowFlags.SetRectLights(false);
-	ViewFamily.EngineShowFlags.SetSkyLighting(false);
-	//ViewFamily.EngineShowFlags.SetLighting(false);		// this breaks things
-
-	ViewFamily.EngineShowFlags.SetDiffuse(false);
-	ViewFamily.EngineShowFlags.SetSpecular(false);
-
-	ViewFamily.EngineShowFlags.SetRefraction(false);
-	ViewFamily.EngineShowFlags.SetReflectionEnvironment(false);
-
-	ViewFamily.EngineShowFlags.SetDynamicShadows(false);
-	ViewFamily.EngineShowFlags.SetCapsuleShadows(false);
-	ViewFamily.EngineShowFlags.SetContactShadows(false);
-
 	FCanvas Canvas(RenderTargetResource, nullptr, FGameTime(), NewView->GetFeatureLevel());
 	Canvas.Clear(FLinearColor::Transparent);
 
@@ -746,6 +722,9 @@ bool FWorldRenderCapture::CaptureEmissiveFromPosition(
 	// Cache the view/projection matricies we used to render the scene
 	LastCaptureViewMatrices = NewView->ViewMatrices;
 
+	int32 Width = Dimensions.GetWidth();
+	int32 Height = Dimensions.GetHeight();
+
 	// Copy the contents of the remote texture to system memory
 	ReadImageBuffer.Reset();
 	ReadImageBuffer.SetNumUninitialized(Width * Height);
@@ -753,13 +732,17 @@ bool FWorldRenderCapture::CaptureEmissiveFromPosition(
 	ReadSurfaceDataFlags.SetLinearToGamma(false);
 	RenderTargetResource->ReadLinearColorPixels(ReadImageBuffer, ReadSurfaceDataFlags, FIntRect(0, 0, Width, Height));
 
-	ResultImageOut.SetDimensions(Dimensions);
-	for (int32 yi = 0; yi < Height; ++yi)
 	{
-		for (int32 xi = 0; xi < Width; ++xi)
+		TRACE_CPUPROFILER_EVENT_SCOPE(FWorldRenderCapture::CaptureDeviceDepthFromPosition_CopyBufferToImage);
+
+		ResultImageOut.SetDimensions(Dimensions);
+		for (int32 yi = 0; yi < Height; ++yi)
 		{
-			FLinearColor PixelColorf = ReadImageBuffer[yi * Width + xi];
-			ResultImageOut.SetPixel(FVector2i(xi, yi), PixelColorf);
+			for (int32 xi = 0; xi < Width; ++xi)
+			{
+				FLinearColor PixelColorf = ReadImageBuffer[yi * Width + xi];
+				ResultImageOut.SetPixel(FVector2i(xi, yi), PixelColorf);
+			}
 		}
 	}
 
@@ -774,6 +757,8 @@ bool FWorldRenderCapture::CaptureDeviceDepthFromPosition(
 	FImageAdapter& ResultImageOut,
 	const FRenderCaptureConfig& Config)
 {
+	TRACE_CPUPROFILER_EVENT_SCOPE(FWorldRenderCapture::CaptureDeviceDepthFromPosition);
+
 	DepthRenderTexture = GetDepthRenderTexture();
 	if (ensure(DepthRenderTexture) == false)
 	{
@@ -784,30 +769,8 @@ bool FWorldRenderCapture::CaptureDeviceDepthFromPosition(
 	int32 Width = Dimensions.GetWidth();
 	int32 Height = Dimensions.GetHeight();
 
-	// TODO The input frame should have been the render frame the whole time
-	FFrame3d RenderFrame(Frame.Origin, Frame.Y(), Frame.Z(), Frame.X());
-	FVector ViewOrigin = (FVector)RenderFrame.Origin;
-	FMatrix ViewRotationMatrix = FRotationMatrix::Make((FQuat)RenderFrame.Rotation.Inverse());
-
-	const float HalfFOVRadians = FMath::DegreesToRadians<float>(HorzFOVDegrees) * 0.5f;
-	static_assert((int32)ERHIZBuffer::IsInverted != 0, "Check NearPlane and Projection Matrix");
-	FMatrix ProjectionMatrix = FReversedZPerspectiveMatrix(HalfFOVRadians, 1.0f, 1.0f, (float)NearPlaneDist);
-
-	// unclear if these flags need to be set before creating ViewFamily
 	FEngineShowFlags ShowFlags(ESFIM_Game);
-	ShowFlags.SetAntiAliasing(false); // Intentionally does NOT use Config.bAntiAliasing
-	ShowFlags.SetDepthOfField(false);
-	ShowFlags.SetMotionBlur(false);
-	ShowFlags.SetBloom(false);
-	ShowFlags.SetSceneColorFringe(false);
-	ShowFlags.SetAtmosphere(false);
-	ShowFlags.SetLighting(false);
-	ShowFlags.SetScreenPercentage(false);
-	ShowFlags.SetTranslucency(false);
-	ShowFlags.SetSeparateTranslucency(false);
-	ShowFlags.SetFog(false);
-	ShowFlags.SetVolumetricFog(false);
-	ShowFlags.SetDynamicShadows(false);
+	Internal::SetCommonShowFlags(ShowFlags, false); // Never use AntiAliasing so we dont blend pixels at different depths
 
 	FSceneViewFamilyContext ViewFamily(
 		FSceneViewFamily::ConstructionValues(
@@ -822,20 +785,14 @@ bool FWorldRenderCapture::CaptureDeviceDepthFromPosition(
 	// Request a scene depth render
 	ViewFamily.SceneCaptureSource = ESceneCaptureSource::SCS_DeviceDepth;
 
-	FSceneViewInitOptions ViewInitOptions;
-	
-	ViewInitOptions.SetViewRectangle(FIntRect(0, 0, Width, Height));
+	FSceneViewInitOptions ViewInitOptions = Internal::MakeSceneViewInitOptions(
+		Frame,
+		Dimensions,
+		HorzFOVDegrees,
+		NearPlaneDist,
+		VisiblePrimitives);
 	ViewInitOptions.ViewFamily = &ViewFamily;
-	if (VisiblePrimitives.Num() > 0)
-	{
-		ViewInitOptions.ShowOnlyPrimitives = VisiblePrimitives;
-	}
-	ViewInitOptions.ViewOrigin = ViewOrigin;
-	ViewInitOptions.ViewRotationMatrix = ViewRotationMatrix;
-	ViewInitOptions.ProjectionMatrix = ProjectionMatrix;
-
 	ViewInitOptions.OverrideFarClippingPlaneDistance = -1.f;
-
 	if (ViewFamily.Scene->GetWorld() != nullptr && ViewFamily.Scene->GetWorld()->GetWorldSettings() != nullptr)
 	{
 		ViewInitOptions.WorldToMetersScale = ViewFamily.Scene->GetWorld()->GetWorldSettings()->WorldToMeters;
@@ -846,7 +803,7 @@ bool FWorldRenderCapture::CaptureDeviceDepthFromPosition(
 	NewView->SetupAntiAliasingMethod();
 	ViewFamily.Views.Add(NewView);
 
-	NewView->StartFinalPostprocessSettings(ViewOrigin);
+	NewView->StartFinalPostprocessSettings(ViewInitOptions.ViewOrigin);
 	NewView->EndFinalPostprocessSettings(ViewInitOptions);
 	ViewFamily.SetScreenPercentageInterface(new FLegacyScreenPercentageDriver(ViewFamily, 1.f));
 	
@@ -873,24 +830,29 @@ bool FWorldRenderCapture::CaptureDeviceDepthFromPosition(
 	ReadSurfaceDataFlags.SetLinearToGamma(false);
 
 	RenderTargetResource->ReadLinearColorPixels(ReadImageColorBuffer, ReadSurfaceDataFlags, FIntRect(0, 0, Width, Height));
-	ResultImageOut.SetDimensions(Dimensions);
-	for (int32 yi = 0; yi < Height; ++yi)
+
 	{
-		for (int32 xi = 0; xi < Width; ++xi)
+		TRACE_CPUPROFILER_EVENT_SCOPE(FWorldRenderCapture::CaptureDeviceDepthFromPosition_CopyBufferToImage);
+
+		ResultImageOut.SetDimensions(Dimensions);
+		for (int32 yi = 0; yi < Height; ++yi)
 		{
-			FLinearColor Color = ReadImageColorBuffer[yi * Width + xi];
+			for (int32 xi = 0; xi < Width; ++xi)
+			{
+				FLinearColor Color = ReadImageColorBuffer[yi * Width + xi];
 
-			// Reverse the float encoding used for the Depth value in SceneCapturePixelShader.usf (one minus the device
-			// z is encoded in the RGB components). See DecodeFloatRGB in WaterInfoMerge.usf and, for a reference on the
-			// method, see https://aras-p.info/blog/2009/07/30/encoding-floats-to-rgba-the-final/
-			FVector3f EncodedDepth(Color.R, Color.G, Color.B);
-			float Depth = FVector3f::DotProduct(EncodedDepth, FVector3f(1.0, 1/255.0, 1/65025.0));
+				// Reverse the float encoding used for the Depth value in SceneCapturePixelShader.usf (one minus the device
+				// z is encoded in the RGB components). See DecodeFloatRGB in WaterInfoMerge.usf and, for a reference on the
+				// method, see https://aras-p.info/blog/2009/07/30/encoding-floats-to-rgba-the-final/
+				FVector3f EncodedDepth(Color.R, Color.G, Color.B);
+				float Depth = FVector3f::DotProduct(EncodedDepth, FVector3f(1.0, 1 / 255.0, 1 / 65025.0));
 
-			// Reverse the expression used to compute the Depth to recover the DeviceZ aka normalized device coordinate z
-			float DeviceZ = -(Depth - 1);
-			ensure(DeviceZ >= 0. && DeviceZ <= 1.); // Points on the near plane have Z=1, points on the far plane have Z=0
+				// Reverse the expression used to compute the Depth to recover the DeviceZ aka normalized device coordinate z
+				float DeviceZ = -(Depth - 1);
+				ensure(DeviceZ >= 0. && DeviceZ <= 1.); // Points on the near plane have Z=1, points on the far plane have Z=0
 
-			ResultImageOut.SetPixel(FVector2i(xi, yi), FLinearColor(DeviceZ, 0., 0., 0.));
+				ResultImageOut.SetPixel(FVector2i(xi, yi), FLinearColor(DeviceZ, 0., 0., 0.));
+			}
 		}
 	}
 
@@ -928,58 +890,50 @@ bool FWorldRenderCapture::CaptureDeviceDepthFromPosition(
 
 
 
-/**
- * internal Utility function to render the given Scene to a render target and capture
- * one of the render buffers, defined by VisualizationMode. Not clear where
- * the valid VisualizationMode FNames are defined, possibly this list: "BaseColor,Specular,SubsurfaceColor,WorldNormal,SeparateTranslucencyRGB,,,WorldTangent,SeparateTranslucencyA,,,Opacity,SceneDepth,Roughness,Metallic,ShadingModel,,SceneDepthWorldUnits,SceneColor,PreTonemapHDRColor,PostTonemapHDRColor"
- */
 namespace UE
 {
 namespace Internal
-{ 
+{
+
+	/**
+	 * internal Utility function to render the given Scene to a render target and capture
+	 * one of the render buffers, defined by VisualizationMode.
+	 *
+	 * Not clear where the valid VisualizationMode FNames are defined, possibly this list:
+	 * "BaseColor,Specular,SubsurfaceColor,WorldNormal,SeparateTranslucencyRGB,,,WorldTangent,SeparateTranslucencyA,,,Opacity,SceneDepth,Roughness,Metallic,ShadingModel,,SceneDepthWorldUnits,SceneColor,PreTonemapHDRColor,PostTonemapHDRColor"
+	 * or possibly in the [Engine.BufferVisualizationMaterials] section of Config/BaseEngine.ini, which is read by FBufferVisualizationData::Initialize() and called when the engine initializes in UEngine::Init()
+	 */
 	static void RenderSceneVisualizationToTexture(
 		UTextureRenderTarget2D* RenderTargetTexture,
 		FImageDimensions Dimensions,
 		FSceneInterface* Scene,
 		const FName& VisualizationMode,
-		const FVector& ViewOrigin,
-		const FMatrix& ViewRotationMatrix,
-		const FMatrix& ProjectionMatrix,
+		FSceneViewInitOptions ViewInitOptions,
 		const FRenderCaptureConfig& Config,
-		const TSet<FPrimitiveComponentId>& HiddenPrimitives,		// these primitives will be hidden
-		const TSet<FPrimitiveComponentId>& VisiblePrimitives,		// if non-empty, only these primitives are shown
 		TArray<FLinearColor>& OutSamples,
 		FViewMatrices& LastCaptureViewMatrices
 	)
 	{
+		TRACE_CPUPROFILER_EVENT_SCOPE(Internal::RenderSceneVisualizationToTexture);
+
 		int32 Width = Dimensions.GetWidth();
 		int32 Height = Dimensions.GetHeight();
 		FTextureRenderTargetResource* RenderTargetResource = RenderTargetTexture->GameThread_GetRenderTargetResource();
 
+		FEngineShowFlags ShowFlags(ESFIM_Game);
+		ApplyViewMode(VMI_VisualizeBuffer, true, ShowFlags);
+		Internal::SetCommonShowFlags(ShowFlags, Config.bAntiAliasing);
+		ShowFlags.SetPostProcessMaterial(true);
+		ShowFlags.SetPostProcessing(true);
+		ShowFlags.SetVisualizeBuffer(true);
+
 		FSceneViewFamilyContext ViewFamily(
-			FSceneViewFamily::ConstructionValues(RenderTargetResource, Scene, FEngineShowFlags(ESFIM_Game))
+			FSceneViewFamily::ConstructionValues(RenderTargetResource, Scene, ShowFlags)
 			.SetTime(FGameTime())
 			.SetRealtimeUpdate(false)
 		);
 
-		// To enable visualization mode
-		ViewFamily.EngineShowFlags.SetPostProcessing(true);
-		ViewFamily.EngineShowFlags.SetVisualizeBuffer(true);
-		ViewFamily.EngineShowFlags.SetTonemapper(false);
-		ViewFamily.EngineShowFlags.SetScreenPercentage(false);
-		ViewFamily.EngineShowFlags.SetAntiAliasing(Config.bAntiAliasing);
-
-		FSceneViewInitOptions ViewInitOptions;
-		ViewInitOptions.SetViewRectangle(FIntRect(0, 0, Width, Height));
 		ViewInitOptions.ViewFamily = &ViewFamily;
-		ViewInitOptions.HiddenPrimitives = HiddenPrimitives;
-		if (VisiblePrimitives.Num() > 0)
-		{
-			ViewInitOptions.ShowOnlyPrimitives = VisiblePrimitives;
-		}
-		ViewInitOptions.ViewOrigin = ViewOrigin;
-		ViewInitOptions.ViewRotationMatrix = ViewRotationMatrix;
-		ViewInitOptions.ProjectionMatrix = ProjectionMatrix;
 
 		FSceneView* NewView = new FSceneView(ViewInitOptions);
 		NewView->CurrentBufferVisualizationMode = VisualizationMode;
@@ -1057,6 +1011,8 @@ bool FWorldRenderCapture::CaptureFromPosition(
 		return bOK;
 	}
 
+	TRACE_CPUPROFILER_EVENT_SCOPE(FWorldRenderCapture::CaptureFromPosition_CaptureBufferVisualizationFromPosition);
+
 	// The following handles buffer visualization materials found in /Engine/BufferVisualization/<CaptureTypeName>.
 	// Sometimes these postprocess materials change the raw GBuffer data, in particular:
 	// - The Roughness postprocess material output is:           GBufferRoughness^Gamma,   with Gamma=2.2
@@ -1071,23 +1027,12 @@ bool FWorldRenderCapture::CaptureFromPosition(
 		return false;
 	}
 
-	int32 Width = Dimensions.GetWidth();
-	int32 Height = Dimensions.GetHeight();
-
-	FQuat ViewOrientation = (FQuat)ViewFrame.Rotation;
-	FMatrix ViewRotationMatrix = FInverseRotationMatrix(ViewOrientation.Rotator());
-	FVector ViewOrigin = (FVector)ViewFrame.Origin;
-
-	// convert to rendering coordinate system
-	ViewRotationMatrix = ViewRotationMatrix * FMatrix(
-		FPlane(0, 0, 1, 0),
-		FPlane(1, 0, 0, 0),
-		FPlane(0, 1, 0, 0),
-		FPlane(0, 0, 0, 1));
-
-	const float HalfFOVRadians = FMath::DegreesToRadians<float>(HorzFOVDegrees) * 0.5f;
-	static_assert((int32)ERHIZBuffer::IsInverted != 0, "Check NearPlane and Projection Matrix");
-	FMatrix ProjectionMatrix = FReversedZPerspectiveMatrix( HalfFOVRadians, 1.0f, 1.0f, (float)NearPlaneDist );
+	FSceneViewInitOptions ViewInitOptions = Internal::MakeSceneViewInitOptions(
+		ViewFrame,
+		Dimensions,
+		HorzFOVDegrees,
+		NearPlaneDist,
+		VisiblePrimitives);
 
 	FName CaptureTypeName;
 	switch (CaptureType)
@@ -1116,24 +1061,29 @@ bool FWorldRenderCapture::CaptureFromPosition(
 	}
 
 	ReadImageBuffer.Reset();
-	TSet<FPrimitiveComponentId> HiddenPrimitives;
 	UE::Internal::RenderSceneVisualizationToTexture( 
 		RenderTargetTexture, this->Dimensions,
 		World->Scene, CaptureTypeName,
-		ViewOrigin, ViewRotationMatrix, ProjectionMatrix, Config,
-		HiddenPrimitives,
-		VisiblePrimitives,
+		ViewInitOptions,
+		Config,
 		ReadImageBuffer,
 		LastCaptureViewMatrices);
 
-	ResultImageOut.SetDimensions(Dimensions);
-	for (int32 yi = 0; yi < Height; ++yi)
 	{
-		for (int32 xi = 0; xi < Width; ++xi)
+		TRACE_CPUPROFILER_EVENT_SCOPE(FWorldRenderCapture::CaptureFromPosition_CopyBufferToImage);
+
+		ResultImageOut.SetDimensions(Dimensions);
+
+		int32 Width = Dimensions.GetWidth();
+		int32 Height = Dimensions.GetHeight();
+		for (int32 yi = 0; yi < Height; ++yi)
 		{
-			FLinearColor PixelColorf = ReadImageBuffer[yi * Width + xi];
-			PixelColorf.A = 1.0f;		// ?
-			ResultImageOut.SetPixel(FVector2i(xi, yi), PixelColorf);
+			for (int32 xi = 0; xi < Width; ++xi)
+			{
+				FLinearColor PixelColorf = ReadImageBuffer[yi * Width + xi];
+				PixelColorf.A = 1.0f;
+				ResultImageOut.SetPixel(FVector2i(xi, yi), PixelColorf);
+			}
 		}
 	}
 
