Fix crashes that can happen when the editor is closed while tool operators are in progress

diff --git a/Engine/Plugins/Experimental/GeometryScripting/Source/GeometryScriptingCore/Private/MeshBakeFunctions.cpp b/Engine/Plugins/Experimental/GeometryScripting/Source/GeometryScriptingCore/Private/MeshBakeFunctions.cpp
--- a/Engine/Plugins/Experimental/GeometryScripting/Source/GeometryScriptingCore/Private/MeshBakeFunctions.cpp
+++ b/Engine/Plugins/Experimental/GeometryScripting/Source/GeometryScriptingCore/Private/MeshBakeFunctions.cpp
@@ -1209,7 +1209,7 @@ FGeometryScriptRenderCaptureTextures UGeometryScriptLibrary_MeshBakeFunctions::B
 	}
 
 	FRenderCaptureTextures TexturesOut;
-	GetTexturesFromRenderCaptureBaker(Baker, TexturesOut);
+	GetTexturesFromRenderCaptureBaker(*Baker, TexturesOut);
 
 	// Update source data
 	if (TexturesOut.BaseColorMap)
diff --git a/Engine/Plugins/Experimental/MeshModelingToolsetExp/Source/MeshModelingToolsEditorOnlyExp/Private/BakeRenderCaptureTool.cpp b/Engine/Plugins/Experimental/MeshModelingToolsetExp/Source/MeshModelingToolsEditorOnlyExp/Private/BakeRenderCaptureTool.cpp
--- a/Engine/Plugins/Experimental/MeshModelingToolsetExp/Source/MeshModelingToolsEditorOnlyExp/Private/BakeRenderCaptureTool.cpp
+++ b/Engine/Plugins/Experimental/MeshModelingToolsetExp/Source/MeshModelingToolsEditorOnlyExp/Private/BakeRenderCaptureTool.cpp
@@ -169,15 +169,15 @@ bool URenderCaptureProperties::operator!=(const URenderCaptureProperties& Other)
 class FRenderCaptureMapBakerOp : public TGenericDataOperator<FMeshMapBaker>
 {
 public:
-	UE::Geometry::FDynamicMesh3* BaseMesh = nullptr;
-	UE::Geometry::FDynamicMeshAABBTree3* BaseMeshSpatial = nullptr;
+	TSharedPtr<UE::Geometry::FDynamicMesh3, ESPMode::ThreadSafe> BaseMesh;
+	TSharedPtr<UE::Geometry::FDynamicMeshAABBTree3, ESPMode::ThreadSafe> BaseMeshSpatial;
 	TSharedPtr<UE::Geometry::FMeshTangentsd, ESPMode::ThreadSafe> BaseMeshTangents;
 	TSharedPtr<TArray<int32>, ESPMode::ThreadSafe> BaseMeshUVCharts;
 	int32 TargetUVLayer;
 	double ValidSampleDepthThreshold;
 	EBakeTextureResolution TextureImageSize;
 	EBakeTextureSamplesPerPixel SamplesPerPixel;
-	FSceneCapturePhotoSet* SceneCapture = nullptr;
+	TSharedPtr<FSceneCapturePhotoSet, ESPMode::ThreadSafe> SceneCapture;
 
 	// Used to pass the channels which need baking via the bBakeXXX and bUsePackedMRS members
 	// PendingBake allows us to skip baking for computed capture types previously baked
@@ -192,10 +192,10 @@ public:
 void FRenderCaptureMapBakerOp::CalculateResult(FProgressCancel*)
 {
 	FSceneCapturePhotoSetSampler Sampler(
-		SceneCapture,
+		SceneCapture.Get(),
 		ValidSampleDepthThreshold,
-		BaseMesh,
-		BaseMeshSpatial,
+		BaseMesh.Get(),
+		BaseMeshSpatial.Get(),
 		BaseMeshTangents.Get());
 
 	const FImageDimensions TextureDimensions(
@@ -205,10 +205,10 @@ void FRenderCaptureMapBakerOp::CalculateResult(FProgressCancel*)
 	FRenderCaptureOcclusionHandler OcclusionHandler(TextureDimensions);
 
 	Result = MakeRenderCaptureBaker(
-		BaseMesh,
+		BaseMesh.Get(),
 		BaseMeshTangents,
 		BaseMeshUVCharts,
-		SceneCapture,
+		SceneCapture.Get(),
 		&Sampler,
 		PendingBake,
 		TargetUVLayer,
@@ -276,20 +276,22 @@ void UBakeRenderCaptureTool::Setup()
 	// Initialize the datastructures used by the bake background compute/tool operator
 	PreviewMesh->ProcessMesh([this](const FDynamicMesh3& Mesh)
 	{
-		TargetMesh.Copy(Mesh);
+		TargetMesh = MakeShared<FDynamicMesh3, ESPMode::ThreadSafe>();
+		TargetMesh->Copy(Mesh);
 		const FTransformSRT3d BaseToWorld = UE::ToolTarget::GetLocalToWorldTransform(Targets[0]);
-		MeshTransforms::ApplyTransform(TargetMesh, BaseToWorld, true);
+		MeshTransforms::ApplyTransform(*TargetMesh, BaseToWorld, true);
 
 		// Initialize UV charts
 		TargetMeshUVCharts = MakeShared<TArray<int32>, ESPMode::ThreadSafe>();
-		FMeshMapBaker::ComputeUVCharts(TargetMesh, *TargetMeshUVCharts);
+		FMeshMapBaker::ComputeUVCharts(*TargetMesh, *TargetMeshUVCharts);
 
 		// Initialize tangents
-		TargetMeshTangents = MakeShared<FMeshTangentsd, ESPMode::ThreadSafe>(&TargetMesh);
-		TargetMeshTangents->CopyTriVertexTangents(TargetMesh);
+		TargetMeshTangents = MakeShared<FMeshTangentsd, ESPMode::ThreadSafe>(TargetMesh.Get());
+		TargetMeshTangents->CopyTriVertexTangents(*TargetMesh);
 
 		// Initialize spatial index
-		TargetMeshSpatial.SetMesh(&TargetMesh, true);
+		TargetMeshSpatial = MakeShared<FDynamicMeshAABBTree3, ESPMode::ThreadSafe>();
+		TargetMeshSpatial->SetMesh(TargetMesh.Get(), true);
 	});
 
 	// Initialize actors
@@ -352,7 +354,7 @@ void UBakeRenderCaptureTool::Setup()
 	InputMeshSettings->RestoreProperties(this);
 	AddToolPropertySource(InputMeshSettings);
 	InputMeshSettings->TargetStaticMesh = UE::ToolTarget::GetStaticMeshFromTargetIfAvailable(Target);
-	UpdateUVLayerNames(InputMeshSettings->TargetUVLayer, InputMeshSettings->TargetUVLayerNamesList, TargetMesh);
+	UpdateUVLayerNames(InputMeshSettings->TargetUVLayer, InputMeshSettings->TargetUVLayerNamesList, *TargetMesh);
 	InputMeshSettings->WatchProperty(InputMeshSettings->TargetUVLayer, [this](FString) { OpState |= EBakeOpState::Evaluate; });
 	
 	{
@@ -381,7 +383,7 @@ void UBakeRenderCaptureTool::Setup()
 		UE::ToolTarget::HideSourceObject(Targets[Idx]);
 	}
 
-	SceneCapture = MakeUnique<FSceneCapturePhotoSet>();
+	SceneCapture = MakeShared<FSceneCapturePhotoSet>();
 
 	// Make sure we trigger SceneCapture computation in UpdateResult
 	OpState |= EBakeOpState::Evaluate;
@@ -790,15 +792,15 @@ TUniquePtr<TGenericDataOperator<FMeshMapBaker>> UBakeRenderCaptureTool::MakeNewO
 	check(SceneCapture.IsValid());
 
 	TUniquePtr<FRenderCaptureMapBakerOp> Op = MakeUnique<FRenderCaptureMapBakerOp>();
-	Op->BaseMesh = &TargetMesh;
-	Op->BaseMeshSpatial = &TargetMeshSpatial;
+	Op->BaseMesh = TargetMesh;
+	Op->BaseMeshSpatial = TargetMeshSpatial;
 	Op->BaseMeshTangents = TargetMeshTangents;
 	Op->BaseMeshUVCharts = TargetMeshUVCharts;
 	Op->TargetUVLayer = InputMeshSettings->GetTargetUVLayerIndex();
 	Op->ValidSampleDepthThreshold = Settings->ValidSampleDepthThreshold;
 	Op->TextureImageSize = Settings->TextureSize;
 	Op->SamplesPerPixel = Settings->SamplesPerPixel;
-	Op->SceneCapture = SceneCapture.Get();
+	Op->SceneCapture = SceneCapture;
 
 	ForEachCaptureType([this, &Op](ERenderCaptureType CaptureType)
 	{
@@ -822,7 +824,7 @@ void UBakeRenderCaptureTool::OnMapsUpdated(const TUniquePtr<FMeshMapBaker>& NewR
 	TRACE_CPUPROFILER_EVENT_SCOPE(BakeRenderCaptureTool_Textures_BuildTextures);
 
 	FRenderCaptureTextures TexturesOut;
-	GetTexturesFromRenderCaptureBaker(NewResult, TexturesOut);
+	GetTexturesFromRenderCaptureBaker(*NewResult, TexturesOut);
 
 	// The NewResult will contain the newly baked textures so we only update those and not overwrite any already baked
 	// valid TexturesOut. If a texture is invalidated by some tool property change it will null on entry to this function
@@ -871,7 +873,7 @@ bool UBakeRenderCaptureTool::ValidTargetMeshTangents()
 {
 	if (bCheckTargetMeshTangents)
 	{
-		bValidTargetMeshTangents = TargetMeshTangents ? FDynamicMeshTangents(&TargetMesh).HasValidTangents(true) : false;
+		bValidTargetMeshTangents = TargetMeshTangents ? FDynamicMeshTangents(TargetMesh.Get()).HasValidTangents(true) : false;
 		bCheckTargetMeshTangents = false;
 	}
 	return bValidTargetMeshTangents;
@@ -1073,7 +1075,7 @@ void UBakeRenderCaptureTool::UpdateResult()
 		{
 			const auto HasDegenerateUVs = [this]
 			{
-				FDynamicMeshUVOverlay* UVOverlay = TargetMesh.Attributes()->GetUVLayer(InputMeshSettings->GetTargetUVLayerIndex());
+				FDynamicMeshUVOverlay* UVOverlay = TargetMesh->Attributes()->GetUVLayer(InputMeshSettings->GetTargetUVLayerIndex());
 				FAxisAlignedBox2f Bounds = FAxisAlignedBox2f::Empty();
 				for (const int Index : UVOverlay->ElementIndicesItr())
 				{
@@ -1084,7 +1086,7 @@ void UBakeRenderCaptureTool::UpdateResult()
 				return Bounds.Min == Bounds.Max;
 			};
 
-			if (TargetMesh.Attributes()->GetUVLayer(InputMeshSettings->GetTargetUVLayerIndex()) == nullptr)
+			if (TargetMesh->Attributes()->GetUVLayer(InputMeshSettings->GetTargetUVLayerIndex()) == nullptr)
 			{
 				ErrorMessage = LOCTEXT("TargetMeshMissingUVs", "The Target Mesh UV layer is missing");
 			}
@@ -1140,7 +1142,7 @@ void UBakeRenderCaptureTool::UpdateResult()
 				UE::ToolTarget::ShowSourceObject(Targets[Idx]);
 			}
 
-			UpdatedChannels = UpdateSceneCapture(SceneCapture, Actors, DesiredConfig, false);
+			UpdatedChannels = UpdateSceneCapture(*SceneCapture, Actors, DesiredConfig, false);
 
 			// Hide the source meshes after the render capture so they don't occlude the preview
 			for (int Idx = 1; Idx < Targets.Num(); ++Idx)
@@ -1148,7 +1150,7 @@ void UBakeRenderCaptureTool::UpdateResult()
 				UE::ToolTarget::HideSourceObject(Targets[Idx]);
 			}
 
-			const FSceneCaptureConfig AchievedConfig = GetSceneCaptureConfig(SceneCapture);
+			const FSceneCaptureConfig AchievedConfig = GetSceneCaptureConfig(*SceneCapture);
 			ensure(SceneCapture->Cancelled() == (AchievedConfig != DesiredConfig));
 
 			// If the scene capture was cancelled make sure the tool properties are consistent with the computed captures
@@ -1404,7 +1406,7 @@ void UBakeRenderCaptureTool::GatherAnalytics(FBakeAnalytics::FMeshSettings& Data
 		return;
 	}
 
-	Data.NumTargetMeshTris = TargetMesh.TriangleCount();
+	Data.NumTargetMeshTris = TargetMesh->TriangleCount();
 	Data.NumDetailMesh = Actors.Num();
 	Data.NumDetailMeshTris = 0;
 	for (AActor* Actor : Actors)
diff --git a/Engine/Plugins/Experimental/MeshModelingToolsetExp/Source/MeshModelingToolsEditorOnlyExp/Public/BakeRenderCaptureTool.h b/Engine/Plugins/Experimental/MeshModelingToolsetExp/Source/MeshModelingToolsEditorOnlyExp/Public/BakeRenderCaptureTool.h
--- a/Engine/Plugins/Experimental/MeshModelingToolsetExp/Source/MeshModelingToolsEditorOnlyExp/Public/BakeRenderCaptureTool.h
+++ b/Engine/Plugins/Experimental/MeshModelingToolsetExp/Source/MeshModelingToolsEditorOnlyExp/Public/BakeRenderCaptureTool.h
@@ -353,8 +353,8 @@ protected:
 	TUniquePtr<TGenericDataBackgroundCompute<UE::Geometry::FMeshMapBaker>> Compute = nullptr;
 	EBakeOpState OpState = EBakeOpState::Evaluate;
 
-	UE::Geometry::FDynamicMesh3 TargetMesh;
-	UE::Geometry::FDynamicMeshAABBTree3 TargetMeshSpatial;
+	TSharedPtr<UE::Geometry::FDynamicMesh3, ESPMode::ThreadSafe> TargetMesh;
+	TSharedPtr<UE::Geometry::FDynamicMeshAABBTree3, ESPMode::ThreadSafe> TargetMeshSpatial;
 	TSharedPtr<TArray<int32>, ESPMode::ThreadSafe> TargetMeshUVCharts;
 	TSharedPtr<UE::Geometry::FMeshTangentsd, ESPMode::ThreadSafe> TargetMeshTangents;
 
@@ -401,7 +401,7 @@ protected:
 	TObjectPtr<UPreviewMesh> PreviewMesh;
 
 	// Note: We need to compute this on the game thread because the implementation has checks for this
-	TUniquePtr<UE::Geometry::FSceneCapturePhotoSet> SceneCapture = nullptr;
+	TSharedPtr<UE::Geometry::FSceneCapturePhotoSet, ESPMode::ThreadSafe> SceneCapture;
 
 	// These are used to determine if we need to re-bake results
 	float ComputedValidDepthThreshold;
diff --git a/Engine/Plugins/Experimental/MeshModelingToolsetExp/Source/MeshModelingToolsExp/Private/BakeMeshAttributeMapsTool.cpp b/Engine/Plugins/Experimental/MeshModelingToolsetExp/Source/MeshModelingToolsExp/Private/BakeMeshAttributeMapsTool.cpp
--- a/Engine/Plugins/Experimental/MeshModelingToolsetExp/Source/MeshModelingToolsExp/Private/BakeMeshAttributeMapsTool.cpp
+++ b/Engine/Plugins/Experimental/MeshModelingToolsetExp/Source/MeshModelingToolsExp/Private/BakeMeshAttributeMapsTool.cpp
@@ -101,7 +101,7 @@ public:
 	TSharedPtr<UE::Geometry::FDynamicMesh3, ESPMode::ThreadSafe> DetailMesh;
 	TSharedPtr<UE::Geometry::FDynamicMeshAABBTree3, ESPMode::ThreadSafe> DetailSpatial;
 	TSharedPtr<UE::Geometry::TMeshTangents<double>, ESPMode::ThreadSafe> DetailMeshTangents;
-	UE::Geometry::FDynamicMesh3* BaseMesh = nullptr;
+	TSharedPtr<UE::Geometry::FDynamicMesh3, ESPMode::ThreadSafe> BaseMesh;
 	TUniquePtr<UE::Geometry::FMeshMapBaker> Baker;
 	UBakeMeshAttributeMapsTool::FBakeSettings BakeSettings;
 	TSharedPtr<UE::Geometry::TMeshTangents<double>, ESPMode::ThreadSafe> BaseMeshTangents;
@@ -134,7 +134,7 @@ public:
 		Baker->CancelF = [Progress]() {
 			return Progress && Progress->Cancelled();
 		};
-		Baker->SetTargetMesh(BaseMesh);
+		Baker->SetTargetMesh(BaseMesh.Get());
 		Baker->SetTargetMeshUVLayer(BakeSettings.TargetUVLayer);
 		Baker->SetDimensions(BakeSettings.Dimensions);
 		Baker->SetProjectionDistance(BakeSettings.ProjectionDistance);
@@ -292,9 +292,10 @@ void UBakeMeshAttributeMapsTool::Setup()
 
 	PreviewMesh->ProcessMesh([this](const FDynamicMesh3& Mesh)
 	{
-		TargetMesh.Copy(Mesh);
-		TargetSpatial.SetMesh(&TargetMesh, true);
-		TargetMeshTangents = MakeShared<FMeshTangentsd, ESPMode::ThreadSafe>(&TargetMesh);
+		TargetMesh = MakeShared<FDynamicMesh3, ESPMode::ThreadSafe>();
+		TargetMesh->Copy(Mesh);
+
+		TargetMeshTangents = MakeShared<FMeshTangentsd, ESPMode::ThreadSafe>(TargetMesh.Get());
 		TargetMeshTangents->CopyTriVertexTangents(Mesh);
 	});
 
@@ -328,7 +329,7 @@ void UBakeMeshAttributeMapsTool::Setup()
 	InputMeshSettings->SourceSkeletalMesh = !bIsBakeToSelf ? UE::ToolTarget::GetSkeletalMeshFromTargetIfAvailable(DetailTarget) : nullptr;
 	InputMeshSettings->SourceDynamicMesh = !bIsBakeToSelf ? GetTargetActorViaIPersistentDynamicMeshSource(DetailTarget) : nullptr;
 	InputMeshSettings->SourceNormalMap = nullptr;
-	UpdateUVLayerNames(InputMeshSettings->TargetUVLayer, InputMeshSettings->TargetUVLayerNamesList, TargetMesh);
+	UpdateUVLayerNames(InputMeshSettings->TargetUVLayer, InputMeshSettings->TargetUVLayerNamesList, *TargetMesh);
 	InputMeshSettings->WatchProperty(InputMeshSettings->bHideSourceMesh, [this](bool bState) { SetSourceObjectVisible(!bState); });
 	InputMeshSettings->WatchProperty(InputMeshSettings->TargetUVLayer, [this](FString) { OpState |= EBakeOpState::Evaluate; });
 	InputMeshSettings->WatchProperty(InputMeshSettings->ProjectionDistance, [this](float) { OpState |= EBakeOpState::Evaluate; });
@@ -435,7 +436,7 @@ TUniquePtr<UE::Geometry::TGenericDataOperator<FMeshMapBaker>> UBakeMeshAttribute
 	TUniquePtr<FMeshMapBakerOp> Op = MakeUnique<FMeshMapBakerOp>();
 	Op->DetailMesh = DetailMesh;
 	Op->DetailSpatial = DetailSpatial;
-	Op->BaseMesh = &TargetMesh;
+	Op->BaseMesh = TargetMesh;
 	Op->BakeSettings = CachedBakeSettings;
 	Op->BaseMeshUVCharts = TargetMeshUVCharts;
 	Op->bIsBakeToSelf = bIsBakeToSelf;
@@ -821,7 +822,7 @@ void UBakeMeshAttributeMapsTool::GatherAnalytics(FBakeAnalytics::FMeshSettings&
 {
 	if (FEngineAnalytics::IsAvailable())
 	{
-		Data.NumTargetMeshTris = TargetMesh.TriangleCount();
+		Data.NumTargetMeshTris = TargetMesh->TriangleCount();
 		Data.NumDetailMesh = 1;
 		Data.NumDetailMeshTris = DetailMesh->TriangleCount();
 	}
diff --git a/Engine/Plugins/Experimental/MeshModelingToolsetExp/Source/MeshModelingToolsExp/Private/BakeMeshAttributeMapsToolBase.cpp b/Engine/Plugins/Experimental/MeshModelingToolsetExp/Source/MeshModelingToolsExp/Private/BakeMeshAttributeMapsToolBase.cpp
--- a/Engine/Plugins/Experimental/MeshModelingToolsetExp/Source/MeshModelingToolsExp/Private/BakeMeshAttributeMapsToolBase.cpp
+++ b/Engine/Plugins/Experimental/MeshModelingToolsetExp/Source/MeshModelingToolsExp/Private/BakeMeshAttributeMapsToolBase.cpp
@@ -71,7 +71,7 @@ void UBakeMeshAttributeMapsToolBase::PostSetup()
 	// Initialize UV charts
 	// TODO: Compute UV charts asynchronously
 	TargetMeshUVCharts = MakeShared<TArray<int32>, ESPMode::ThreadSafe>();
-	FMeshMapBaker::ComputeUVCharts(TargetMesh, *TargetMeshUVCharts);
+	FMeshMapBaker::ComputeUVCharts(*TargetMesh, *TargetMeshUVCharts);
 
 	GatherAnalytics(BakeAnalytics.MeshSettings);
 }
diff --git a/Engine/Plugins/Experimental/MeshModelingToolsetExp/Source/MeshModelingToolsExp/Private/BakeMeshAttributeTool.cpp b/Engine/Plugins/Experimental/MeshModelingToolsetExp/Source/MeshModelingToolsExp/Private/BakeMeshAttributeTool.cpp
--- a/Engine/Plugins/Experimental/MeshModelingToolsetExp/Source/MeshModelingToolsExp/Private/BakeMeshAttributeTool.cpp
+++ b/Engine/Plugins/Experimental/MeshModelingToolsetExp/Source/MeshModelingToolsExp/Private/BakeMeshAttributeTool.cpp
@@ -29,7 +29,7 @@ bool UBakeMeshAttributeTool::ValidTargetMeshTangents()
 {
 	if (bCheckTargetMeshTangents)
 	{
-		bValidTargetMeshTangents = TargetMeshTangents ? FDynamicMeshTangents(&TargetMesh).HasValidTangents(true) : false;
+		bValidTargetMeshTangents = TargetMeshTangents ? FDynamicMeshTangents(TargetMesh.Get()).HasValidTangents(true) : false;
 		bCheckTargetMeshTangents = false;
 	}
 	return bValidTargetMeshTangents;
diff --git a/Engine/Plugins/Experimental/MeshModelingToolsetExp/Source/MeshModelingToolsExp/Private/BakeMeshAttributeVertexTool.cpp b/Engine/Plugins/Experimental/MeshModelingToolsetExp/Source/MeshModelingToolsExp/Private/BakeMeshAttributeVertexTool.cpp
--- a/Engine/Plugins/Experimental/MeshModelingToolsetExp/Source/MeshModelingToolsExp/Private/BakeMeshAttributeVertexTool.cpp
+++ b/Engine/Plugins/Experimental/MeshModelingToolsetExp/Source/MeshModelingToolsExp/Private/BakeMeshAttributeVertexTool.cpp
@@ -273,10 +273,9 @@ void UBakeMeshAttributeVertexTool::Setup()
 	// TargetMesh stores the original target mesh. It is intended to remain
 	// const throughout this tool and is used to refresh the PreviewMesh back
 	// to its original state.
-	TargetMesh = UE::ToolTarget::GetDynamicMeshCopy(Targets[0], true);
-	TargetSpatial.SetMesh(&TargetMesh, true);
-	TargetMeshTangents = MakeShared<FMeshTangentsd, ESPMode::ThreadSafe>(&TargetMesh);
-	TargetMeshTangents->CopyTriVertexTangents(TargetMesh);
+	TargetMesh = MakeShared<FDynamicMesh3, ESPMode::ThreadSafe>(UE::ToolTarget::GetDynamicMeshCopy(Targets[0], true));
+	TargetMeshTangents = MakeShared<FMeshTangentsd, ESPMode::ThreadSafe>(TargetMesh.Get());
+	TargetMeshTangents->CopyTriVertexTangents(*TargetMesh);
 
 	// PreviewMesh stores computed result mesh. On shutdown, PreviewMesh will be
 	// used to commit the dynamic mesh to the target tool target.
@@ -626,7 +625,7 @@ void UBakeMeshAttributeVertexTool::UpdateColorTopology()
 			}, 0.0f);
 
 		// Copy source vertex colors onto new color overlay topology.
-		const FDynamicMeshColorOverlay* TargetColorOverlay = TargetMesh.HasAttributes() ? TargetMesh.Attributes()->PrimaryColors() : nullptr;
+		const FDynamicMeshColorOverlay* TargetColorOverlay = TargetMesh->HasAttributes() ? TargetMesh->Attributes()->PrimaryColors() : nullptr;
 		FDynamicMeshColorOverlay* PreviewColorOverlay = Mesh.Attributes()->PrimaryColors(); 
 		if (TargetColorOverlay)
 		{
@@ -782,8 +781,8 @@ void UBakeMeshAttributeVertexTool::GatherAnalytics(FBakeAnalytics::FMeshSettings
 		return;
 	}
 	
-	Data.NumTargetMeshVerts = TargetMesh.VertexCount();
-	Data.NumTargetMeshTris = TargetMesh.TriangleCount();
+	Data.NumTargetMeshVerts = TargetMesh->VertexCount();
+	Data.NumTargetMeshTris = TargetMesh->TriangleCount();
 	Data.NumDetailMesh = 1;
 	Data.NumDetailMeshTris = DetailMesh->TriangleCount();
 }
diff --git a/Engine/Plugins/Experimental/MeshModelingToolsetExp/Source/MeshModelingToolsExp/Private/BakeMultiMeshAttributeMapsTool.cpp b/Engine/Plugins/Experimental/MeshModelingToolsetExp/Source/MeshModelingToolsExp/Private/BakeMultiMeshAttributeMapsTool.cpp
--- a/Engine/Plugins/Experimental/MeshModelingToolsetExp/Source/MeshModelingToolsExp/Private/BakeMultiMeshAttributeMapsTool.cpp
+++ b/Engine/Plugins/Experimental/MeshModelingToolsetExp/Source/MeshModelingToolsExp/Private/BakeMultiMeshAttributeMapsTool.cpp
@@ -375,7 +375,7 @@ class FMultiMeshMapBakerOp : public TGenericDataOperator<FMeshMapBaker>
 public:
 	// General bake settings
 	TSharedPtr<FMeshSceneAdapter, ESPMode::ThreadSafe> DetailMeshScene;
-	UE::Geometry::FDynamicMesh3* BaseMesh = nullptr;
+	TSharedPtr<UE::Geometry::FDynamicMesh3, ESPMode::ThreadSafe> BaseMesh;
 	TSharedPtr<UE::Geometry::FMeshTangentsd, ESPMode::ThreadSafe> BaseMeshTangents;
 	TUniquePtr<UE::Geometry::FMeshMapBaker> Baker;
 	UBakeMultiMeshAttributeMapsTool::FBakeSettings BakeSettings;
@@ -395,7 +395,7 @@ public:
 		Baker->CancelF = [Progress]() {
 			return Progress && Progress->Cancelled();
 		};
-		Baker->SetTargetMesh(BaseMesh);
+		Baker->SetTargetMesh(BaseMesh.Get());
 		Baker->SetTargetMeshUVLayer(BakeSettings.TargetUVLayer);
 		Baker->SetDimensions(BakeSettings.Dimensions);
 		Baker->SetProjectionDistance(BakeSettings.ProjectionDistance);
@@ -486,10 +486,11 @@ void UBakeMultiMeshAttributeMapsTool::Setup()
 	// Initialize base mesh
 	PreviewMesh->ProcessMesh([this](const FDynamicMesh3& Mesh)
 	{
-		TargetMesh.Copy(Mesh);
-		TargetMeshTangents = MakeShared<FMeshTangentsd, ESPMode::ThreadSafe>(&TargetMesh);
+		TargetMesh = MakeShared<FDynamicMesh3, ESPMode::ThreadSafe>();
+		TargetMesh->Copy(Mesh);
+
+		TargetMeshTangents = MakeShared<FMeshTangentsd, ESPMode::ThreadSafe>(TargetMesh.Get());
 		TargetMeshTangents->CopyTriVertexTangents(Mesh);
-		TargetSpatial.SetMesh(&TargetMesh, true);
 	});
 
 	// Initialize detail sampler
@@ -531,7 +532,7 @@ void UBakeMultiMeshAttributeMapsTool::Setup()
 	InputMeshSettings->TargetStaticMesh = UE::ToolTarget::GetStaticMeshFromTargetIfAvailable(Target);
 	InputMeshSettings->TargetSkeletalMesh = UE::ToolTarget::GetSkeletalMeshFromTargetIfAvailable(Target);
 	InputMeshSettings->TargetDynamicMesh = GetTargetActorViaIPersistentDynamicMeshSource(Target);
-	UpdateUVLayerNames(InputMeshSettings->TargetUVLayer, InputMeshSettings->TargetUVLayerNamesList, TargetMesh);
+	UpdateUVLayerNames(InputMeshSettings->TargetUVLayer, InputMeshSettings->TargetUVLayerNamesList, *TargetMesh);
 	InputMeshSettings->WatchProperty(InputMeshSettings->TargetUVLayer, [this](FString) { OpState |= EBakeOpState::Evaluate; });
 	InputMeshSettings->WatchProperty(InputMeshSettings->ProjectionDistance, [this](float) { OpState |= EBakeOpState::Evaluate; });
 
@@ -609,7 +610,7 @@ TUniquePtr<UE::Geometry::TGenericDataOperator<FMeshMapBaker>> UBakeMultiMeshAttr
 {
 	TUniquePtr<FMultiMeshMapBakerOp> Op = MakeUnique<FMultiMeshMapBakerOp>();
 	Op->DetailMeshScene = DetailMeshScene;
-	Op->BaseMesh = &TargetMesh;
+	Op->BaseMesh = TargetMesh;
 	Op->BaseMeshUVCharts = TargetMeshUVCharts;
 	Op->BaseToWorldTransform = UE::ToolTarget::GetLocalToWorldTransform(Targets[0]);
 	Op->BakeSettings = CachedBakeSettings;
@@ -819,7 +820,7 @@ void UBakeMultiMeshAttributeMapsTool::GatherAnalytics(FBakeAnalytics::FMeshSetti
 {
 	if (FEngineAnalytics::IsAvailable())
 	{
-		Data.NumTargetMeshTris = TargetMesh.TriangleCount();
+		Data.NumTargetMeshTris = TargetMesh->TriangleCount();
 		Data.NumDetailMesh = 0;
 		Data.NumDetailMeshTris = 0;
 		DetailMeshScene->ProcessActorChildMeshes([&Data](const FActorAdapter* ActorAdapter, const FActorChildMesh* ChildMesh)
diff --git a/Engine/Plugins/Experimental/MeshModelingToolsetExp/Source/MeshModelingToolsExp/Public/BakeMeshAttributeTool.h b/Engine/Plugins/Experimental/MeshModelingToolsetExp/Source/MeshModelingToolsExp/Public/BakeMeshAttributeTool.h
--- a/Engine/Plugins/Experimental/MeshModelingToolsetExp/Source/MeshModelingToolsExp/Public/BakeMeshAttributeTool.h
+++ b/Engine/Plugins/Experimental/MeshModelingToolsetExp/Source/MeshModelingToolsExp/Public/BakeMeshAttributeTool.h
@@ -129,9 +129,8 @@ protected:
 	//
 	EBakeOpState OpState = EBakeOpState::Evaluate;
 
-	UE::Geometry::FDynamicMesh3 TargetMesh;
+	TSharedPtr<UE::Geometry::FDynamicMesh3, ESPMode::ThreadSafe> TargetMesh;
 	TSharedPtr<UE::Geometry::TMeshTangents<double>, ESPMode::ThreadSafe> TargetMeshTangents;
-	UE::Geometry::FDynamicMeshAABBTree3 TargetSpatial;
 
 	/**
 	 * Compute validity of the Target Mesh tangents. Only checks validity
diff --git a/Engine/Plugins/Runtime/MeshModelingToolset/Source/ModelingComponents/Private/Baking/RenderCaptureFunctions.cpp b/Engine/Plugins/Runtime/MeshModelingToolset/Source/ModelingComponents/Private/Baking/RenderCaptureFunctions.cpp
--- a/Engine/Plugins/Runtime/MeshModelingToolset/Source/ModelingComponents/Private/Baking/RenderCaptureFunctions.cpp
+++ b/Engine/Plugins/Runtime/MeshModelingToolset/Source/ModelingComponents/Private/Baking/RenderCaptureFunctions.cpp
@@ -225,30 +225,24 @@ bool FSceneCaptureConfig::operator!=(const FSceneCaptureConfig& Other) const
 
 
 void UE::Geometry::ConfigureSceneCapture(
-	const TUniquePtr<FSceneCapturePhotoSet>& SceneCapture,
+	FSceneCapturePhotoSet& SceneCapture,
 	const TArray<TObjectPtr<AActor>>& Actors,
 	const FSceneCaptureConfig& Options,
 	bool bAllowCancel)
 {
-	if (SceneCapture.IsValid() == false)
-	{
-		ensure(false);
-		return;
-	}
-
 	ForEachCaptureType([&SceneCapture, &Options](ERenderCaptureType CaptureType)
 	{
 		const bool bCaptureTypeEnabled = Options.Flags[CaptureType];
-		SceneCapture->SetCaptureTypeEnabled(CaptureType, bCaptureTypeEnabled);
+		SceneCapture.SetCaptureTypeEnabled(CaptureType, bCaptureTypeEnabled);
 
 		FRenderCaptureConfig Config;
 		Config.bAntiAliasing = (CaptureType == ERenderCaptureType::DeviceDepth ? false : Options.bAntiAliasing);
-		SceneCapture->SetCaptureConfig(CaptureType, Config);
+		SceneCapture.SetCaptureConfig(CaptureType, Config);
 	});
 
 	UWorld* World = Actors.IsEmpty() ? nullptr : Actors[0]->GetWorld();
 
-	SceneCapture->SetCaptureSceneActors(World, Actors);
+	SceneCapture.SetCaptureSceneActors(World, Actors);
 
 	const TArray<FSpatialPhotoParams> SpatialParams = ComputeStandardExteriorSpatialPhotoParameters(
 		World,
@@ -259,31 +253,29 @@ void UE::Geometry::ConfigureSceneCapture(
 		Options.NearPlaneDist,
 		true, true, true, true, true);
 
-	SceneCapture->SetSpatialPhotoParams(SpatialParams);
+	SceneCapture.SetSpatialPhotoParams(SpatialParams);
 
-	SceneCapture->SetAllowCancel(bAllowCancel);
+	SceneCapture.SetAllowCancel(bAllowCancel);
 }
 
 
 
 UE::Geometry::FRenderCaptureTypeFlags UE::Geometry::UpdateSceneCapture(
-	const TUniquePtr<FSceneCapturePhotoSet>& SceneCapture,
+	FSceneCapturePhotoSet& SceneCapture,
 	const TArray<TObjectPtr<AActor>>& Actors,
 	const FSceneCaptureConfig& DesiredConfig,
 	bool bAllowCancel)
 {
-	check(SceneCapture.IsValid());
-
-	const FSceneCapturePhotoSet::FStatus PreConfigStatus = SceneCapture->GetSceneCaptureStatus();
+	const FSceneCapturePhotoSet::FStatus PreConfigStatus = SceneCapture.GetSceneCaptureStatus();
 	ConfigureSceneCapture(SceneCapture, Actors, DesiredConfig, true);
-	const FSceneCapturePhotoSet::FStatus PreComputeStatus = SceneCapture->GetSceneCaptureStatus();
-	SceneCapture->Compute();
-	const FSceneCapturePhotoSet::FStatus PostComputeStatus = SceneCapture->GetSceneCaptureStatus();
+	const FSceneCapturePhotoSet::FStatus PreComputeStatus = SceneCapture.GetSceneCaptureStatus();
+	SceneCapture.Compute();
+	const FSceneCapturePhotoSet::FStatus PostComputeStatus = SceneCapture.GetSceneCaptureStatus();
 
 	const FSceneCaptureConfig AchievedConfig = GetSceneCaptureConfig(SceneCapture);
 
 	// DesiredConfig and AchievedConfig mismatch iff the SceneCapture was cancelled. Also, the mismatch should only be in the Flags member
-	ensure(SceneCapture->Cancelled() == (AchievedConfig != DesiredConfig));
+	ensure(SceneCapture.Cancelled() == (AchievedConfig != DesiredConfig));
 	ensure(AchievedConfig.RenderCaptureImageSize == DesiredConfig.RenderCaptureImageSize);
 	ensure(AchievedConfig.FieldOfViewDegrees == DesiredConfig.FieldOfViewDegrees);
 	ensure(AchievedConfig.bAntiAliasing == DesiredConfig.bAntiAliasing);
@@ -374,17 +366,17 @@ PRAGMA_ENABLE_DEPRECATION_WARNINGS
 
 // Return the SceneCapture configuration, the Flags in the returned struct are set if the corresponding capture type is computed
 FSceneCaptureConfig UE::Geometry::GetSceneCaptureConfig(
-	const TUniquePtr<FSceneCapturePhotoSet>& SceneCapture,
+	const FSceneCapturePhotoSet& SceneCapture,
 	const FSceneCapturePhotoSet::ECaptureTypeStatus QueryStatus)
 {
 	FSceneCaptureConfig Result;
 
 	ForEachCaptureType([&SceneCapture, &QueryStatus, &Result](ERenderCaptureType CaptureType)
 	{
-		Result.Flags[CaptureType] = (SceneCapture->GetCaptureTypeStatus(CaptureType) == QueryStatus);
+		Result.Flags[CaptureType] = (SceneCapture.GetCaptureTypeStatus(CaptureType) == QueryStatus);
 	});
 
-	const TArray<FSpatialPhotoParams> SpatialParams = SceneCapture->GetSpatialPhotoParams();
+	const TArray<FSpatialPhotoParams> SpatialParams = SceneCapture.GetSpatialPhotoParams();
 
 	if (!SpatialParams.IsEmpty())
 	{
@@ -394,16 +386,16 @@ FSceneCaptureConfig UE::Geometry::GetSceneCaptureConfig(
 		Result.NearPlaneDist = SpatialParams[0].NearPlaneDist;
 
 		const bool Values[] = {
-			SceneCapture->GetCaptureConfig(ERenderCaptureType::BaseColor).bAntiAliasing,
-			SceneCapture->GetCaptureConfig(ERenderCaptureType::WorldNormal).bAntiAliasing,
-			SceneCapture->GetCaptureConfig(ERenderCaptureType::Metallic).bAntiAliasing,
-			SceneCapture->GetCaptureConfig(ERenderCaptureType::Roughness).bAntiAliasing,
-			SceneCapture->GetCaptureConfig(ERenderCaptureType::Specular).bAntiAliasing,
-			SceneCapture->GetCaptureConfig(ERenderCaptureType::CombinedMRS).bAntiAliasing,
-			SceneCapture->GetCaptureConfig(ERenderCaptureType::Emissive).bAntiAliasing,
-			SceneCapture->GetCaptureConfig(ERenderCaptureType::Opacity).bAntiAliasing,
-			SceneCapture->GetCaptureConfig(ERenderCaptureType::SubsurfaceColor).bAntiAliasing,
-			SceneCapture->GetCaptureConfig(ERenderCaptureType::DeviceDepth).bAntiAliasing,
+			SceneCapture.GetCaptureConfig(ERenderCaptureType::BaseColor).bAntiAliasing,
+			SceneCapture.GetCaptureConfig(ERenderCaptureType::WorldNormal).bAntiAliasing,
+			SceneCapture.GetCaptureConfig(ERenderCaptureType::Metallic).bAntiAliasing,
+			SceneCapture.GetCaptureConfig(ERenderCaptureType::Roughness).bAntiAliasing,
+			SceneCapture.GetCaptureConfig(ERenderCaptureType::Specular).bAntiAliasing,
+			SceneCapture.GetCaptureConfig(ERenderCaptureType::CombinedMRS).bAntiAliasing,
+			SceneCapture.GetCaptureConfig(ERenderCaptureType::Emissive).bAntiAliasing,
+			SceneCapture.GetCaptureConfig(ERenderCaptureType::Opacity).bAntiAliasing,
+			SceneCapture.GetCaptureConfig(ERenderCaptureType::SubsurfaceColor).bAntiAliasing,
+			SceneCapture.GetCaptureConfig(ERenderCaptureType::DeviceDepth).bAntiAliasing,
 		};
 		Result.bAntiAliasing = Algo::AnyOf(Values);
 	}
@@ -706,15 +698,15 @@ TUniquePtr<FMeshMapBaker> UE::Geometry::MakeRenderCaptureBaker(
 PRAGMA_ENABLE_DEPRECATION_WARNINGS
 
 
-void UE::Geometry::GetTexturesFromRenderCaptureBaker(const TUniquePtr<FMeshMapBaker>& Baker, FRenderCaptureTextures& TexturesOut)
+void UE::Geometry::GetTexturesFromRenderCaptureBaker(FMeshMapBaker& Baker, FRenderCaptureTextures& TexturesOut)
 {
 	// We do this to defer work I guess, it was like this in the original ApproximateActors implementation :DeferredPopulateSourceData 
 	constexpr bool bPopulateSourceData = false;
 
-	const int32 NumEval = Baker->NumEvaluators();
+	const int32 NumEval = Baker.NumEvaluators();
 	for (int32 EvalIdx = 0; EvalIdx < NumEval; ++EvalIdx)
 	{
-		FMeshMapEvaluator* BaseEval = Baker->GetEvaluator(EvalIdx);
+		FMeshMapEvaluator* BaseEval = Baker.GetEvaluator(EvalIdx);
 		check(BaseEval->DataLayout().Num() == 1);
 		switch (BaseEval->DataLayout()[0])
 		{
@@ -722,7 +714,7 @@ void UE::Geometry::GetTexturesFromRenderCaptureBaker(const TUniquePtr<FMeshMapBa
 
 			{
 				FRenderCaptureMapEvaluator<FVector4f>* Eval = static_cast<FRenderCaptureMapEvaluator<FVector4f>*>(BaseEval);
-				TUniquePtr<TImageBuilder<FVector4f>> ImageBuilder = MoveTemp(Baker->GetBakeResults(EvalIdx)[0]);
+				TUniquePtr<TImageBuilder<FVector4f>> ImageBuilder = MoveTemp(Baker.GetBakeResults(EvalIdx)[0]);
 				
 				if (ensure(ImageBuilder.IsValid()) == false) return;
 
@@ -827,7 +819,7 @@ void UE::Geometry::GetTexturesFromRenderCaptureBaker(const TUniquePtr<FMeshMapBa
 
 			{
 				FRenderCaptureMapEvaluator<FVector3f>* Eval = static_cast<FRenderCaptureMapEvaluator<FVector3f>*>(BaseEval);
-				TUniquePtr<TImageBuilder<FVector4f>> ImageBuilder = MoveTemp(Baker->GetBakeResults(EvalIdx)[0]);
+				TUniquePtr<TImageBuilder<FVector4f>> ImageBuilder = MoveTemp(Baker.GetBakeResults(EvalIdx)[0]);
 				
 				if (ensure(ImageBuilder.IsValid()) == false) return;
 				if (ensure(Eval->Channel == ERenderCaptureChannel::WorldNormal) == false) return;
@@ -848,4 +840,39 @@ void UE::Geometry::GetTexturesFromRenderCaptureBaker(const TUniquePtr<FMeshMapBa
 }
 
 
+
+////////////////////////////////////////////////////////////////////////////////////////////////////////////////////////
+// Everything that follows is deprecated 
+////////////////////////////////////////////////////////////////////////////////////////////////////////////////////////
+
+void UE::Geometry::ConfigureSceneCapture(
+	const TUniquePtr<FSceneCapturePhotoSet>& SceneCapture,
+	const TArray<TObjectPtr<AActor>>& Actors,
+	const FSceneCaptureConfig& Options,
+	bool bAllowCancel)
+{
+	return ConfigureSceneCapture(*SceneCapture, Actors, Options, bAllowCancel);
+}
+
+UE::Geometry::FRenderCaptureTypeFlags UE::Geometry::UpdateSceneCapture(
+	const TUniquePtr<FSceneCapturePhotoSet>& SceneCapture,
+	const TArray<TObjectPtr<AActor>>& Actors,
+	const FSceneCaptureConfig& DesiredConfig,
+	bool bAllowCancel)
+{
+	return UpdateSceneCapture(*SceneCapture, Actors, DesiredConfig, bAllowCancel);
+}
+
+FSceneCaptureConfig UE::Geometry::GetSceneCaptureConfig(
+	const TUniquePtr<FSceneCapturePhotoSet>& SceneCapture,
+	const FSceneCapturePhotoSet::ECaptureTypeStatus QueryStatus)
+{
+	return GetSceneCaptureConfig(*SceneCapture, QueryStatus);
+}
+
+void UE::Geometry::GetTexturesFromRenderCaptureBaker(const TUniquePtr<FMeshMapBaker>& Baker, FRenderCaptureTextures& TexturesOut)
+{
+	GetTexturesFromRenderCaptureBaker(*Baker, TexturesOut);
+}
+
 #undef LOCTEXT_NAMESPACE
diff --git a/Engine/Plugins/Runtime/MeshModelingToolset/Source/ModelingComponents/Public/Baking/RenderCaptureFunctions.h b/Engine/Plugins/Runtime/MeshModelingToolset/Source/ModelingComponents/Public/Baking/RenderCaptureFunctions.h
--- a/Engine/Plugins/Runtime/MeshModelingToolset/Source/ModelingComponents/Public/Baking/RenderCaptureFunctions.h
+++ b/Engine/Plugins/Runtime/MeshModelingToolset/Source/ModelingComponents/Public/Baking/RenderCaptureFunctions.h
@@ -4,6 +4,7 @@
 
 #include "Scene/SceneCapturePhotoSet.h"
 #include "Sampling/MeshBakerCommon.h"
+#include "Sampling/MeshMapBaker.h"
 #include "Image/ImageInfilling.h"
 #include "Baking/BakingTypes.h"
 #include "DynamicMesh/MeshTangents.h"
@@ -17,8 +18,6 @@ namespace UE
 namespace Geometry
 {
 
-class FMeshMapBaker;
-
 
 
 
@@ -109,21 +108,21 @@ struct MODELINGCOMPONENTS_API FSceneCaptureConfig
 
 MODELINGCOMPONENTS_API
 void ConfigureSceneCapture(
-	const TUniquePtr<FSceneCapturePhotoSet>& SceneCapture,
+	FSceneCapturePhotoSet& SceneCapture,
 	const TArray<TObjectPtr<AActor>>& Actors,
 	const FSceneCaptureConfig& Config,
 	bool bAllowCancel);
 
 MODELINGCOMPONENTS_API
 FRenderCaptureTypeFlags UpdateSceneCapture(
-	const TUniquePtr<FSceneCapturePhotoSet>& SceneCapture,
+	FSceneCapturePhotoSet& SceneCapture,
 	const TArray<TObjectPtr<AActor>>& Actors,
 	const FSceneCaptureConfig& DesiredConfig,
 	bool bAllowCancel);
 
 MODELINGCOMPONENTS_API
 FSceneCaptureConfig GetSceneCaptureConfig(
-	const TUniquePtr<FSceneCapturePhotoSet>& SceneCapture,
+	const FSceneCapturePhotoSet& SceneCapture,
 	FSceneCapturePhotoSet::ECaptureTypeStatus QueryStatus = FSceneCapturePhotoSet::ECaptureTypeStatus::Computed);
 
 // Return a render capture baker, note the lifetime of all arguments such match the lifetime of the returned baker
@@ -156,7 +155,7 @@ struct MODELINGCOMPONENTS_API FRenderCaptureTextures
 // Note: The source data in the textures is *not* updated by this function
 MODELINGCOMPONENTS_API
 void GetTexturesFromRenderCaptureBaker(
-	const TUniquePtr<FMeshMapBaker>& Baker,
+	FMeshMapBaker& Baker,
 	FRenderCaptureTextures& TexturesOut);
 
 
@@ -171,11 +170,40 @@ void GetTexturesFromRenderCaptureBaker(
 
 
 ////////////////////////////////////////////////////////////////////////////////////////////////////////////////////////
-// The following is all deprecated 
+// Everything that follows is deprecated 
 ////////////////////////////////////////////////////////////////////////////////////////////////////////////////////////
 
 PRAGMA_DISABLE_DEPRECATION_WARNINGS
 
+UE_DEPRECATED(5.4, "Please use the overload passing SceneCapture by reference. Using smart pointers was an error since ownership semantics were not intended.")
+MODELINGCOMPONENTS_API
+void ConfigureSceneCapture(
+	const TUniquePtr<FSceneCapturePhotoSet>& SceneCapture,
+	const TArray<TObjectPtr<AActor>>& Actors,
+	const FSceneCaptureConfig& Config,
+	bool bAllowCancel);
+
+UE_DEPRECATED(5.4, "Please use the overload passing SceneCapture by reference. Using smart pointers was an error since ownership semantics were not intended.")
+MODELINGCOMPONENTS_API
+FRenderCaptureTypeFlags UpdateSceneCapture(
+	const TUniquePtr<FSceneCapturePhotoSet>& SceneCapture,
+	const TArray<TObjectPtr<AActor>>& Actors,
+	const FSceneCaptureConfig& DesiredConfig,
+	bool bAllowCancel);
+
+UE_DEPRECATED(5.4, "Please use the overload passing SceneCapture by reference. Using smart pointers was an error since ownership semantics were not intended.")
+MODELINGCOMPONENTS_API
+FSceneCaptureConfig GetSceneCaptureConfig(
+	const TUniquePtr<FSceneCapturePhotoSet>& SceneCapture,
+	FSceneCapturePhotoSet::ECaptureTypeStatus QueryStatus = FSceneCapturePhotoSet::ECaptureTypeStatus::Computed);
+
+// Note: The source data in the textures is *not* updated by this function
+UE_DEPRECATED(5.4, "Please use the overload passing Baker by reference. Using smart pointers was an error since ownership semantics were not intended.")
+MODELINGCOMPONENTS_API
+void GetTexturesFromRenderCaptureBaker(
+	const TUniquePtr<FMeshMapBaker>& Baker,
+	FRenderCaptureTextures& TexturesOut);
+
 struct UE_DEPRECATED(5.3, "FRenderCaptureOptions is only used by deprecated functions, see those deprecation notes for more info.") MODELINGCOMPONENTS_API FRenderCaptureOptions
 {
 	//
