MeshModelingTools: Make group edge loop selection select boundary loops when clicking those.

diff --git a/Engine/Plugins/Runtime/MeshModelingToolset/Source/ModelingComponents/Private/Selection/GroupTopologySelector.cpp b/Engine/Plugins/Runtime/MeshModelingToolset/Source/ModelingComponents/Private/Selection/GroupTopologySelector.cpp
--- a/Engine/Plugins/Runtime/MeshModelingToolset/Source/ModelingComponents/Private/Selection/GroupTopologySelector.cpp
+++ b/Engine/Plugins/Runtime/MeshModelingToolset/Source/ModelingComponents/Private/Selection/GroupTopologySelector.cpp
@@ -29,7 +29,7 @@ void FGroupTopologyUtils::AddNewEdgeLoopEdgesFromCorner(int32 EdgeID, int32 Corn
 		const FGroupTopology::FGroupEdge& LastEdge = GroupTopology->Edges[LastEdgeID];
 		LastCornerID = LastEdge.EndpointCorners[0] == LastCornerID ? LastEdge.EndpointCorners[1] : LastEdge.EndpointCorners[0];
 
-		check(LastCornerID != IndexConstants::InvalidID);
+		ensure(LastCornerID != IndexConstants::InvalidID);
 	}
 }
 
@@ -47,6 +47,15 @@ bool FGroupTopologyUtils::GetNextEdgeLoopEdge(int32 IncomingEdgeID, int32 Corner
 		return false; // Not a valence 4 corner
 	}
 
+	// We want to stop when we hit boundary vertices because our valence-4 method of finding the next edge typically
+	// breaks down at these points. In some single-hole cases we could treat the hole as a group and find the next
+	// edge, but it doesn't always make sense to do that, it gets in the way of us using loop selection to select
+	// hole boundaries, and it breaks down with more than one hole. Best to just stop.
+	if (GroupTopology->GetMesh()->IsBoundaryVertex(GroupTopology->GetCornerVertexID(CornerID)))
+	{
+		return false;
+	}
+
 	const FGroupTopology::FGroupEdge& IncomingEdge = GroupTopology->Edges[IncomingEdgeID];
 
 	// We want to find the edge that shares this corner but does not border either of the neighboring groups of
@@ -81,6 +90,74 @@ bool FGroupTopologyUtils::GetNextEdgeLoopEdge(int32 IncomingEdgeID, int32 Corner
 	return false;
 }
 
+void FGroupTopologyUtils::AddNewBoundaryLoopEdges(int32 StartEdgeID, int32 ForwardCornerID, TSet<int32>& EdgeSet) const
+{
+	if (!GroupTopology->IsBoundaryEdge(StartEdgeID))
+	{
+		return;
+	}
+
+	int32 LastCornerID = ForwardCornerID;
+	int32 LastEdgeID = StartEdgeID;
+	while (true)
+	{
+		int32 NextEid;
+		if (!GetNextBoundaryLoopEdge(LastEdgeID, LastCornerID, NextEid))
+		{
+			break; // Hit a bowtie
+		}
+		if (EdgeSet.Contains(NextEid))
+		{
+			break; // Either we finished the loop, or we'll continue it from another selection
+		}
+
+		EdgeSet.Add(NextEid);
+
+		LastEdgeID = NextEid;
+		const FGroupTopology::FGroupEdge& LastEdge = GroupTopology->Edges[LastEdgeID];
+		LastCornerID = LastEdge.EndpointCorners[0] == LastCornerID ? LastEdge.EndpointCorners[1] : LastEdge.EndpointCorners[0];
+
+		ensure(LastCornerID != IndexConstants::InvalidID);
+	}
+}
+
+bool FGroupTopologyUtils::GetNextBoundaryLoopEdge(int32 IncomingEdgeID, int32 CornerID, int32& NextEdgeIDOut) const
+{
+	NextEdgeIDOut = IndexConstants::InvalidID;
+
+	if (!GroupTopology->IsBoundaryEdge(IncomingEdgeID))
+	{
+		return false;
+	}
+
+	TArray<int32> IncidentEdges;
+	GroupTopology->FindCornerNbrEdges(CornerID, IncidentEdges);
+	
+	for (int32 GroupEdgeID : IncidentEdges)
+	{
+		if (GroupTopology->IsBoundaryEdge(GroupEdgeID))
+		{
+			if (GroupEdgeID == IncomingEdgeID)
+			{
+				continue;
+			}
+			else if (NextEdgeIDOut == IndexConstants::InvalidID)
+			{
+				// Found a different boundary edge than IncomingEdgeID
+				NextEdgeIDOut = GroupEdgeID;
+			}
+			else
+			{
+				// Must be a bowtie, because we saw a third boundary edge
+				NextEdgeIDOut = IndexConstants::InvalidID;
+				return false;
+			}
+		}
+	}
+	
+	return ensure(NextEdgeIDOut != IndexConstants::InvalidID);
+}
+
 void FGroupTopologyUtils::AddNewEdgeRingEdges(int32 StartEdgeID, int32 ForwardGroupID, TSet<int32>& EdgeSet) const
 {
 	int32 CurrentEdgeID = StartEdgeID;
@@ -113,7 +190,10 @@ void FGroupTopologyUtils::AddNewEdgeRingEdges(int32 StartEdgeID, int32 ForwardGr
 bool FGroupTopologyUtils::GetQuadOppositeEdge(int32 EdgeIDIn, int32 GroupID, int32& OppositeEdgeIDOut) const
 {
 	const FGroupTopology::FGroup* Group = GroupTopology->FindGroupByID(GroupID);
-	check(Group);
+	if (!ensure(Group))
+	{
+		return false;
+	}
 
 	// Find the boundary that contains this edge
 	for (int32 i = 0; i < Group->Boundaries.Num(); ++i)
@@ -131,8 +211,7 @@ bool FGroupTopologyUtils::GetQuadOppositeEdge(int32 EdgeIDIn, int32 GroupID, int
 			return true;
 		}
 	}
-	check(false); // No boundary of the given group contained the given edge
-	return false;
+	return ensure(false); // No boundary of the given group contained the given edge
 }
 
 
@@ -151,7 +230,7 @@ FIndex2i FGroupTopologyUtils::GetEdgeGroups(int EdgeID) const
 
 bool FGroupTopologySelector::ExpandSelectionByEdgeLoops(FGroupTopologySelection& Selection)
 {
-	TSet<int32> EdgeSet(Selection.SelectedEdgeIDs);
+	TSet<int32>& EdgeSet = Selection.SelectedEdgeIDs;
 	int32 OriginalNumEdges = Selection.SelectedEdgeIDs.Num();
 	for (int32 Eid : Selection.SelectedEdgeIDs)
 	{
@@ -166,20 +245,38 @@ bool FGroupTopologySelector::ExpandSelectionByEdgeLoops(FGroupTopologySelection&
 		GroupTopologyUtils.AddNewEdgeLoopEdgesFromCorner(Eid, EndpointCorners[1], EdgeSet);
 	}
 
-	if (EdgeSet.Num() > OriginalNumEdges)
-	{
-		Selection.SelectedEdgeIDs.Append(EdgeSet);
-		return true;
-	}
-	else
+	return EdgeSet.Num() > OriginalNumEdges;
+}
+
+bool FGroupTopologySelector::ExpandSelectionByBoundaryLoops(FGroupTopologySelection& Selection)
+{
+	TSet<int32>& EdgeSet = Selection.SelectedEdgeIDs;
+	int32 OriginalNumEdges = Selection.SelectedEdgeIDs.Num();
+
+	for (int32 Eid : Selection.SelectedEdgeIDs)
 	{
-		return false;
+		if (!GroupTopologyUtils.GroupTopology->IsBoundaryEdge(Eid))
+		{
+			continue;
+		}
+
+		const FIndex2i EndpointCorners = GroupTopologyUtils.GetEdgeEndpointCorners(Eid);
+		if (EndpointCorners[0] == IndexConstants::InvalidID)
+		{
+			continue; // This FGroupEdge is a loop unto itself (and already in our selection, since we're looking at it).
+		}
+
+		// Go forward and backward adding edges
+		GroupTopologyUtils.AddNewBoundaryLoopEdges(Eid, EndpointCorners[0], EdgeSet);
+		GroupTopologyUtils.AddNewBoundaryLoopEdges(Eid, EndpointCorners[1], EdgeSet);
 	}
+
+	return EdgeSet.Num() > OriginalNumEdges;
 }
 
 bool FGroupTopologySelector::ExpandSelectionByEdgeRings(FGroupTopologySelection& Selection)
 {
-	TSet<int32> EdgeSet(Selection.SelectedEdgeIDs);
+	TSet<int32>& EdgeSet = Selection.SelectedEdgeIDs;
 	int32 OriginalNumEdges = Selection.SelectedEdgeIDs.Num();
 	for (int32 Eid : Selection.SelectedEdgeIDs)
 	{
@@ -196,15 +293,7 @@ bool FGroupTopologySelector::ExpandSelectionByEdgeRings(FGroupTopologySelection&
 		}
 	}
 
-	if (EdgeSet.Num() > OriginalNumEdges)
-	{
-		Selection.SelectedEdgeIDs.Append(EdgeSet);
-		return true;
-	}
-	else
-	{
-		return false;
-	}
+	return EdgeSet.Num() > OriginalNumEdges;
 }
 
 FGroupTopologySelector::FGroupTopologySelector(const FDynamicMesh3* Mesh, const FGroupTopology* Topology) : 
diff --git a/Engine/Plugins/Runtime/MeshModelingToolset/Source/ModelingComponents/Private/Selection/PolygonSelectionMechanic.cpp b/Engine/Plugins/Runtime/MeshModelingToolset/Source/ModelingComponents/Private/Selection/PolygonSelectionMechanic.cpp
--- a/Engine/Plugins/Runtime/MeshModelingToolset/Source/ModelingComponents/Private/Selection/PolygonSelectionMechanic.cpp
+++ b/Engine/Plugins/Runtime/MeshModelingToolset/Source/ModelingComponents/Private/Selection/PolygonSelectionMechanic.cpp
@@ -267,6 +267,7 @@ bool UPolygonSelectionMechanic::UpdateHighlight(const FRay& WorldRay)
 	if (HilightSelection.SelectedEdgeIDs.Num() > 0 && Properties->bSelectEdgeLoops && ShouldSelectEdgeLoopsFunc())
 	{
 		GroupTopoSelector->ExpandSelectionByEdgeLoops(HilightSelection);
+		GroupTopoSelector->ExpandSelectionByBoundaryLoops(HilightSelection);
 	}
 
 	// Don't hover highlight a selection that we already selected, because people didn't like that
@@ -346,6 +347,7 @@ bool UPolygonSelectionMechanic::UpdateSelection(const FRay& WorldRay, FVector3d&
 		if (Selection.SelectedEdgeIDs.Num() > 0 && Properties->bSelectEdgeLoops && ShouldSelectEdgeLoopsFunc())
 		{
 			GroupTopoSelector->ExpandSelectionByEdgeLoops(Selection);
+			GroupTopoSelector->ExpandSelectionByBoundaryLoops(Selection);
 		}
 	}
 
diff --git a/Engine/Plugins/Runtime/MeshModelingToolset/Source/ModelingComponents/Public/Selection/GroupTopologySelector.h b/Engine/Plugins/Runtime/MeshModelingToolset/Source/ModelingComponents/Public/Selection/GroupTopologySelector.h
--- a/Engine/Plugins/Runtime/MeshModelingToolset/Source/ModelingComponents/Public/Selection/GroupTopologySelector.h
+++ b/Engine/Plugins/Runtime/MeshModelingToolset/Source/ModelingComponents/Public/Selection/GroupTopologySelector.h
@@ -103,6 +103,8 @@ struct FGroupTopologyUtils
 {
 	bool GetNextEdgeLoopEdge(int32 IncomingEdgeID, int32 CornerID, int32& NextEdgeIDOut) const;
 	void AddNewEdgeLoopEdgesFromCorner(int32 EdgeID, int32 CornerID, TSet<int32>& EdgeSet) const;
+	void AddNewBoundaryLoopEdges(int32 StartEdgeID, int32 ForwardCornerID, TSet<int32>& EdgeSet) const;
+	bool GetNextBoundaryLoopEdge(int32 IncomingEdgeID, int32 CornerID, int32& NextEdgeIDOut) const;
 	void AddNewEdgeRingEdges(int32 StartEdgeID, int32 ForwardGroupID, TSet<int32>& EdgeSet) const;
 	bool GetQuadOppositeEdge(int32 EdgeIDIn, int32 GroupID, int32& OppositeEdgeIDOut) const;
 	FIndex2i GetEdgeEndpointCorners(int EdgeID) const;
@@ -148,6 +150,17 @@ public:
 	 */
 	bool ExpandSelectionByEdgeLoops(FGroupTopologySelection& Selection);
 
+	/**
+	 * Using the edges in the given selection as starting points, add any "boundary loops" containing the edges. A
+	 * boundary loop is just the sequence of edges that surrounds a hole in the mesh, or the outside of an open mesh.
+	 * The loop may not complete if it hits a bowtie vertex (a vertex with more than two incident boundary edges).
+	 *
+	 * @param Selection Selection to expand.
+	 * @returns true if selection was modified (i.e., were the already selected edges part of any boundary loops whose
+	 *  member edges were not yet all selected).
+	 */
+	bool ExpandSelectionByBoundaryLoops(FGroupTopologySelection& Selection);
+
 	/**
 	 * Using the edges in the given selection as starting points, add any "edge rings" containing the edges. An
 	 * edge ring is a sequence of edges that lie on opposite sides of quad-like faces, meaning faces that have four
diff --git a/Engine/Plugins/Runtime/MeshModelingToolset/Source/ModelingComponents/Public/Selection/MeshTopologySelectionMechanic.h b/Engine/Plugins/Runtime/MeshModelingToolset/Source/ModelingComponents/Public/Selection/MeshTopologySelectionMechanic.h
--- a/Engine/Plugins/Runtime/MeshModelingToolset/Source/ModelingComponents/Public/Selection/MeshTopologySelectionMechanic.h
+++ b/Engine/Plugins/Runtime/MeshModelingToolset/Source/ModelingComponents/Public/Selection/MeshTopologySelectionMechanic.h
@@ -40,7 +40,7 @@ public:
 	UPROPERTY(EditAnywhere, Category = SelectionFilter)
 	bool bSelectFaces = true;
 
-	/** When true, will select edge loops. Edge loops are paths along a string of valence-4 vertices. */
+	/** When true, will select edge loops. Edge loops are either paths through vertices with 4 edges, or boundaries of holes. */
 	UPROPERTY(EditAnywhere, Category = SelectionFilter)
 	bool bSelectEdgeLoops = false;
 
