stop primitives being placed behind the camera when user hovers/clicks the sky in any of the Add Box/Sphere/Cylinder/etc tools

diff --git a/Engine/Plugins/Runtime/MeshModelingToolset/Source/MeshModelingTools/Private/AddPrimitiveTool.cpp b/Engine/Plugins/Runtime/MeshModelingToolset/Source/MeshModelingTools/Private/AddPrimitiveTool.cpp
--- a/Engine/Plugins/Runtime/MeshModelingToolset/Source/MeshModelingTools/Private/AddPrimitiveTool.cpp
+++ b/Engine/Plugins/Runtime/MeshModelingToolset/Source/MeshModelingTools/Private/AddPrimitiveTool.cpp
@@ -260,10 +260,20 @@ void UAddPrimitiveTool::UpdatePreviewPosition(const FInputDeviceRay& DeviceClick
 	// hit position (temp)
 	bool bHit = false;
 
+	auto RaycastPlaneWithFallback = [](const FVector3d& Origin, const FVector3d& Direction, const FPlane& Plane, double FallbackDistance = 1000) -> FVector3d
+	{
+		double IntersectTime = FMath::RayPlaneIntersectionParam(Origin, Direction, Plane);
+		if (IntersectTime < 0 || !FMath::IsFinite(IntersectTime))
+		{
+			IntersectTime = FallbackDistance;
+		}
+		return Origin + Direction * IntersectTime;
+	};
+
 	FPlane DrawPlane(FVector::ZeroVector, FVector(0, 0, 1));
 	if (ShapeSettings->TargetSurface == EMakeMeshPlacementType::GroundPlane)
 	{
-		FVector3d DrawPlanePos = (FVector3d)FMath::RayPlaneIntersection(ClickPosWorldRay.Origin, ClickPosWorldRay.Direction, DrawPlane);
+		FVector3d DrawPlanePos = RaycastPlaneWithFallback(ClickPosWorldRay.Origin, ClickPosWorldRay.Direction, DrawPlane);
 		bHit = true;
 		ShapeFrame = FFrame3d(DrawPlanePos);
 	}
@@ -285,7 +295,7 @@ void UAddPrimitiveTool::UpdatePreviewPosition(const FInputDeviceRay& DeviceClick
 		else
 		{
 			// fall back to ground plane if we don't have a scene hit
-			FVector3d DrawPlanePos = (FVector3d)FMath::RayPlaneIntersection(ClickPosWorldRay.Origin, ClickPosWorldRay.Direction, DrawPlane);
+			FVector3d DrawPlanePos = RaycastPlaneWithFallback(ClickPosWorldRay.Origin, ClickPosWorldRay.Direction, DrawPlane);
 			bHit = true;
 			ShapeFrame = FFrame3d(DrawPlanePos);
 		}
