FWorldRenderCapture: Improve performance of pixel format conversions

diff --git a/Engine/Plugins/Runtime/MeshModelingToolset/Source/ModelingComponents/ModelingComponents.Build.cs b/Engine/Plugins/Runtime/MeshModelingToolset/Source/ModelingComponents/ModelingComponents.Build.cs
--- a/Engine/Plugins/Runtime/MeshModelingToolset/Source/ModelingComponents/ModelingComponents.Build.cs
+++ b/Engine/Plugins/Runtime/MeshModelingToolset/Source/ModelingComponents/ModelingComponents.Build.cs
@@ -35,6 +35,7 @@ public class ModelingComponents : ModuleRules
 				"ImageWriteQueue",
 				"SkeletalMeshDescription", // FSkeletalMeshAttributes::DefaultSkinWeightProfileName
 				"SlateCore",
+				"ImageCore",
 				// ... add private dependencies that you statically link with here ...
 			}
 		);
diff --git a/Engine/Plugins/Runtime/MeshModelingToolset/Source/ModelingComponents/Private/Scene/SceneCapturePhotoSet.cpp b/Engine/Plugins/Runtime/MeshModelingToolset/Source/ModelingComponents/Private/Scene/SceneCapturePhotoSet.cpp
--- a/Engine/Plugins/Runtime/MeshModelingToolset/Source/ModelingComponents/Private/Scene/SceneCapturePhotoSet.cpp
+++ b/Engine/Plugins/Runtime/MeshModelingToolset/Source/ModelingComponents/Private/Scene/SceneCapturePhotoSet.cpp
@@ -142,6 +142,8 @@ FSceneCapturePhotoSet::FStatus FSceneCapturePhotoSet::GetSceneCaptureStatus() co
 
 void FSceneCapturePhotoSet::Compute()
 {
+	TRACE_CPUPROFILER_EVENT_SCOPE(SceneCapturePhotoSet::Compute);
+
 	bWasCancelled = false;
 
 	int NumPending = 0;
@@ -155,7 +157,6 @@ void FSceneCapturePhotoSet::Compute()
 		return;
 	}
 
-	TRACE_CPUPROFILER_EVENT_SCOPE(SceneCapturePhotoSet::Compute);
 	FScopedSlowTask Progress(static_cast<float>(NumPending), LOCTEXT("CapturingScene", "Capturing Scene..."));
 	Progress.MakeDialog(bAllowCancel);
 
@@ -189,10 +190,7 @@ void FSceneCapturePhotoSet::Compute()
 	FWorldRenderCapture RenderCapture;
 	RenderCapture.SetWorld(TargetWorld);
 	RenderCapture.SetVisibleActorsAndComponents(VisibleActors, VisibleComponents);
-	if (bWriteDebugImages)
-	{
-		RenderCapture.SetEnableWriteDebugImage(true, 0, DebugImagesFolderName);
-	}
+	RenderCapture.SetEnableWriteDebugImage(bWriteDebugImages, 0, DebugImagesFolderName);
 
 	auto CapturePhoto3f = [this, &RenderCapture](ERenderCaptureType CaptureType, const FSpatialPhotoParams& Params)
 	{
@@ -262,16 +260,55 @@ void FSceneCapturePhotoSet::Compute()
 		}
 	};
 
-	CapturePhotoSet(CapturePhoto1f, ERenderCaptureType::DeviceDepth, LOCTEXT("CapturingScene_DeviceDepth", "Capturing Device Depth"));
-	CapturePhotoSet(CapturePhoto3f, ERenderCaptureType::BaseColor, LOCTEXT("CapturingScene_BaseColor", "Capturing Base Color"));
-	CapturePhotoSet(CapturePhoto3f, ERenderCaptureType::WorldNormal, LOCTEXT("CapturingScene_WorldNormal", "Capturing World Normal"));
-	CapturePhotoSet(CapturePhoto3f, ERenderCaptureType::CombinedMRS, LOCTEXT("CapturingScene_CombinedMRS", "Capturing Packed MRS"));
-	CapturePhotoSet(CapturePhoto1f, ERenderCaptureType::Metallic, LOCTEXT("CapturingScene_Metallic", "Capturing Metallic"));
-	CapturePhotoSet(CapturePhoto1f, ERenderCaptureType::Roughness, LOCTEXT("CapturingScene_Roughness", "Capturing Roughness"));
-	CapturePhotoSet(CapturePhoto1f, ERenderCaptureType::Specular, LOCTEXT("CapturingScene_Specular", "Capturing Specular"));
-	CapturePhotoSet(CapturePhoto3f, ERenderCaptureType::Emissive, LOCTEXT("CapturingScene_Emissive", "Capturing Emissive"));
-	CapturePhotoSet(CapturePhoto1f, ERenderCaptureType::Opacity, LOCTEXT("CapturingScene_Opacity", "Capturing Opacity"));
-	CapturePhotoSet(CapturePhoto3f, ERenderCaptureType::SubsurfaceColor, LOCTEXT("CapturingScene_SubsurfaceColor", "Capturing Subsurface Color"));
+	{
+		TRACE_CPUPROFILER_EVENT_SCOPE(FSceneCapturePhotoSet::Compute_DeviceDepth);
+		CapturePhotoSet(CapturePhoto1f, ERenderCaptureType::DeviceDepth, LOCTEXT("CapturingScene_DeviceDepth", "Capturing Device Depth"));
+	}
+
+	{
+		TRACE_CPUPROFILER_EVENT_SCOPE(FSceneCapturePhotoSet::Compute_BaseColor);
+		CapturePhotoSet(CapturePhoto3f, ERenderCaptureType::BaseColor, LOCTEXT("CapturingScene_BaseColor", "Capturing Base Color"));
+	}
+
+	{
+		TRACE_CPUPROFILER_EVENT_SCOPE(FSceneCapturePhotoSet::Compute_WorldNormal);
+		CapturePhotoSet(CapturePhoto3f, ERenderCaptureType::WorldNormal, LOCTEXT("CapturingScene_WorldNormal", "Capturing World Normal"));
+	}
+	
+	{
+		TRACE_CPUPROFILER_EVENT_SCOPE(FSceneCapturePhotoSet::Compute_CombinedMRS);
+		CapturePhotoSet(CapturePhoto3f, ERenderCaptureType::CombinedMRS, LOCTEXT("CapturingScene_CombinedMRS", "Capturing Packed MRS"));
+	}
+	
+	{
+		TRACE_CPUPROFILER_EVENT_SCOPE(FSceneCapturePhotoSet::Compute_Metallic);
+		CapturePhotoSet(CapturePhoto1f, ERenderCaptureType::Metallic, LOCTEXT("CapturingScene_Metallic", "Capturing Metallic"));
+	}
+
+	{
+		TRACE_CPUPROFILER_EVENT_SCOPE(FSceneCapturePhotoSet::Compute_Roughness);
+		CapturePhotoSet(CapturePhoto1f, ERenderCaptureType::Roughness, LOCTEXT("CapturingScene_Roughness", "Capturing Roughness"));
+	}
+
+	{
+		TRACE_CPUPROFILER_EVENT_SCOPE(FSceneCapturePhotoSet::Compute_Specular);
+		CapturePhotoSet(CapturePhoto1f, ERenderCaptureType::Specular, LOCTEXT("CapturingScene_Specular", "Capturing Specular"));
+	}
+
+	{
+		TRACE_CPUPROFILER_EVENT_SCOPE(FSceneCapturePhotoSet::Compute_Emissive);
+		CapturePhotoSet(CapturePhoto3f, ERenderCaptureType::Emissive, LOCTEXT("CapturingScene_Emissive", "Capturing Emissive"));
+	}
+
+	{
+		TRACE_CPUPROFILER_EVENT_SCOPE(FSceneCapturePhotoSet::Compute_Opacity);
+		CapturePhotoSet(CapturePhoto1f, ERenderCaptureType::Opacity, LOCTEXT("CapturingScene_Opacity", "Capturing Opacity"));
+	}
+
+	{
+		TRACE_CPUPROFILER_EVENT_SCOPE(FSceneCapturePhotoSet::Compute_SubsurfaceColor);
+		CapturePhotoSet(CapturePhoto3f, ERenderCaptureType::SubsurfaceColor, LOCTEXT("CapturingScene_SubsurfaceColor", "Capturing Subsurface Color"));
+	}
 }
 
 
diff --git a/Engine/Plugins/Runtime/MeshModelingToolset/Source/ModelingComponents/Private/Scene/WorldRenderCapture.cpp b/Engine/Plugins/Runtime/MeshModelingToolset/Source/ModelingComponents/Private/Scene/WorldRenderCapture.cpp
--- a/Engine/Plugins/Runtime/MeshModelingToolset/Source/ModelingComponents/Private/Scene/WorldRenderCapture.cpp
+++ b/Engine/Plugins/Runtime/MeshModelingToolset/Source/ModelingComponents/Private/Scene/WorldRenderCapture.cpp
@@ -14,6 +14,8 @@
 #include "Rendering/Texture2DResource.h"
 #include "ImagePixelData.h"
 #include "ImageWriteStream.h"
+#include "ImageCore.h"
+#include "Async/ParallelFor.h"
 
 #include "PreviewScene.h"
 #include "Components/DirectionalLightComponent.h"
@@ -416,77 +418,45 @@ static void PerformSceneRender(FCanvas* Canvas, FSceneViewFamily* ViewFamily)
 	} while (bCompiledAssets);
 }
 
-
-/**
- * Reads data from FImagePixelData and stores in an output image, with optional ColorTransformFunc
- */
-static bool ReadPixelDataToImage(TUniquePtr<FImagePixelData>& PixelData, FImageAdapter& ResultImageOut, bool bLinear)
+// Copied from ImageCore.cpp
+template <typename Lambda>
+static void ParallelLoop(const TCHAR* DebugName, int32 NumJobs, int64 TexelsPerJob, int64 NumTexels, const Lambda& Func)
 {
-	TRACE_CPUPROFILER_EVENT_SCOPE(Internal::ReadPixelDataToImage);
-
-	int64 SizeInBytes;
-	const void* OutRawData = nullptr;
-	PixelData->GetRawData(OutRawData, SizeInBytes);
-	int32 Width = PixelData->GetSize().X;
-	int32 Height = PixelData->GetSize().Y;
-
-	ResultImageOut.SetDimensions(FImageDimensions(Width, Height));
-
-	switch (PixelData->GetType())
+	ParallelFor(DebugName, NumJobs, 1, [=](int64 JobIndex)
 	{
-		case EImagePixelType::Color:
+		const int64 StartIndex = JobIndex * TexelsPerJob;
+		const int64 EndIndex = FMath::Min(StartIndex + TexelsPerJob, NumTexels);
+		for (int64 TexelIndex = StartIndex; TexelIndex < EndIndex; ++TexelIndex)
 		{
-			const uint8* SourceBuffer = (uint8*)OutRawData;
-			for (int32 yi = 0; yi < Height; ++yi)
-			{
-				for (int32 xi = 0; xi < Width; ++xi)
-				{
-					const uint8* BufferPixel = &SourceBuffer[(yi * Width + xi) * 4];
-					FColor Color(BufferPixel[2], BufferPixel[1], BufferPixel[0], BufferPixel[3]);		// BGRA
-					FLinearColor PixelColorf = (bLinear) ? Color.ReinterpretAsLinear() : FLinearColor(Color);
-					ResultImageOut.SetPixel(FVector2i(xi, yi), PixelColorf);
-				}
-			}
-			return true;
+			Func(TexelIndex);
 		}
+	}, EParallelForFlags::Unbalanced);
+}
 
-		case EImagePixelType::Float16:
-		{
-			ensure(bLinear);		// data must be linear
-			const FFloat16Color* SourceBuffer = (FFloat16Color*)(OutRawData);
-			for (int32 yi = 0; yi < Height; ++yi)
-			{
-				for (int32 xi = 0; xi < Width; ++xi)
-				{
-					const FFloat16Color* BufferPixel = &SourceBuffer[yi * Width + xi];
-					FLinearColor PixelColorf = BufferPixel->GetFloats();
-					ResultImageOut.SetPixel(FVector2i(xi, yi), PixelColorf);
-				}
-			}
-			return true;
-		}
+// This function is needed because we need to convert pixels from the GPU read back format to the format we prefer in the calling code
+// e.g., we may readback linear colors but prefer to store only 3 components to save memory
+static void CopyCaptureDataToOutputImageFormat(
+	const FImageView& CaptureData,
+	FImageAdapter& OutputImage,
+	const TFunctionRef<FLinearColor(const FLinearColor&)>& Transform = [](const FLinearColor Color) { return Color; })
+{
+	TRACE_CPUPROFILER_EVENT_SCOPE(CopyCaptureDataToOutputImageFormat);
 
-		case EImagePixelType::Float32:
-		{
-			ensure(bLinear);		// data must be linear
-			const FLinearColor* SourceBuffer = (FLinearColor*)(OutRawData);
-			for (int32 yi = 0; yi < Height; ++yi)
-			{
-				for (int32 xi = 0; xi < Width; ++xi)
-				{
-					FLinearColor PixelColorf = SourceBuffer[yi * Width + xi];
-					ResultImageOut.SetPixel(FVector2i(xi, yi), PixelColorf);
-				}
-			}
-			return true;
-		}
+	const FImageDimensions Dims(CaptureData.GetWidth(), CaptureData.GetHeight());
+	OutputImage.SetDimensions(Dims);
 
-		default:
-		{
-			ensure(false);
-			return false;
-		}
-	}
+	const int64 NumTexels = CaptureData.GetNumPixels();
+	int64 TexelsPerJob;
+	const int32 NumJobs = ImageParallelForComputeNumJobsForPixels(TexelsPerJob, NumTexels);
+	ParallelLoop(TEXT("PF.CopyCaptureDataToOutputImageFormat"), NumJobs, TexelsPerJob, NumTexels,
+		[&CaptureData, &OutputImage, &Transform, Dims](int64 TexelIndex)
+	{
+		const FVector2i Coords = Dims.GetCoords(TexelIndex);
+		const FLinearColor Color = CaptureData.GetOnePixelLinear(Coords.X, Coords.Y);
+
+		// This call will store a subset of the components of the (transformed) color, depending on OutputImage.ImageType
+		OutputImage.SetPixel(Coords, Transform(Color));
+	});
 }
 
 // The .ViewFamily member should be set on the return value later
@@ -563,6 +533,32 @@ void SetCommonShowFlags(FEngineShowFlags& ShowFlags, bool bAntiAliasing)
 	ShowFlags.SetTranslucency(false);
 }
 
+FName GetBufferVisualizationModeName(ERenderCaptureType CaptureType)
+{
+	switch (CaptureType)
+	{
+	case ERenderCaptureType::BaseColor:
+		return FName("BaseColor");
+	case ERenderCaptureType::Roughness:
+		return FName("Roughness");
+	case ERenderCaptureType::Metallic:
+		return FName("Metallic");
+	case ERenderCaptureType::Specular:
+		return FName("Specular");
+	case ERenderCaptureType::WorldNormal:
+		return FName("WorldNormal");
+	case ERenderCaptureType::DeviceDepth:
+		return FName("DeviceDepth");
+	case ERenderCaptureType::Opacity:
+		return FName("Opacity");
+	case ERenderCaptureType::SubsurfaceColor:
+		return FName("SubsurfaceColor");
+	default:
+		ensure(false);
+	}
+	return FName("BaseColor");
+}
+
 } // end namespace Internal
 } // end namespace UE
 
@@ -574,7 +570,7 @@ bool FWorldRenderCapture::CaptureMRSFromPosition(
 	FImageAdapter& ResultImageOut,
 	const FRenderCaptureConfig& Config)
 {
-	TRACE_CPUPROFILER_EVENT_SCOPE(FWorldRenderCapture::CaptureMRSFromPosition);
+	TRACE_CPUPROFILER_EVENT_SCOPE(FWorldRenderCapture_CombinedMRS);
 
 	// this post-process material renders an image with R=Metallic, G=Roughness, B=Specular, A=AmbientOcclusion
 	FString MRSPostProcessMaterialAssetPath = TEXT("/MeshModelingToolsetExp/Materials/PostProcess_PackedMRSA.PostProcess_PackedMRSA");
@@ -659,13 +655,19 @@ bool FWorldRenderCapture::CaptureMRSFromPosition(
 	// Cache the view/projection matricies we used to render the scene
 	LastCaptureViewMatrices = NewView->ViewMatrices;
 
-	if (PostProcessPassPixelData.IsValid())
+	if (ensure(PostProcessPassPixelData.IsValid()) == false)
+	{
+		return false;
+	}
+
+	UE::Internal::CopyCaptureDataToOutputImageFormat(PostProcessPassPixelData->GetImageView(), ResultImageOut);
+
+	if (bWriteDebugImage)
 	{
-		UE::Internal::ReadPixelDataToImage(PostProcessPassPixelData, ResultImageOut, true );
-		return true;
+		WriteDebugImage(ResultImageOut, TEXT("CombinedMRS"));
 	}
 
-	return false;
+	return true;
 }
 
 
@@ -679,7 +681,7 @@ bool FWorldRenderCapture::CaptureEmissiveFromPosition(
 	FImageAdapter& ResultImageOut,
 	const FRenderCaptureConfig& Config)
 {
-	TRACE_CPUPROFILER_EVENT_SCOPE(FWorldRenderCapture::CaptureEmissiveFromPosition);
+	TRACE_CPUPROFILER_EVENT_SCOPE(FWorldRenderCapture_Emissive);
 
 	UTextureRenderTarget2D* RenderTargetTexture = GetRenderTexture(true);
 	if (ensure(RenderTargetTexture) == false)
@@ -732,18 +734,12 @@ bool FWorldRenderCapture::CaptureEmissiveFromPosition(
 	ReadSurfaceDataFlags.SetLinearToGamma(false);
 	RenderTargetResource->ReadLinearColorPixels(ReadImageBuffer, ReadSurfaceDataFlags, FIntRect(0, 0, Width, Height));
 
-	{
-		TRACE_CPUPROFILER_EVENT_SCOPE(FWorldRenderCapture::CaptureDeviceDepthFromPosition_CopyBufferToImage);
+	FImageView CaptureData(ReadImageBuffer.GetData(), Width, Height);
+	UE::Internal::CopyCaptureDataToOutputImageFormat(CaptureData, ResultImageOut);
 
-		ResultImageOut.SetDimensions(Dimensions);
-		for (int32 yi = 0; yi < Height; ++yi)
-		{
-			for (int32 xi = 0; xi < Width; ++xi)
-			{
-				FLinearColor PixelColorf = ReadImageBuffer[yi * Width + xi];
-				ResultImageOut.SetPixel(FVector2i(xi, yi), PixelColorf);
-			}
-		}
+	if (bWriteDebugImage)
+	{
+		WriteDebugImage(ResultImageOut, TEXT("Emissive"));
 	}
 
 	return true;
@@ -757,7 +753,7 @@ bool FWorldRenderCapture::CaptureDeviceDepthFromPosition(
 	FImageAdapter& ResultImageOut,
 	const FRenderCaptureConfig& Config)
 {
-	TRACE_CPUPROFILER_EVENT_SCOPE(FWorldRenderCapture::CaptureDeviceDepthFromPosition);
+	TRACE_CPUPROFILER_EVENT_SCOPE(FWorldRenderCapture_DeviceDepth);
 
 	DepthRenderTexture = GetDepthRenderTexture();
 	if (ensure(DepthRenderTexture) == false)
@@ -824,37 +820,28 @@ bool FWorldRenderCapture::CaptureDeviceDepthFromPosition(
 	LastCaptureViewMatrices = NewView->ViewMatrices;
 
 	// Copy the contents of the remote texture to system memory
-	TArray<FLinearColor> ReadImageColorBuffer;
-	ReadImageColorBuffer.SetNumUninitialized(Width * Height);
+	ReadImageBuffer.Reset();
+	ReadImageBuffer.SetNumUninitialized(Width * Height);
 	FReadSurfaceDataFlags ReadSurfaceDataFlags(RCM_MinMax);
 	ReadSurfaceDataFlags.SetLinearToGamma(false);
 
-	RenderTargetResource->ReadLinearColorPixels(ReadImageColorBuffer, ReadSurfaceDataFlags, FIntRect(0, 0, Width, Height));
+	RenderTargetResource->ReadLinearColorPixels(ReadImageBuffer, ReadSurfaceDataFlags, FIntRect(0, 0, Width, Height));
 
+	FImageView CaptureData(ReadImageBuffer.GetData(), Width, Height);
+	UE::Internal::CopyCaptureDataToOutputImageFormat(CaptureData, ResultImageOut, [](const FLinearColor& Color)
 	{
-		TRACE_CPUPROFILER_EVENT_SCOPE(FWorldRenderCapture::CaptureDeviceDepthFromPosition_CopyBufferToImage);
-
-		ResultImageOut.SetDimensions(Dimensions);
-		for (int32 yi = 0; yi < Height; ++yi)
-		{
-			for (int32 xi = 0; xi < Width; ++xi)
-			{
-				FLinearColor Color = ReadImageColorBuffer[yi * Width + xi];
-
-				// Reverse the float encoding used for the Depth value in SceneCapturePixelShader.usf (one minus the device
-				// z is encoded in the RGB components). See DecodeFloatRGB in WaterInfoMerge.usf and, for a reference on the
-				// method, see https://aras-p.info/blog/2009/07/30/encoding-floats-to-rgba-the-final/
-				FVector3f EncodedDepth(Color.R, Color.G, Color.B);
-				float Depth = FVector3f::DotProduct(EncodedDepth, FVector3f(1.0, 1 / 255.0, 1 / 65025.0));
+		// Reverse the float encoding used for the Depth value in SceneCapturePixelShader.usf (one minus the device
+		// z is encoded in the RGB components). See DecodeFloatRGB in WaterInfoMerge.usf and, for a reference on the
+		// method, see https://aras-p.info/blog/2009/07/30/encoding-floats-to-rgba-the-final/
+		FVector3f EncodedDepth(Color.R, Color.G, Color.B);
+		float Depth = FVector3f::DotProduct(EncodedDepth, FVector3f(1.0, 1 / 255.0, 1 / 65025.0));
 
-				// Reverse the expression used to compute the Depth to recover the DeviceZ aka normalized device coordinate z
-				float DeviceZ = -(Depth - 1);
-				ensure(DeviceZ >= 0. && DeviceZ <= 1.); // Points on the near plane have Z=1, points on the far plane have Z=0
+		// Reverse the expression used to compute the Depth to recover the DeviceZ aka normalized device coordinate z
+		float DeviceZ = -(Depth - 1);
+		ensure(DeviceZ >= 0. && DeviceZ <= 1.); // Points on the near plane have Z=1, points on the far plane have Z=0
 
-				ResultImageOut.SetPixel(FVector2i(xi, yi), FLinearColor(DeviceZ, 0., 0., 0.));
-			}
-		}
-	}
+		return FLinearColor(DeviceZ, 0., 0., 0.);
+	});
 
 	// Set this to true to compute a world point cloud for debugging.
 	// You probably want to change the logging so it writes an .obj file
@@ -885,79 +872,98 @@ bool FWorldRenderCapture::CaptureDeviceDepthFromPosition(
 		} // yi
 	}
 
+	if (bWriteDebugImage)
+	{
+		WriteDebugImage(ResultImageOut, TEXT("DeviceDepth"));
+	}
+
 	return true;
 }
 
 
 
-namespace UE
-{
-namespace Internal
+bool FWorldRenderCapture::CaptureBufferVisualizationFromPosition(
+	const FName& VisualizationMode,
+	const FFrame3d& Frame,
+	double HorzFOVDegrees,
+	double NearPlaneDist,
+	FImageAdapter& ResultImageOut,
+	const FRenderCaptureConfig& Config)
 {
+	TRACE_CPUPROFILER_EVENT_SCOPE_TEXT(*FString::Printf(TEXT("FWorldRenderCapture_%s"), *VisualizationMode.ToString()));
 
-	/**
-	 * internal Utility function to render the given Scene to a render target and capture
-	 * one of the render buffers, defined by VisualizationMode.
-	 *
-	 * Not clear where the valid VisualizationMode FNames are defined, possibly this list:
-	 * "BaseColor,Specular,SubsurfaceColor,WorldNormal,SeparateTranslucencyRGB,,,WorldTangent,SeparateTranslucencyA,,,Opacity,SceneDepth,Roughness,Metallic,ShadingModel,,SceneDepthWorldUnits,SceneColor,PreTonemapHDRColor,PostTonemapHDRColor"
-	 * or possibly in the [Engine.BufferVisualizationMaterials] section of Config/BaseEngine.ini, which is read by FBufferVisualizationData::Initialize() and called when the engine initializes in UEngine::Init()
-	 */
-	static void RenderSceneVisualizationToTexture(
-		UTextureRenderTarget2D* RenderTargetTexture,
-		FImageDimensions Dimensions,
-		FSceneInterface* Scene,
-		const FName& VisualizationMode,
-		FSceneViewInitOptions ViewInitOptions,
-		const FRenderCaptureConfig& Config,
-		TArray<FLinearColor>& OutSamples,
-		FViewMatrices& LastCaptureViewMatrices
-	)
+	// The following handles buffer visualization materials found in /Engine/BufferVisualization/<VisualizationMode>.
+	// Sometimes these postprocess materials change the raw GBuffer data, in particular:
+	// - The Roughness postprocess material output is:           GBufferRoughness^Gamma,   with Gamma=2.2
+	// - The SubsurfaceColor postprocess material output is:     GBufferSSColor^(1/Gamma), with Gamma=2.2
+	// Note the relationship between linear and gamma encoded values: EncodedValue = LinearValue^(1/Gamma)
+
+	// Roughness visualization is rendered with gamma correction (unclear why)
+	bool bLinear = (VisualizationMode != Internal::GetBufferVisualizationModeName(ERenderCaptureType::Roughness));
+	UTextureRenderTarget2D* RenderTargetTexture = GetRenderTexture(bLinear);
+	if (ensure(RenderTargetTexture) == false)
 	{
-		TRACE_CPUPROFILER_EVENT_SCOPE(Internal::RenderSceneVisualizationToTexture);
+		return false;
+	}
 
-		int32 Width = Dimensions.GetWidth();
-		int32 Height = Dimensions.GetHeight();
-		FTextureRenderTargetResource* RenderTargetResource = RenderTargetTexture->GameThread_GetRenderTargetResource();
+	FSceneViewInitOptions ViewInitOptions = Internal::MakeSceneViewInitOptions(
+		Frame,
+		Dimensions,
+		HorzFOVDegrees,
+		NearPlaneDist,
+		VisiblePrimitives);
 
-		FEngineShowFlags ShowFlags(ESFIM_Game);
-		ApplyViewMode(VMI_VisualizeBuffer, true, ShowFlags);
-		Internal::SetCommonShowFlags(ShowFlags, Config.bAntiAliasing);
-		ShowFlags.SetPostProcessMaterial(true);
-		ShowFlags.SetPostProcessing(true);
-		ShowFlags.SetVisualizeBuffer(true);
+	int32 Width = Dimensions.GetWidth();
+	int32 Height = Dimensions.GetHeight();
+	FTextureRenderTargetResource* RenderTargetResource = RenderTargetTexture->GameThread_GetRenderTargetResource();
 
-		FSceneViewFamilyContext ViewFamily(
-			FSceneViewFamily::ConstructionValues(RenderTargetResource, Scene, ShowFlags)
-			.SetTime(FGameTime())
-			.SetRealtimeUpdate(false)
-		);
+	FEngineShowFlags ShowFlags(ESFIM_Game);
+	ApplyViewMode(VMI_VisualizeBuffer, true, ShowFlags);
+	Internal::SetCommonShowFlags(ShowFlags, Config.bAntiAliasing);
+	ShowFlags.SetPostProcessMaterial(true);
+	ShowFlags.SetPostProcessing(true);
+	ShowFlags.SetVisualizeBuffer(true);
 
-		ViewInitOptions.ViewFamily = &ViewFamily;
+	FSceneViewFamilyContext ViewFamily(
+		FSceneViewFamily::ConstructionValues(RenderTargetResource, World->Scene, ShowFlags)
+		.SetTime(FGameTime())
+		.SetRealtimeUpdate(false)
+	);
 
-		FSceneView* NewView = new FSceneView(ViewInitOptions);
-		NewView->CurrentBufferVisualizationMode = VisualizationMode;
-		ViewFamily.Views.Add(NewView);
+	ViewInitOptions.ViewFamily = &ViewFamily;
 
-		ViewFamily.SetScreenPercentageInterface(new FLegacyScreenPercentageDriver(ViewFamily, 1.0f));
+	FSceneView* NewView = new FSceneView(ViewInitOptions);
+	NewView->CurrentBufferVisualizationMode = VisualizationMode;
+	ViewFamily.Views.Add(NewView);
 
-		// should we cache the FCanvas?
-		FCanvas Canvas(RenderTargetResource, nullptr, FGameTime(), NewView->GetFeatureLevel());
-		Canvas.Clear(FLinearColor::Transparent);
+	ViewFamily.SetScreenPercentageInterface(new FLegacyScreenPercentageDriver(ViewFamily, 1.0f));
 
-		UE::Internal::PerformSceneRender(&Canvas, &ViewFamily);
+	// should we cache the FCanvas?
+	FCanvas Canvas(RenderTargetResource, nullptr, FGameTime(), NewView->GetFeatureLevel());
+	Canvas.Clear(FLinearColor::Transparent);
 
-		// Cache the view/projection matricies we used to render the scene
-		LastCaptureViewMatrices = NewView->ViewMatrices;
+	UE::Internal::PerformSceneRender(&Canvas, &ViewFamily);
 
-		// Copy the contents of the remote texture to system memory
-		OutSamples.SetNumUninitialized(Width * Height);
-		//FReadSurfaceDataFlags ReadSurfaceDataFlags(RCM_MinMax);		// should we use MinMax to avoid normalization?
-		FReadSurfaceDataFlags ReadSurfaceDataFlags;
-		ReadSurfaceDataFlags.SetLinearToGamma(false);
-		RenderTargetResource->ReadLinearColorPixels(OutSamples, ReadSurfaceDataFlags, FIntRect(0, 0, Width, Height));
+	// Cache the view/projection matricies we used to render the scene
+	LastCaptureViewMatrices = NewView->ViewMatrices;
+
+	// Copy the contents of the remote texture to system memory
+	ReadImageBuffer.Reset();
+	ReadImageBuffer.SetNumUninitialized(Width * Height);
+	//FReadSurfaceDataFlags ReadSurfaceDataFlags(RCM_MinMax);		// should we use MinMax to avoid normalization?
+	FReadSurfaceDataFlags ReadSurfaceDataFlags;
+	ReadSurfaceDataFlags.SetLinearToGamma(false);
+	RenderTargetResource->ReadLinearColorPixels(ReadImageBuffer, ReadSurfaceDataFlags, FIntRect(0, 0, Width, Height));
+
+	FImageView CaptureData(ReadImageBuffer.GetData(), Dimensions.GetWidth(), Dimensions.GetHeight());
+	UE::Internal::CopyCaptureDataToOutputImageFormat(CaptureData, ResultImageOut);
+
+	if (bWriteDebugImage)
+	{
+		WriteDebugImage(ResultImageOut, *VisualizationMode.ToString());
 	}
-}
+
+	return true;
 }
 
 
@@ -977,122 +983,30 @@ bool FWorldRenderCapture::CaptureFromPosition(
 	FImageAdapter& ResultImageOut,
 	const FRenderCaptureConfig& Config)
 {
-	TRACE_CPUPROFILER_EVENT_SCOPE(WorldRenderCapture::CaptureFromPosition);
-
 	RenderCaptureInterface::FScopedCapture RenderCapture(RenderCaptureDraws > 0, TEXT("RenderCaptureFromPosition"));
 	RenderCaptureDraws--;
 
+	bool bCaptured = false;
+
 	if (CaptureType == ERenderCaptureType::Emissive)
 	{
-		bool bOK = CaptureEmissiveFromPosition(ViewFrame, HorzFOVDegrees, NearPlaneDist, ResultImageOut, Config);
-		if (bWriteDebugImage)
-		{
-			WriteDebugImage(ResultImageOut, TEXT("Emissive"));
-		}
-		return bOK;
+		bCaptured = CaptureEmissiveFromPosition(ViewFrame, HorzFOVDegrees, NearPlaneDist, ResultImageOut, Config);
 	}
 	else if (CaptureType == ERenderCaptureType::CombinedMRS)
 	{
-		bool bOK = CaptureMRSFromPosition(ViewFrame, HorzFOVDegrees, NearPlaneDist, ResultImageOut, Config);
-		if (bWriteDebugImage)
-		{
-			WriteDebugImage(ResultImageOut, TEXT("CombinedMRS"));
-		}
-		return bOK;
+		bCaptured = CaptureMRSFromPosition(ViewFrame, HorzFOVDegrees, NearPlaneDist, ResultImageOut, Config);
 	}
 	else if (CaptureType == ERenderCaptureType::DeviceDepth)
 	{
-		
-		bool bOK = CaptureDeviceDepthFromPosition(ViewFrame, HorzFOVDegrees, NearPlaneDist, ResultImageOut, Config);
-		if (bWriteDebugImage)
-		{
-			WriteDebugImage(ResultImageOut, TEXT("DeviceDepth"));
-		}
-		return bOK;
+		bCaptured = CaptureDeviceDepthFromPosition(ViewFrame, HorzFOVDegrees, NearPlaneDist, ResultImageOut, Config);
 	}
-
-	TRACE_CPUPROFILER_EVENT_SCOPE(FWorldRenderCapture::CaptureFromPosition_CaptureBufferVisualizationFromPosition);
-
-	// The following handles buffer visualization materials found in /Engine/BufferVisualization/<CaptureTypeName>.
-	// Sometimes these postprocess materials change the raw GBuffer data, in particular:
-	// - The Roughness postprocess material output is:           GBufferRoughness^Gamma,   with Gamma=2.2
-	// - The SubsurfaceColor postprocess material output is:     GBufferSSColor^(1/Gamma), with Gamma=2.2
-	// Note the relationship between linear and gamma encoded values: EncodedValue = LinearValue^(1/Gamma)
-
-	// Roughness visualization is rendered with gamma correction (unclear why)
-	bool bLinear = (CaptureType != ERenderCaptureType::Roughness);
-	UTextureRenderTarget2D* RenderTargetTexture = GetRenderTexture(bLinear);
-	if (ensure(RenderTargetTexture) == false)
-	{
-		return false;
-	}
-
-	FSceneViewInitOptions ViewInitOptions = Internal::MakeSceneViewInitOptions(
-		ViewFrame,
-		Dimensions,
-		HorzFOVDegrees,
-		NearPlaneDist,
-		VisiblePrimitives);
-
-	FName CaptureTypeName;
-	switch (CaptureType)
-	{
-	case ERenderCaptureType::WorldNormal:
-		CaptureTypeName = FName("WorldNormal");
-		break;
-	case ERenderCaptureType::Roughness:
-		CaptureTypeName = FName("Roughness");
-		break;
-	case ERenderCaptureType::Metallic:
-		CaptureTypeName = FName("Metallic");
-		break;
-	case ERenderCaptureType::Specular:
-		CaptureTypeName = FName("Specular");
-		break;
-	case ERenderCaptureType::Opacity:
-		CaptureTypeName = FName("Opacity");
-		break;
-	case ERenderCaptureType::SubsurfaceColor:
-		CaptureTypeName = FName("SubsurfaceColor");
-		break;
-	case ERenderCaptureType::BaseColor:
-	default:
-		CaptureTypeName = FName("BaseColor");
-	}
-
-	ReadImageBuffer.Reset();
-	UE::Internal::RenderSceneVisualizationToTexture( 
-		RenderTargetTexture, this->Dimensions,
-		World->Scene, CaptureTypeName,
-		ViewInitOptions,
-		Config,
-		ReadImageBuffer,
-		LastCaptureViewMatrices);
-
-	{
-		TRACE_CPUPROFILER_EVENT_SCOPE(FWorldRenderCapture::CaptureFromPosition_CopyBufferToImage);
-
-		ResultImageOut.SetDimensions(Dimensions);
-
-		int32 Width = Dimensions.GetWidth();
-		int32 Height = Dimensions.GetHeight();
-		for (int32 yi = 0; yi < Height; ++yi)
-		{
-			for (int32 xi = 0; xi < Width; ++xi)
-			{
-				FLinearColor PixelColorf = ReadImageBuffer[yi * Width + xi];
-				PixelColorf.A = 1.0f;
-				ResultImageOut.SetPixel(FVector2i(xi, yi), PixelColorf);
-			}
-		}
-	}
-
-	if (bWriteDebugImage)
+	else
 	{
-		WriteDebugImage(ResultImageOut, *CaptureTypeName.ToString());
+		const FName BufferVisualizationMode = UE::Internal::GetBufferVisualizationModeName(CaptureType);
+		bCaptured = CaptureBufferVisualizationFromPosition(BufferVisualizationMode, ViewFrame, HorzFOVDegrees, NearPlaneDist, ResultImageOut, Config);
 	}
 
-	return true;
+	return bCaptured;
 }
 
 
diff --git a/Engine/Plugins/Runtime/MeshModelingToolset/Source/ModelingComponents/Public/Scene/WorldRenderCapture.h b/Engine/Plugins/Runtime/MeshModelingToolset/Source/ModelingComponents/Public/Scene/WorldRenderCapture.h
--- a/Engine/Plugins/Runtime/MeshModelingToolset/Source/ModelingComponents/Public/Scene/WorldRenderCapture.h
+++ b/Engine/Plugins/Runtime/MeshModelingToolset/Source/ModelingComponents/Public/Scene/WorldRenderCapture.h
@@ -317,6 +317,21 @@ protected:
 		FImageAdapter& ResultImageOut,
 		const FRenderCaptureConfig& Config = {});
 
+	/**
+	 * The following values for VisualizationMode are supported:
+	 *   BaseColor, Metallic, Roughness, Specular, WorldNormal, Opacity, SubsurfaceColor
+	 * 
+	 * Its not clear where the valid VisualizationMode FNames are defined, it may be in the [Engine.BufferVisualizationMaterials] section of
+	 * Config/BaseEngine.ini, which is read by FBufferVisualizationData::Initialize() and called when the engine initializes in UEngine::Init()
+	 */
+	bool CaptureBufferVisualizationFromPosition(
+		const FName& VisualizationMode,
+		const FFrame3d& Frame,
+		double HorzFOVDegrees,
+		double NearPlaneDist,
+		FImageAdapter& ResultImageOut,
+		const FRenderCaptureConfig& Config = {});
+
 	bool bWriteDebugImage = false;
 	int32 DebugImageCounter = -1;
 	FString DebugImageFolderName = TEXT("WorldRenderCapture");
