FWorldRenderCapture: Remove unecessary caching of view matrices

diff --git a/Engine/Plugins/Runtime/MeshModelingToolset/Source/ModelingComponents/Private/Scene/SceneCapturePhotoSet.cpp b/Engine/Plugins/Runtime/MeshModelingToolset/Source/ModelingComponents/Private/Scene/SceneCapturePhotoSet.cpp
--- a/Engine/Plugins/Runtime/MeshModelingToolset/Source/ModelingComponents/Private/Scene/SceneCapturePhotoSet.cpp
+++ b/Engine/Plugins/Runtime/MeshModelingToolset/Source/ModelingComponents/Private/Scene/SceneCapturePhotoSet.cpp
@@ -222,11 +222,6 @@ void FSceneCapturePhotoSet::Compute()
 		RenderCapture.SetDimensions(Params.Dimensions);
 		RenderCapture.CaptureFromPosition(CaptureType, NewPhoto.Frame, NewPhoto.HorzFOVDegrees, NewPhoto.NearPlaneDist, Image, Config);
 		GetPhotoSet1f(CaptureType).Add(MoveTemp(NewPhoto));
-
-		if (CaptureType == ERenderCaptureType::DeviceDepth)
-		{
-			PhotoViewMatricies.Add(RenderCapture.GetLastCaptureViewMatrices());
-		}
 	};
 
 	// Iterate by channel computing all the photos rather than by photo/viewpoint computing all the channels.
@@ -486,10 +481,6 @@ void FSceneCapturePhotoSet::AddExteriorCaptures(
 		if (PhotoSetStatus[ERenderCaptureType::DeviceDepth] == ECaptureTypeStatus::Pending)
 		{
 			CaptureImageTypeFunc_1f(ERenderCaptureType::DeviceDepth, DeviceDepthPhotoSet);
-			if (PhotoViewMatricies.Num() < NumDirections)
-			{
-				PhotoViewMatricies.Add(RenderCapture.GetLastCaptureViewMatrices());
-			}
 		}
 		if (PhotoSetStatus[ERenderCaptureType::BaseColor] == ECaptureTypeStatus::Pending)
 		{
@@ -635,7 +626,6 @@ bool FSceneCapturePhotoSet::ComputeSampleLocation(
 	if (ValidSampleDepthThreshold > 0)
 	{
 		check(DeviceDepthPhotoSet.Num() == PhotoSetParams.Num());
-		check(PhotoViewMatricies.Num() == PhotoSetParams.Num());
 	}
 
 	PhotoIndex = IndexConstants::InvalidID;
@@ -649,6 +639,8 @@ bool FSceneCapturePhotoSet::ComputeSampleLocation(
 		const FSpatialPhotoParams& Params = PhotoSetParams[Index];
 		check(Params.Dimensions.IsSquare());
 
+		const FViewMatrices ViewMatrices = GetRenderCaptureViewMatrices(Params.Frame, Params.HorzFOVDegrees, Params.NearPlaneDist, Params.Dimensions);
+
 		FFrame3d RenderFrame(Params.Frame.Origin, Params.Frame.Y(), Params.Frame.Z(), Params.Frame.X());
 
 		FVector3d ViewDirection = RenderFrame.Z();
@@ -694,7 +686,7 @@ bool FSceneCapturePhotoSet::ComputeSampleLocation(
 						{
 							// Compute the pixel position in world space to use it to compute a depth according to the render
 							FVector3d PixelPositionDevice{DeviceXY, DeviceZ};
-							FVector4d PixelPositionWorld = PhotoViewMatricies[Index].GetInvViewProjectionMatrix().TransformPosition(PixelPositionDevice);
+							FVector4d PixelPositionWorld = ViewMatrices.GetInvViewProjectionMatrix().TransformPosition(PixelPositionDevice);
 							PixelPositionWorld /= PixelPositionWorld.W;
 
 							// Compare the depth of the sample with the depth of the pixel and consider the sample invalid
@@ -816,6 +808,8 @@ void FSceneCapturePhotoSet::GetSceneSamples(FSceneSamples& OutSamples)
 	{
 		const FSpatialPhotoParams& Params = GetSpatialPhotoParams()[PhotoIndex];
 
+		const FViewMatrices ViewMatrices = GetRenderCaptureViewMatrices(Params.Frame, Params.HorzFOVDegrees, Params.NearPlaneDist, Params.Dimensions);
+
 		for (int64 PixelLinearIndex = 0; PixelLinearIndex < Params.Dimensions.Num(); ++PixelLinearIndex)
 		{
 			const FSpatialPhoto1f& DeviceDepthPhoto = GetDeviceDepthPhotoSet().Get(PhotoIndex);
@@ -843,8 +837,7 @@ void FSceneCapturePhotoSet::GetSceneSamples(FSceneSamples& OutSamples)
 
 					// Map from normalized device coordinates to world coordinates
 					FVector3d PointDevice(DeviceXY, DeviceZ);
-					const FMatrix& InvViewProjectionMatrix = PhotoViewMatricies[PhotoIndex].GetInvViewProjectionMatrix();
-					const FVector4d PointWorld4 = InvViewProjectionMatrix.TransformPosition(PointDevice);
+					const FVector4d PointWorld4 = ViewMatrices.GetInvViewProjectionMatrix().TransformPosition(PointDevice);
 
 					// Convert from Homogenous to Cartesian coordinates
 					PointWorld.X = static_cast<float>(PointWorld4.X / PointWorld4.W);
@@ -1020,9 +1013,7 @@ void FSceneCapturePhotoSet::EmptyPhotoSet(ERenderCaptureType CaptureType)
 		SubsurfaceColorPhotoSet.Empty();
 		break;
 	case ERenderCaptureType::DeviceDepth:
-		// For the device depth photo set we have two containers to empty
 		DeviceDepthPhotoSet.Empty();
-		PhotoViewMatricies.Empty();
 		break;
 	default:
 		ensure(false);
diff --git a/Engine/Plugins/Runtime/MeshModelingToolset/Source/ModelingComponents/Private/Scene/WorldRenderCapture.cpp b/Engine/Plugins/Runtime/MeshModelingToolset/Source/ModelingComponents/Private/Scene/WorldRenderCapture.cpp
--- a/Engine/Plugins/Runtime/MeshModelingToolset/Source/ModelingComponents/Private/Scene/WorldRenderCapture.cpp
+++ b/Engine/Plugins/Runtime/MeshModelingToolset/Source/ModelingComponents/Private/Scene/WorldRenderCapture.cpp
@@ -653,7 +653,9 @@ bool FWorldRenderCapture::CaptureMRSFromPosition(
 	UE::Internal::PerformSceneRender(&Canvas, &ViewFamily);
 
 	// Cache the view/projection matricies we used to render the scene
+	PRAGMA_DISABLE_DEPRECATION_WARNINGS
 	LastCaptureViewMatrices = NewView->ViewMatrices;
+	PRAGMA_ENABLE_DEPRECATION_WARNINGS
 
 	if (ensure(PostProcessPassPixelData.IsValid()) == false)
 	{
@@ -722,7 +724,9 @@ bool FWorldRenderCapture::CaptureEmissiveFromPosition(
 	UE::Internal::PerformSceneRender(&Canvas, &ViewFamily);
 
 	// Cache the view/projection matricies we used to render the scene
+	PRAGMA_DISABLE_DEPRECATION_WARNINGS
 	LastCaptureViewMatrices = NewView->ViewMatrices;
+	PRAGMA_ENABLE_DEPRECATION_WARNINGS
 
 	int32 Width = Dimensions.GetWidth();
 	int32 Height = Dimensions.GetHeight();
@@ -817,7 +821,9 @@ bool FWorldRenderCapture::CaptureDeviceDepthFromPosition(
 	UE::Internal::PerformSceneRender(&Canvas, &ViewFamily);
 
 	// Cache the view/projection matricies we used to render the scene
+	PRAGMA_DISABLE_DEPRECATION_WARNINGS
 	LastCaptureViewMatrices = NewView->ViewMatrices;
+	PRAGMA_ENABLE_DEPRECATION_WARNINGS
 
 	// Copy the contents of the remote texture to system memory
 	ReadImageBuffer.Reset();
@@ -861,7 +867,7 @@ bool FWorldRenderCapture::CaptureDeviceDepthFromPosition(
 					FVector2d DeviceXY = FRenderCaptureCoordinateConverter2D::PixelToDevice(FVector2i(xi, yi), Width, Height);
 
 					// Compute world coordinates from normalized device coordinates
-					FVector4d Point = GetLastCaptureViewMatrices().GetInvViewProjectionMatrix().TransformPosition(FVector3d(DeviceXY, DeviceZ));
+					FVector4d Point = NewView->ViewMatrices.GetInvViewProjectionMatrix().TransformPosition(FVector3d(DeviceXY, DeviceZ));
 					Point /= Point.W;
 
 					// Log the point
@@ -945,7 +951,9 @@ bool FWorldRenderCapture::CaptureBufferVisualizationFromPosition(
 	UE::Internal::PerformSceneRender(&Canvas, &ViewFamily);
 
 	// Cache the view/projection matricies we used to render the scene
+	PRAGMA_DISABLE_DEPRECATION_WARNINGS
 	LastCaptureViewMatrices = NewView->ViewMatrices;
+	PRAGMA_ENABLE_DEPRECATION_WARNINGS
 
 	// Copy the contents of the remote texture to system memory
 	ReadImageBuffer.Reset();
@@ -1028,4 +1036,18 @@ void FWorldRenderCapture::WriteDebugImage(const FImageAdapter& ResultImageOut, c
 	static int32 CaptureIndex = 0;
 	int32 UseCounter = (DebugImageCounter >= 0) ? DebugImageCounter : CaptureIndex++;
 	UE::AssetUtils::SaveDebugImage(ResultImageOut, false, DebugImageFolderName, *ImageTypeName, UseCounter);
-}
\ No newline at end of file
+}
+
+
+FViewMatrices UE::Geometry::GetRenderCaptureViewMatrices(const FFrame3d& ViewFrame, double HorzFOVDegrees, double NearPlaneDist, FImageDimensions ViewRect)
+{
+	const FSceneViewInitOptions ViewInitOptions = UE::Internal::MakeSceneViewInitOptions(
+		ViewFrame,
+		ViewRect,
+		HorzFOVDegrees,
+		NearPlaneDist,
+		{});
+
+	return FViewMatrices(ViewInitOptions);
+}
+
diff --git a/Engine/Plugins/Runtime/MeshModelingToolset/Source/ModelingComponents/Public/Scene/SceneCapturePhotoSet.h b/Engine/Plugins/Runtime/MeshModelingToolset/Source/ModelingComponents/Public/Scene/SceneCapturePhotoSet.h
--- a/Engine/Plugins/Runtime/MeshModelingToolset/Source/ModelingComponents/Public/Scene/SceneCapturePhotoSet.h
+++ b/Engine/Plugins/Runtime/MeshModelingToolset/Source/ModelingComponents/Public/Scene/SceneCapturePhotoSet.h
@@ -303,9 +303,6 @@ protected:
 
 	TArray<FSpatialPhotoParams> PhotoSetParams;
 
-	// This used to unproject the DeviceDepth render capture
-	TArray<FViewMatrices> PhotoViewMatricies;
-
 	bool bWriteDebugImages = false;
 	FString DebugImagesFolderName = TEXT("SceneCapturePhotoSet");
 
diff --git a/Engine/Plugins/Runtime/MeshModelingToolset/Source/ModelingComponents/Public/Scene/WorldRenderCapture.h b/Engine/Plugins/Runtime/MeshModelingToolset/Source/ModelingComponents/Public/Scene/WorldRenderCapture.h
--- a/Engine/Plugins/Runtime/MeshModelingToolset/Source/ModelingComponents/Public/Scene/WorldRenderCapture.h
+++ b/Engine/Plugins/Runtime/MeshModelingToolset/Source/ModelingComponents/Public/Scene/WorldRenderCapture.h
@@ -195,6 +195,16 @@ struct FRenderCaptureCoordinateConverter2D
 	}
 };
 
+/**
+ * Returns the view matrices corresponding to the given render capture cameara parameters
+ * Useful to get the needed information to unproject the DeviceDepth render capture e.g., using the following
+ *
+ *   FViewMatrices ViewMatrices = GetRenderCaptureViewMatrices(...);
+ *   FVector2d DeviceXY = FRenderCaptureCoordinateConverter2D::PixelToDevice(...);
+ *   FVector4d WorldPoint = ViewMatrices.GetInvViewProjectionMatrix().TransformPosition(FVector3d(DeviceXY, DeviceZ));
+ *   WorldPoint /= WorldPoint.W;
+ */
+FViewMatrices GetRenderCaptureViewMatrices(const FFrame3d& ViewFrame, double HorzFOVDegrees, double NearPlaneDist, FImageDimensions ViewRect);
 
 /**
  * FWorldRenderCapture captures a rendering of a set of Actors in a World from a
@@ -248,7 +258,13 @@ public:
 	 * @return view matrices used in the last call to CaptureFromPosition
 	 * Useful to get the needed information to unproject the DeviceDepth render capture
 	 */
-	const FViewMatrices& GetLastCaptureViewMatrices() const { return LastCaptureViewMatrices; }
+	UE_DEPRECATED(5.4, "Please use GetRenderCaptureViewMatrices instead")
+	const FViewMatrices& GetLastCaptureViewMatrices() const
+	{
+		PRAGMA_DISABLE_DEPRECATION_WARNINGS
+		return LastCaptureViewMatrices;
+		PRAGMA_ENABLE_DEPRECATION_WARNINGS
+	}
 
 	/**
 	 * Capture the desired buffer type CaptureType with the given view/camera parameters. The returned image
@@ -288,6 +304,7 @@ protected:
 	UTextureRenderTarget2D* GetRenderTexture(bool bLinear);
 	UTextureRenderTarget2D* GetDepthRenderTexture();
 
+	UE_DEPRECATED(5.4, "Please use GetRenderCaptureViewMatrices instead")
 	FViewMatrices LastCaptureViewMatrices;
 
 	// temporary buffer used to read from texture
