When triangulating splines w/ 'Leave Separate' enabled and 'best fit planes' enabled, use separate different origins for each plane so moving one curve doesn't move the projection plane of the others

diff --git a/Engine/Plugins/Runtime/MeshModelingToolset/Source/ModelingOperators/Private/CurveOps/TriangulateCurvesOp.cpp b/Engine/Plugins/Runtime/MeshModelingToolset/Source/ModelingOperators/Private/CurveOps/TriangulateCurvesOp.cpp
--- a/Engine/Plugins/Runtime/MeshModelingToolset/Source/ModelingOperators/Private/CurveOps/TriangulateCurvesOp.cpp
+++ b/Engine/Plugins/Runtime/MeshModelingToolset/Source/ModelingOperators/Private/CurveOps/TriangulateCurvesOp.cpp
@@ -300,7 +300,8 @@ void FTriangulateCurvesOp::CalculateResult(FProgressCancel* Progress)
 			FVector3d PlanePoint;
 			FVector3d Normal;
 			PolygonTriangulation::ComputePolygonPlane<double>(Paths[CurveIdx].Vertices, Normal, PlanePoint);
-			GroupIdx = ProcessPaths(Normal, SharedFrameOrigin, TArrayView<const FCurvePath>(&Paths[CurveIdx], 1), GroupIdx);
+			FVector3d ProjectedFrameOrigin = FVector3d::PointPlaneProject(SharedFrameOrigin, FPlane(PlanePoint, Normal));
+			GroupIdx = ProcessPaths(Normal, ProjectedFrameOrigin, TArrayView<const FCurvePath>(&Paths[CurveIdx], 1), GroupIdx);
 		}
 	}
 	else
