Add comments to Mesh Splines tool properties/enums

diff --git a/Engine/Plugins/Experimental/MeshModelingToolsetExp/Source/MeshModelingToolsExp/Public/TriangulateSplinesTool.h b/Engine/Plugins/Experimental/MeshModelingToolsetExp/Source/MeshModelingToolsExp/Public/TriangulateSplinesTool.h
--- a/Engine/Plugins/Experimental/MeshModelingToolsetExp/Source/MeshModelingToolsExp/Public/TriangulateSplinesTool.h
+++ b/Engine/Plugins/Experimental/MeshModelingToolsetExp/Source/MeshModelingToolsExp/Public/TriangulateSplinesTool.h
@@ -38,7 +38,7 @@ public:
 	UPROPERTY(EditAnywhere, Category = Spline)
 	EFlattenCurveMethod FlattenMethod = EFlattenCurveMethod::DoNotFlatten;
 
-	// Note: Combining and offsetting curves only works when curves are flattened; curves will be left separate and non-offset if FlattenMethod is DoNotFlatten
+	// Whether or how to combine the curves
 	UPROPERTY(EditAnywhere, Category = Spline, meta = (EditCondition = "FlattenMethod != EFlattenCurveMethod::DoNotFlatten", EditConditionHides))
 	ECombineCurvesMethod CombineMethod = ECombineCurvesMethod::LeaveSeparate;
 
diff --git a/Engine/Plugins/Runtime/MeshModelingToolset/Source/ModelingOperators/Public/CurveOps/TriangulateCurvesOp.h b/Engine/Plugins/Runtime/MeshModelingToolset/Source/ModelingOperators/Public/CurveOps/TriangulateCurvesOp.h
--- a/Engine/Plugins/Runtime/MeshModelingToolset/Source/ModelingOperators/Public/CurveOps/TriangulateCurvesOp.h
+++ b/Engine/Plugins/Runtime/MeshModelingToolset/Source/ModelingOperators/Public/CurveOps/TriangulateCurvesOp.h
@@ -12,51 +12,72 @@ class USplineComponent;
 UENUM()
 enum class EFlattenCurveMethod : uint8
 {
+	// Do not flatten the curves before triangulations
 	DoNotFlatten,
+	// Fit planes to the curves, and flatten the curves by projection to their plane
 	ToBestFitPlane,
+	// Flatten by projection along the X axis
 	AlongX,
+	// Flatten by projection along the Y axis
 	AlongY,
+	// Flatten by projection along the Z axis
 	AlongZ
 };
 
 UENUM()
 enum class ECombineCurvesMethod : uint8
 {
+	// Triangulate each curve separately
 	LeaveSeparate,
+	// Triangulate the union of the curve polygons -- the space covered by any of the polygons
 	Union,
+	// Triangulate the intersection of the curve polygons -- the space covered by all of the polygons
 	Intersect,
+	// Triangulate the difference of the first curve polygon minus the remaining curve polygons
 	Difference,
+	// Triangulate the exclusive-or of the curve polygons -- the space covered by the first polygon, or the remaining polygons, but not both
 	ExclusiveOr
 };
 
 UENUM()
 enum class EOffsetClosedCurvesMethod : uint8
 {
+	// Do not offset the closed curves
 	DoNotOffset,
+	// Offset the outside of the closed curves -- growing or shrinking the solid shape
 	OffsetOuterSide,
+	// Offset both sides of the closed curves -- creating hollow shapes that follow the curves with Curve Offset width
 	OffsetBothSides
 };
 
 UENUM()
 enum class EOffsetOpenCurvesMethod : uint8
 {
+	// Treat open curves as if they were closed -- connecting the last point back to the first
 	TreatAsClosed,
+	// Offset the open curves, creating shapes following the curves with Curve Offset width
 	Offset
 };
 
 UENUM()
 enum class EOffsetJoinMethod : uint8
 {
+	// Cut off corners between offset edges with square shapes
 	Square,
+	// Miter corners between offset edges, extending the neighboring curve edges straight to their intersection point, unless that point is farther than the miter limit distance
 	Miter,
+	// Smoothly join corners between offset edges with circular paths
 	Round
 };
 
 UENUM()
 enum class EOpenCurveEndShapes : uint8
 {
+	// Close the ends of open paths with square end caps
 	Square,
+	// Close the ends of open paths with round end caps
 	Round,
+	// Close the ends of open paths abruptly with no end caps
 	Butt
 };
 
