Fix horde issue 383823

#horde 383823

diff --git a/Engine/Plugins/Runtime/MeshModelingToolset/Source/ModelingComponents/Public/Scene/WorldRenderCapture.h b/Engine/Plugins/Runtime/MeshModelingToolset/Source/ModelingComponents/Public/Scene/WorldRenderCapture.h
--- a/Engine/Plugins/Runtime/MeshModelingToolset/Source/ModelingComponents/Public/Scene/WorldRenderCapture.h
+++ b/Engine/Plugins/Runtime/MeshModelingToolset/Source/ModelingComponents/Public/Scene/WorldRenderCapture.h
@@ -211,6 +211,7 @@ FViewMatrices GetRenderCaptureViewMatrices(const FFrame3d& ViewFrame, double Hor
  * specific viewpoint. Various types of rendering are supported, as defined by ERenderCaptureType.
  * Currently rendering an entire World, ie without an explicit list of Actors or Components, is not supported.
  */
+PRAGMA_DISABLE_DEPRECATION_WARNINGS // Silence deprecation warnings for deprecated LastCaptureViewMatrices member in implicit constructors
 class MODELINGCOMPONENTS_API FWorldRenderCapture
 {
 public:
@@ -354,6 +355,7 @@ protected:
 	FString DebugImageFolderName = TEXT("WorldRenderCapture");
 	void WriteDebugImage(const FImageAdapter& ResultImageOut, const FString& ImageTypeName);
 };
+PRAGMA_ENABLE_DEPRECATION_WARNINGS
 
 
 
