Fix horde issue 383964

#horde 383964

diff --git a/Engine/Plugins/Runtime/MeshModelingToolset/Source/ModelingComponents/Private/Scene/WorldRenderCapture.cpp b/Engine/Plugins/Runtime/MeshModelingToolset/Source/ModelingComponents/Private/Scene/WorldRenderCapture.cpp
--- a/Engine/Plugins/Runtime/MeshModelingToolset/Source/ModelingComponents/Private/Scene/WorldRenderCapture.cpp
+++ b/Engine/Plugins/Runtime/MeshModelingToolset/Source/ModelingComponents/Private/Scene/WorldRenderCapture.cpp
@@ -170,11 +170,6 @@ void FlushRHIThreadToUpdateTextureRenderTargetReference()
 
 
 
-FWorldRenderCapture::FWorldRenderCapture()
-{
-	Dimensions = FImageDimensions(128, 128);
-}
-
 FWorldRenderCapture::~FWorldRenderCapture()
 {
 	Shutdown();
diff --git a/Engine/Plugins/Runtime/MeshModelingToolset/Source/ModelingComponents/Public/Scene/WorldRenderCapture.h b/Engine/Plugins/Runtime/MeshModelingToolset/Source/ModelingComponents/Public/Scene/WorldRenderCapture.h
--- a/Engine/Plugins/Runtime/MeshModelingToolset/Source/ModelingComponents/Public/Scene/WorldRenderCapture.h
+++ b/Engine/Plugins/Runtime/MeshModelingToolset/Source/ModelingComponents/Public/Scene/WorldRenderCapture.h
@@ -211,11 +211,18 @@ FViewMatrices GetRenderCaptureViewMatrices(const FFrame3d& ViewFrame, double Hor
  * specific viewpoint. Various types of rendering are supported, as defined by ERenderCaptureType.
  * Currently rendering an entire World, ie without an explicit list of Actors or Components, is not supported.
  */
-PRAGMA_DISABLE_DEPRECATION_WARNINGS // Silence deprecation warnings for deprecated LastCaptureViewMatrices member in implicit constructors
 class MODELINGCOMPONENTS_API FWorldRenderCapture
 {
 public:
-	FWorldRenderCapture();
+
+PRAGMA_DISABLE_DEPRECATION_WARNINGS // Silence deprecation warnings for deprecated LastCaptureViewMatrices member in implicit constructors
+	FWorldRenderCapture() = default;
+	FWorldRenderCapture(const FWorldRenderCapture&) = default;
+	FWorldRenderCapture(FWorldRenderCapture&&) noexcept = default;
+	FWorldRenderCapture& operator=(const FWorldRenderCapture&) = default;
+	FWorldRenderCapture& operator=(FWorldRenderCapture&&) noexcept = default;
+PRAGMA_ENABLE_DEPRECATION_WARNINGS
+
 	~FWorldRenderCapture();
 
 	/** Explicitly release any allocated textures or other data structres */
@@ -294,7 +301,7 @@ protected:
 	TSet<FPrimitiveComponentId> VisiblePrimitives;
 	FBoxSphereBounds VisibleBounds;
 
-	FImageDimensions Dimensions;
+	FImageDimensions Dimensions = FImageDimensions(128, 128);
 
 	// Temporary textures used as render targets. We explicitly prevent this from being GC'd internally
 	UTextureRenderTarget2D* LinearRenderTexture = nullptr;
@@ -355,7 +362,6 @@ protected:
 	FString DebugImageFolderName = TEXT("WorldRenderCapture");
 	void WriteDebugImage(const FImageAdapter& ResultImageOut, const FString& ImageTypeName);
 };
-PRAGMA_ENABLE_DEPRECATION_WARNINGS
 
 
 
