Fix Auto UV tool failing to find target's bounding box when run on the sky sphere, due to the bounds call expecting to find bounds for colliding components.
 + Add fallback logic in case the primitive bounds are still invalid / empty

diff --git a/Engine/Plugins/Runtime/MeshModelingToolset/Source/MeshModelingToolsEditorOnly/Private/ParameterizeMeshTool.cpp b/Engine/Plugins/Runtime/MeshModelingToolset/Source/MeshModelingToolsEditorOnly/Private/ParameterizeMeshTool.cpp
--- a/Engine/Plugins/Runtime/MeshModelingToolset/Source/MeshModelingToolsEditorOnly/Private/ParameterizeMeshTool.cpp
+++ b/Engine/Plugins/Runtime/MeshModelingToolset/Source/MeshModelingToolsEditorOnly/Private/ParameterizeMeshTool.cpp
@@ -106,7 +106,21 @@ void UParameterizeMeshTool::Setup()
 		UVLayoutView = NewObject<UUVLayoutPreview>(this);
 		UVLayoutView->CreateInWorld(GetTargetWorld());
 		UVLayoutView->SetSourceMaterials(MaterialSet);
-		UVLayoutView->SetSourceWorldPosition(InputTransform, UE::ToolTarget::GetTargetActor(Target)->GetComponentsBoundingBox());
+		FBox Bounds = UE::ToolTarget::GetTargetActor(Target)->GetComponentsBoundingBox(true /*bNonColliding*/);
+		if (!Bounds.IsValid) // If component did not have valid bounds ...
+		{
+			// Try getting bounds from mesh
+			Bounds = (FBox)InputMesh->GetBounds();
+			Bounds = Bounds.TransformBy(InputTransform);
+			// If mesh is also empty, just create some small valid Bounds
+			if (!Bounds.IsValid)
+			{
+				UE_LOG(LogGeometry, Warning, TEXT("Auto UV Tool started on mesh with empty bounding box"));
+				constexpr double SmallSize = FMathd::ZeroTolerance;
+				Bounds = FBox(FVector(-1, -1, -1) * SmallSize, FVector(1, 1, 1) * SmallSize);
+			}
+		}
+		UVLayoutView->SetSourceWorldPosition(InputTransform, Bounds);
 		UVLayoutView->Settings->bEnabled = false;
 		UVLayoutView->Settings->bShowWireframe = false;
 		UVLayoutView->Settings->RestoreProperties(this, TEXT("ParameterizeMeshTool"));
