add concept of saving-on-clear/restoring the selection manager's selection, so that tools can restore the before-tool selection if the tool is cancelled, or if the tool otherwise did not change the underlying meshes.

 + add a callback to the modeling editor mode so it can know whether a tool was cancelled, and make the modeling mode use the selection manager's new save-on-clear/restore feature automatically when its tools are cancelled.

 + make some tools (normals, tangents, inspector) restore the before-tool-start selection after tool use via a utility function, since they shouldn't change mesh IDs that the selection relies on.

diff --git a/Engine/Plugins/Editor/ModelingToolsEditorMode/Source/ModelingToolsEditorMode/Private/ModelingToolsEditorMode.cpp b/Engine/Plugins/Editor/ModelingToolsEditorMode/Source/ModelingToolsEditorMode/Private/ModelingToolsEditorMode.cpp
--- a/Engine/Plugins/Editor/ModelingToolsEditorMode/Source/ModelingToolsEditorMode/Private/ModelingToolsEditorMode.cpp
+++ b/Engine/Plugins/Editor/ModelingToolsEditorMode/Source/ModelingToolsEditorMode/Private/ModelingToolsEditorMode.cpp
@@ -1008,6 +1008,25 @@ void UModelingToolsEditorMode::Enter()
 		}
 	});
 
+	// Restore saved selections when tool is cancelled or tool declares it is safe to do so via the IInteractiveToolManageGeometrySelectionAPI
+	GetToolManager()->OnToolEndedWithStatus.AddLambda([this](UInteractiveToolManager* Manager, UInteractiveTool* Tool, EToolShutdownType ShutdownType)
+	{
+		bool bCanRestore = (ShutdownType == EToolShutdownType::Cancel);
+		if (IInteractiveToolManageGeometrySelectionAPI* ManageSelectionTool = Cast<IInteractiveToolManageGeometrySelectionAPI>(Tool))
+		{
+			bCanRestore = bCanRestore || ManageSelectionTool->IsInputSelectionValidOnOutput();
+		}
+		if (bCanRestore)
+		{
+			GetSelectionManager()->RestoreSavedSelection();
+		}
+		else
+		{
+			GetSelectionManager()->DiscardSavedSelection();
+		}
+		ensureMsgf(!GetSelectionManager()->HasSavedSelection(), TEXT("Selection manager's saved selection should be cleared on tool end."));
+	});
+
 	// do any toolkit UI initialization that depends on the mode setup above
 	if (Toolkit.IsValid())
 	{
@@ -1568,7 +1587,8 @@ void UModelingToolsEditorMode::OnToolPostBuild(
 	// built, so that the Tool has a chance to see the Selection
 	if (GetSelectionManager() && GetSelectionManager()->HasSelection())
 	{
-		GetSelectionManager()->ClearSelection();
+		ensureMsgf(!GetSelectionManager()->HasSavedSelection(), TEXT("Selection manager should not already have a saved selection before we save-on-clear here in tool setup."));
+		GetSelectionManager()->ClearSelection(/*bSaveSelectionBeforeClear*/ true);
 	}
 }
 
diff --git a/Engine/Plugins/Experimental/MeshModelingToolsetExp/Source/MeshModelingToolsEditorOnlyExp/Public/MeshTangentsTool.h b/Engine/Plugins/Experimental/MeshModelingToolsetExp/Source/MeshModelingToolsEditorOnlyExp/Public/MeshTangentsTool.h
--- a/Engine/Plugins/Experimental/MeshModelingToolsetExp/Source/MeshModelingToolsEditorOnlyExp/Public/MeshTangentsTool.h
+++ b/Engine/Plugins/Experimental/MeshModelingToolsetExp/Source/MeshModelingToolsEditorOnlyExp/Public/MeshTangentsTool.h
@@ -94,7 +94,7 @@ public:
  * Simple Mesh Simplifying Tool
  */
 UCLASS()
-class MESHMODELINGTOOLSEDITORONLYEXP_API UMeshTangentsTool : public USingleSelectionMeshEditingTool, public UE::Geometry::IGenericDataOperatorFactory<UE::Geometry::FMeshTangentsd>
+class MESHMODELINGTOOLSEDITORONLYEXP_API UMeshTangentsTool : public USingleSelectionMeshEditingTool, public UE::Geometry::IGenericDataOperatorFactory<UE::Geometry::FMeshTangentsd>, public IInteractiveToolManageGeometrySelectionAPI
 {
 	GENERATED_BODY()
 public:
@@ -113,6 +113,12 @@ public:
 	// input selection support
 	void SetGeometrySelection(UE::Geometry::FGeometrySelection&& SelectionIn);
 
+	// IInteractiveToolManageGeometrySelectionAPI -- this tool won't update external geometry selection or change selection-relevant mesh IDs
+	virtual bool IsInputSelectionValidOnOutput() override
+	{
+		return true;
+	}
+
 	// IGenericDataOperatorFactory API
 	virtual TUniquePtr<UE::Geometry::TGenericDataOperator<UE::Geometry::FMeshTangentsd>> MakeNewOperator() override;
 
diff --git a/Engine/Plugins/Experimental/MeshModelingToolsetExp/Source/MeshModelingToolsExp/Public/EditNormalsTool.h b/Engine/Plugins/Experimental/MeshModelingToolsetExp/Source/MeshModelingToolsExp/Public/EditNormalsTool.h
--- a/Engine/Plugins/Experimental/MeshModelingToolsetExp/Source/MeshModelingToolsExp/Public/EditNormalsTool.h
+++ b/Engine/Plugins/Experimental/MeshModelingToolsetExp/Source/MeshModelingToolsExp/Public/EditNormalsTool.h
@@ -122,7 +122,7 @@ public:
  * Simple Mesh Normal Updating Tool
  */
 UCLASS()
-class MESHMODELINGTOOLSEXP_API UEditNormalsTool : public UMultiSelectionMeshEditingTool
+class MESHMODELINGTOOLSEXP_API UEditNormalsTool : public UMultiSelectionMeshEditingTool, public IInteractiveToolManageGeometrySelectionAPI
 {
 	GENERATED_BODY()
 
@@ -150,6 +150,12 @@ public:
 	// input selection support
 	void SetGeometrySelection(UE::Geometry::FGeometrySelection&& SelectionIn);
 
+	// IInteractiveToolManageGeometrySelectionAPI -- this tool won't update external geometry selection or change selection-relevant mesh IDs
+	virtual bool IsInputSelectionValidOnOutput() override
+	{
+		return true;
+	}
+
 protected:
 
 	UPROPERTY()
diff --git a/Engine/Plugins/Experimental/MeshModelingToolsetExp/Source/MeshModelingToolsExp/Public/MeshInspectorTool.h b/Engine/Plugins/Experimental/MeshModelingToolsetExp/Source/MeshModelingToolsExp/Public/MeshInspectorTool.h
--- a/Engine/Plugins/Experimental/MeshModelingToolsetExp/Source/MeshModelingToolsExp/Public/MeshInspectorTool.h
+++ b/Engine/Plugins/Experimental/MeshModelingToolsetExp/Source/MeshModelingToolsExp/Public/MeshInspectorTool.h
@@ -224,7 +224,7 @@ public:
  * Mesh Inspector Tool for visualizing mesh information
  */
 UCLASS()
-class MESHMODELINGTOOLSEXP_API UMeshInspectorTool : public USingleSelectionMeshEditingTool
+class MESHMODELINGTOOLSEXP_API UMeshInspectorTool : public USingleSelectionMeshEditingTool, public IInteractiveToolManageGeometrySelectionAPI
 {
 	GENERATED_BODY()
 
@@ -245,6 +245,12 @@ public:
 
 	virtual void OnPropertyModified(UObject* PropertySet, FProperty* Property) override;
 
+	// IInteractiveToolManageGeometrySelectionAPI -- this tool won't update external geometry selection or change selection-relevant mesh IDs
+	virtual bool IsInputSelectionValidOnOutput() override
+	{
+		return true;
+	}
+
 public:
 
 	virtual void IncreaseLineWidthAction();
diff --git a/Engine/Plugins/Runtime/MeshModelingToolset/Source/ModelingComponents/Private/Selection/GeometrySelectionManager.cpp b/Engine/Plugins/Runtime/MeshModelingToolset/Source/ModelingComponents/Private/Selection/GeometrySelectionManager.cpp
--- a/Engine/Plugins/Runtime/MeshModelingToolset/Source/ModelingComponents/Private/Selection/GeometrySelectionManager.cpp
+++ b/Engine/Plugins/Runtime/MeshModelingToolset/Source/ModelingComponents/Private/Selection/GeometrySelectionManager.cpp
@@ -1,6 +1,7 @@
 // Copyright Epic Games, Inc. All Rights Reserved.
 
 #include "Selection/GeometrySelectionManager.h"
+#include "CoreGlobals.h" // for GIsTransacting
 #include "Engine/Engine.h"
 #include "Selection/DynamicMeshSelector.h"
 #include "Selection/ToolSelectionUtil.h"
@@ -41,6 +42,7 @@ void UGeometrySelectionManager::RegisterSelectorFactory(TUniquePtr<IGeometrySele
 
 void UGeometrySelectionManager::Shutdown()
 {
+	DiscardSavedSelection();
 	OnSelectionModified.Clear();
 	ToolsContext = nullptr;
 	TransactionsAPI = nullptr;
@@ -312,6 +314,7 @@ void UGeometrySelectionManager::ClearActiveTargets()
 	// undo that cannot be redone later, because on redo the Targets will not exist yet
 	// (one possibility would be to emit separate changes for when the target set is modified?? would that work w/ delete?? )
 	ensure(HasSelection() == false);
+	DiscardSavedSelection();
 
 	for (TSharedPtr<FGeometrySelectionTarget> Target : ActiveTargetReferences)
 	{
@@ -617,13 +620,18 @@ bool UGeometrySelectionManager::RayHitTest(
 }
 
 
-void UGeometrySelectionManager::ClearSelection()
+void UGeometrySelectionManager::ClearSelection(bool bSaveSelectionBeforeClear)
 {
 	if (!HasSelection())
 	{
 		return;
 	}
 
+	if (bSaveSelectionBeforeClear)
+	{
+		SaveCurrentSelection();
+	}
+
 	GetTransactionsAPI()->BeginUndoTransaction(LOCTEXT("ClearSelection", "Clear Selection"));
 
 	for (TSharedPtr<FGeometrySelectionTarget> Target : ActiveTargetReferences)
@@ -863,6 +871,87 @@ bool UGeometrySelectionManager::SetSelectionForComponent(UPrimitiveComponent* Co
 	return false;
 }
 
+void UGeometrySelectionManager::SaveCurrentSelection()
+{
+	SavedSelection.Reset();
+	for (TSharedPtr<FGeometrySelectionTarget> Target : ActiveTargetReferences)
+	{
+		SavedSelection.Targets.Add(Target->TargetIdentifier.TargetObject);
+		SavedSelection.Selections.Add(Target->Selection);
+	}
+}
+
+bool UGeometrySelectionManager::RestoreSavedSelection()
+{
+#if WITH_EDITORONLY_DATA
+	// Cannot update the selection if we're already in a transaction (can happen e.g. when we undo out of a tool)
+	if (GIsTransacting)
+	{
+		DiscardSavedSelection();
+		return false;
+	}
+#endif
+
+	check(SavedSelection.Targets.Num() == SavedSelection.Selections.Num());
+	GetTransactionsAPI()->BeginUndoTransaction(LOCTEXT("RestoreSelection", "Restore Selection"));
+
+	bool bSuccess = true;
+	for (int32 TargetIdx = 0; TargetIdx < SavedSelection.Targets.Num(); ++TargetIdx)
+	{
+		if (!SavedSelection.Targets[TargetIdx].IsValid())
+		{
+			bSuccess = false;
+			continue;
+		}
+		const FGeometrySelection& NewSelection = SavedSelection.Selections[TargetIdx];
+		bool bFound = false;
+		for (TSharedPtr<FGeometrySelectionTarget> Target : ActiveTargetReferences)
+		{
+			if (SavedSelection.Targets[TargetIdx] == Target->TargetIdentifier.TargetObject)
+			{
+				FGeometrySelection InitialSelection = Target->Selection;
+				FGeometrySelectionDelta AfterDelta;
+				Target->Selector->UpdateSelectionFromSelection(
+					NewSelection, true, *Target->SelectionEditor,
+					FGeometrySelectionUpdateConfig{ EGeometrySelectionChangeType::Replace }, &AfterDelta);
+				if (AfterDelta.IsEmpty() == false)
+				{
+					TUniquePtr<FGeometrySelectionReplaceChange> NewSelectionChange = MakeUnique<FGeometrySelectionReplaceChange>();
+					NewSelectionChange->Identifier = Target->TargetIdentifier;
+					NewSelectionChange->After = Target->Selection;
+					NewSelectionChange->Before = InitialSelection;
+					GetTransactionsAPI()->AppendChange(this, MoveTemp(NewSelectionChange), LOCTEXT("RestoreSelection", "Restore Selection"));
+
+					bSelectionRenderCachesDirty = true;
+					OnSelectionModified.Broadcast();
+				}
+				bFound = true;
+				break;
+			}
+		}
+		if (!bFound)
+		{
+			bSuccess = false;
+		}
+	}
+
+	GetTransactionsAPI()->EndUndoTransaction();
+
+	DiscardSavedSelection();
+
+	return bSuccess;
+}
+
+void UGeometrySelectionManager::DiscardSavedSelection()
+{
+	SavedSelection.Empty();
+}
+
+bool UGeometrySelectionManager::HasSavedSelection()
+{
+	return !SavedSelection.Selections.IsEmpty();
+}
+
 
 bool UGeometrySelectionManager::UpdateSelectionPreviewViaRaycast(
 	const FRay3d& WorldRay )
diff --git a/Engine/Plugins/Runtime/MeshModelingToolset/Source/ModelingComponents/Public/Selection/GeometrySelectionManager.h b/Engine/Plugins/Runtime/MeshModelingToolset/Source/ModelingComponents/Public/Selection/GeometrySelectionManager.h
--- a/Engine/Plugins/Runtime/MeshModelingToolset/Source/ModelingComponents/Public/Selection/GeometrySelectionManager.h
+++ b/Engine/Plugins/Runtime/MeshModelingToolset/Source/ModelingComponents/Public/Selection/GeometrySelectionManager.h
@@ -160,8 +160,33 @@ public:
 	/**
 	 * Clear any active element selections. 
 	 * This function will emit a Transaction for the selection change.
+	 * @param bSaveSelectionBeforeClear if true, calls SaveCurrentSelection() before the selection is cleared, so it can be restored later.
 	 */
-	virtual void ClearSelection();
+	virtual void ClearSelection(bool bSaveSelectionBeforeClear = false);
+
+protected:
+	/**
+	 * Save the active selection. Overwrites any existing saved selections with the current selections. Typically used via ClearSelection(true)
+	 */
+	virtual void SaveCurrentSelection();
+
+public:
+	/**
+	 * Attempt to restore (and then discard) the most recent saved selection.
+	 * If there is no active saved selection, does nothing. On failure to restore, will still discard the saved selection.
+	 * @return false if could not restore selection, which can happen if the restore was called while transacting (e.g., when a tool is exited via undo), or if the selection objects were not found or not valid
+	 */
+	virtual bool RestoreSavedSelection();
+	
+	/**
+	 * Discard the saved selection, if there is one.
+	 */
+	virtual void DiscardSavedSelection();
+
+	/**
+	 * @return true if there is a non-empty saved selection, false otherwise.
+	 */
+	virtual bool HasSavedSelection();
 
 	/**
 	 * Use the given WorldRay to update the active element selection based on UpdateConfig.
@@ -459,6 +484,26 @@ protected:
 	UE::Geometry::FGeometrySelectionDelta ActiveTrackedDelta;
 	bool bSelectionModifiedDuringTrackedChange = false;
 
+private:
+
+	// Tracks saved selection state. Useful when the selection is temporarily cleared (e.g., for a tool)
+	struct FSavedSelection
+	{
+		TArray<TWeakObjectPtr<UObject>> Targets;
+		TArray<FGeometrySelection> Selections;
+		void Empty()
+		{
+			Targets.Empty();
+			Selections.Empty();
+		}
+		void Reset()
+		{
+			Targets.Reset();
+			Selections.Reset();
+		}
+	};
+	FSavedSelection SavedSelection;
+
 };
 
 
diff --git a/Engine/Source/Runtime/InteractiveToolsFramework/Private/InteractiveToolManager.cpp b/Engine/Source/Runtime/InteractiveToolsFramework/Private/InteractiveToolManager.cpp
--- a/Engine/Source/Runtime/InteractiveToolsFramework/Private/InteractiveToolManager.cpp
+++ b/Engine/Source/Runtime/InteractiveToolsFramework/Private/InteractiveToolManager.cpp
@@ -282,6 +282,7 @@ void UInteractiveToolManager::DeactivateToolInternal(EToolSide Side, EToolShutdo
 
 		PostInvalidation();
 
+		OnToolEndedWithStatus.Broadcast(this, DoneTool, ShutdownType);
 		OnToolEnded.Broadcast(this, DoneTool);
 
 		bInToolShutdown = false;
diff --git a/Engine/Source/Runtime/InteractiveToolsFramework/Public/InteractiveToolManager.h b/Engine/Source/Runtime/InteractiveToolsFramework/Public/InteractiveToolManager.h
--- a/Engine/Source/Runtime/InteractiveToolsFramework/Public/InteractiveToolManager.h
+++ b/Engine/Source/Runtime/InteractiveToolsFramework/Public/InteractiveToolManager.h
@@ -316,6 +316,10 @@ public:
 	DECLARE_MULTICAST_DELEGATE_TwoParams(FToolManagerToolEndedSignature, UInteractiveToolManager*, UInteractiveTool*);
 	FToolManagerToolEndedSignature OnToolEnded;
 
+	// Variant of OnToolEnded that also reports the EToolShutdownType
+	DECLARE_MULTICAST_DELEGATE_ThreeParams(FToolManagerToolCancelledSignature, UInteractiveToolManager*, UInteractiveTool*, EToolShutdownType);
+	FToolManagerToolCancelledSignature OnToolEndedWithStatus;
+
 
 protected:
 	/** Pointer to current Context-Queries implementation */
diff --git a/Engine/Source/Runtime/InteractiveToolsFramework/Public/InteractiveToolQueryInterfaces.h b/Engine/Source/Runtime/InteractiveToolsFramework/Public/InteractiveToolQueryInterfaces.h
--- a/Engine/Source/Runtime/InteractiveToolsFramework/Public/InteractiveToolQueryInterfaces.h
+++ b/Engine/Source/Runtime/InteractiveToolsFramework/Public/InteractiveToolQueryInterfaces.h
@@ -175,4 +175,23 @@ public:
 
 
 
+// UInterface for IInteractiveToolManageGeometrySelectionAPI
+UINTERFACE(MinimalAPI)
+class UInteractiveToolManageGeometrySelectionAPI : public UInterface
+{
+	GENERATED_BODY()
+};
 
+/**
+ * Provides an API to allow a tool to report how it has affected (updated or invalidated) geometry selections on the tool's targets
+ */
+class IInteractiveToolManageGeometrySelectionAPI
+{
+	GENERATED_BODY()
+public:
+
+	/**
+	 * @return true if the tool implementing this API has not updated the geometry selection or modified geometry to invalidate any previous geometry selection, i.e. has not removed/added vertices/edges/triangles
+	 */
+	virtual bool IsInputSelectionValidOnOutput() { return false; }
+};
