FWorldRenderCapture: Consolidate Emissive render capture code path with the other buffer visualization code paths

diff --git a/Engine/Plugins/Runtime/MeshModelingToolset/Source/ModelingComponents/Private/Scene/WorldRenderCapture.cpp b/Engine/Plugins/Runtime/MeshModelingToolset/Source/ModelingComponents/Private/Scene/WorldRenderCapture.cpp
--- a/Engine/Plugins/Runtime/MeshModelingToolset/Source/ModelingComponents/Private/Scene/WorldRenderCapture.cpp
+++ b/Engine/Plugins/Runtime/MeshModelingToolset/Source/ModelingComponents/Private/Scene/WorldRenderCapture.cpp
@@ -548,6 +548,8 @@ FName GetBufferVisualizationModeName(ERenderCaptureType CaptureType)
 		return FName("Opacity");
 	case ERenderCaptureType::SubsurfaceColor:
 		return FName("SubsurfaceColor");
+	case ERenderCaptureType::Emissive:
+		return FName("PreTonemapHDRColor");
 	default:
 		ensure(false);
 	}
@@ -678,70 +680,8 @@ bool FWorldRenderCapture::CaptureEmissiveFromPosition(
 	FImageAdapter& ResultImageOut,
 	const FRenderCaptureConfig& Config)
 {
-	TRACE_CPUPROFILER_EVENT_SCOPE(FWorldRenderCapture_Emissive);
-
-	UTextureRenderTarget2D* RenderTargetTexture = GetRenderTexture(true);
-	if (ensure(RenderTargetTexture) == false)
-	{
-		return false;
-	}
-	FTextureRenderTargetResource* RenderTargetResource = RenderTargetTexture->GameThread_GetRenderTargetResource();
-
-	FEngineShowFlags ShowFlags(ESFIM_Game);
-	Internal::SetCommonShowFlags(ShowFlags, Config.bAntiAliasing);
-	ShowFlags.SetPostProcessing(true);
-	ShowFlags.SetVisualizeBuffer(true);
-	ShowFlags.SetLighting(true); // Need this because Emissive is lighting I guess
-
-	FSceneViewFamilyContext ViewFamily(
-		FSceneViewFamily::ConstructionValues(RenderTargetResource, World->Scene, ShowFlags)
-		.SetTime(FGameTime())
-		.SetRealtimeUpdate(false)
-	);
-
-	FSceneViewInitOptions ViewInitOptions = Internal::MakeSceneViewInitOptions(
-		Frame,
-		Dimensions,
-		HorzFOVDegrees,
-		NearPlaneDist,
-		VisiblePrimitives);
-	ViewInitOptions.ViewFamily = &ViewFamily;
-
-	FSceneView* NewView = new FSceneView(ViewInitOptions);
-	NewView->CurrentBufferVisualizationMode = FName("PreTonemapHDRColor");
-	ViewFamily.Views.Add(NewView);
-
-	ViewFamily.SetScreenPercentageInterface(new FLegacyScreenPercentageDriver(ViewFamily, 1.0f));
-
-	FCanvas Canvas(RenderTargetResource, nullptr, FGameTime(), NewView->GetFeatureLevel());
-	Canvas.Clear(FLinearColor::Transparent);
-
-	UE::Internal::PerformSceneRender(&Canvas, &ViewFamily);
-
-	// Cache the view/projection matricies we used to render the scene
-	PRAGMA_DISABLE_DEPRECATION_WARNINGS
-	LastCaptureViewMatrices = NewView->ViewMatrices;
-	PRAGMA_ENABLE_DEPRECATION_WARNINGS
-
-	int32 Width = Dimensions.GetWidth();
-	int32 Height = Dimensions.GetHeight();
-
-	// Copy the contents of the remote texture to system memory
-	ReadImageBuffer.Reset();
-	ReadImageBuffer.SetNumUninitialized(Width * Height);
-	FReadSurfaceDataFlags ReadSurfaceDataFlags(RCM_MinMax);		// don't normalize buffer values, we want full HDR range
-	ReadSurfaceDataFlags.SetLinearToGamma(false);
-	RenderTargetResource->ReadLinearColorPixels(ReadImageBuffer, ReadSurfaceDataFlags, FIntRect(0, 0, Width, Height));
-
-	FImageView CaptureData(ReadImageBuffer.GetData(), Width, Height);
-	UE::Internal::CopyCaptureDataToOutputImageFormat(CaptureData, ResultImageOut);
-
-	if (bWriteDebugImage)
-	{
-		WriteDebugImage(ResultImageOut, TEXT("Emissive"));
-	}
-
-	return true;
+	const FName BufferVisualizationMode = UE::Internal::GetBufferVisualizationModeName(ERenderCaptureType::Emissive);
+	return CaptureBufferVisualizationFromPosition(BufferVisualizationMode, Frame, HorzFOVDegrees, NearPlaneDist, ResultImageOut, Config);
 }
 
 
@@ -761,9 +701,6 @@ bool FWorldRenderCapture::CaptureDeviceDepthFromPosition(
 	}
 	FTextureRenderTargetResource* RenderTargetResource = DepthRenderTexture->GameThread_GetRenderTargetResource();
 
-	int32 Width = Dimensions.GetWidth();
-	int32 Height = Dimensions.GetHeight();
-
 	FEngineShowFlags ShowFlags(ESFIM_Game);
 	Internal::SetCommonShowFlags(ShowFlags, false); // Never use AntiAliasing so we dont blend pixels at different depths
 
@@ -794,9 +731,10 @@ bool FWorldRenderCapture::CaptureDeviceDepthFromPosition(
 	}
 
 	FSceneView* NewView = new FSceneView(ViewInitOptions);
+	ViewFamily.Views.Add(NewView);
+
 	NewView->AntiAliasingMethod = EAntiAliasingMethod::AAM_None;
 	NewView->SetupAntiAliasingMethod();
-	ViewFamily.Views.Add(NewView);
 
 	NewView->StartFinalPostprocessSettings(ViewInitOptions.ViewOrigin);
 	NewView->EndFinalPostprocessSettings(ViewInitOptions);
@@ -822,13 +760,12 @@ bool FWorldRenderCapture::CaptureDeviceDepthFromPosition(
 
 	// Copy the contents of the remote texture to system memory
 	ReadImageBuffer.Reset();
-	ReadImageBuffer.SetNumUninitialized(Width * Height);
+	ReadImageBuffer.SetNumUninitialized(Dimensions.Num());
 	FReadSurfaceDataFlags ReadSurfaceDataFlags(RCM_MinMax);
 	ReadSurfaceDataFlags.SetLinearToGamma(false);
+	RenderTargetResource->ReadLinearColorPixels(ReadImageBuffer, ReadSurfaceDataFlags, FIntRect(0, 0, Dimensions.GetWidth(), Dimensions.GetHeight()));
 
-	RenderTargetResource->ReadLinearColorPixels(ReadImageBuffer, ReadSurfaceDataFlags, FIntRect(0, 0, Width, Height));
-
-	FImageView CaptureData(ReadImageBuffer.GetData(), Width, Height);
+	FImageView CaptureData(ReadImageBuffer.GetData(), Dimensions.GetWidth(), Dimensions.GetHeight());
 	UE::Internal::CopyCaptureDataToOutputImageFormat(CaptureData, ResultImageOut, [](const FLinearColor& Color)
 	{
 		// Reverse the float encoding used for the Depth value in SceneCapturePixelShader.usf (one minus the device
@@ -844,37 +781,10 @@ bool FWorldRenderCapture::CaptureDeviceDepthFromPosition(
 		return FLinearColor(DeviceZ, 0., 0., 0.);
 	});
 
-	// Set this to true to compute a world point cloud for debugging.
-	// You probably want to change the logging so it writes an .obj file
-	constexpr bool bDebugDepthCapture = false;
-
-	if constexpr (bDebugDepthCapture)
-	{
-		int PointIndex = 1;
-		for (int32 yi = 0; yi < Height; ++yi)
-		{
-			for (int32 xi = 0; xi < Width; ++xi)
-			{
-				float DeviceZ = ResultImageOut.GetPixel(FVector2i(xi, yi)).X;
-				if (DeviceZ > 0.) // Skip points on the far plane since these unproject to infinity
-				{
-					// Map from pixel space to NDC space
-					FVector2d DeviceXY = FRenderCaptureCoordinateConverter2D::PixelToDevice(FVector2i(xi, yi), Width, Height);
-
-					// Compute world coordinates from normalized device coordinates
-					FVector4d Point = NewView->ViewMatrices.GetInvViewProjectionMatrix().TransformPosition(FVector3d(DeviceXY, DeviceZ));
-					Point /= Point.W;
-
-					// Log the point
-					UE_LOG(LogGeometry, Log, TEXT("DebugDepthCapture: [%d] %s"), PointIndex, *Point.ToString());
-					PointIndex += 1;
-				}
-			} // xi
-		} // yi
-	}
-
 	if (bWriteDebugImage)
 	{
+		// It is also helpful to save a world point cloud .obj file. See FSceneCapturePhotoSet::GetSceneSamples which uses
+		// GetRenderCaptureViewMatrices to turn the depth texture into world positions
 		WriteDebugImage(ResultImageOut, TEXT("DeviceDepth"));
 	}
 
@@ -893,6 +803,9 @@ bool FWorldRenderCapture::CaptureBufferVisualizationFromPosition(
 {
 	TRACE_CPUPROFILER_EVENT_SCOPE_TEXT(*FString::Printf(TEXT("FWorldRenderCapture_%s"), *VisualizationMode.ToString()));
 
+	const bool bEmissive = VisualizationMode == Internal::GetBufferVisualizationModeName(ERenderCaptureType::Emissive);
+	const bool bRoughness = VisualizationMode == Internal::GetBufferVisualizationModeName(ERenderCaptureType::Roughness);
+
 	// The following handles buffer visualization materials found in /Engine/BufferVisualization/<VisualizationMode>.
 	// Sometimes these postprocess materials change the raw GBuffer data, in particular:
 	// - The Roughness postprocess material output is:           GBufferRoughness^Gamma,   with Gamma=2.2
@@ -900,22 +813,11 @@ bool FWorldRenderCapture::CaptureBufferVisualizationFromPosition(
 	// Note the relationship between linear and gamma encoded values: EncodedValue = LinearValue^(1/Gamma)
 
 	// Roughness visualization is rendered with gamma correction (unclear why)
-	bool bLinear = (VisualizationMode != Internal::GetBufferVisualizationModeName(ERenderCaptureType::Roughness));
-	UTextureRenderTarget2D* RenderTargetTexture = GetRenderTexture(bLinear);
+	UTextureRenderTarget2D* RenderTargetTexture = GetRenderTexture(!bRoughness);
 	if (ensure(RenderTargetTexture) == false)
 	{
 		return false;
 	}
-
-	FSceneViewInitOptions ViewInitOptions = Internal::MakeSceneViewInitOptions(
-		Frame,
-		Dimensions,
-		HorzFOVDegrees,
-		NearPlaneDist,
-		VisiblePrimitives);
-
-	int32 Width = Dimensions.GetWidth();
-	int32 Height = Dimensions.GetHeight();
 	FTextureRenderTargetResource* RenderTargetResource = RenderTargetTexture->GameThread_GetRenderTargetResource();
 
 	FEngineShowFlags ShowFlags(ESFIM_Game);
@@ -924,6 +826,7 @@ bool FWorldRenderCapture::CaptureBufferVisualizationFromPosition(
 	ShowFlags.SetPostProcessMaterial(true);
 	ShowFlags.SetPostProcessing(true);
 	ShowFlags.SetVisualizeBuffer(true);
+	ShowFlags.SetLighting(bEmissive); // Need this because Emissive is lighting I guess
 
 	FSceneViewFamilyContext ViewFamily(
 		FSceneViewFamily::ConstructionValues(RenderTargetResource, World->Scene, ShowFlags)
@@ -931,12 +834,19 @@ bool FWorldRenderCapture::CaptureBufferVisualizationFromPosition(
 		.SetRealtimeUpdate(false)
 	);
 
+	FSceneViewInitOptions ViewInitOptions = Internal::MakeSceneViewInitOptions(
+		Frame,
+		Dimensions,
+		HorzFOVDegrees,
+		NearPlaneDist,
+		VisiblePrimitives);
 	ViewInitOptions.ViewFamily = &ViewFamily;
 
 	FSceneView* NewView = new FSceneView(ViewInitOptions);
-	NewView->CurrentBufferVisualizationMode = VisualizationMode;
 	ViewFamily.Views.Add(NewView);
 
+	NewView->CurrentBufferVisualizationMode = VisualizationMode;
+
 	ViewFamily.SetScreenPercentageInterface(new FLegacyScreenPercentageDriver(ViewFamily, 1.0f));
 
 	// should we cache the FCanvas?
@@ -950,13 +860,12 @@ bool FWorldRenderCapture::CaptureBufferVisualizationFromPosition(
 	LastCaptureViewMatrices = NewView->ViewMatrices;
 	PRAGMA_ENABLE_DEPRECATION_WARNINGS
 
-	// Copy the contents of the remote texture to system memory
+	// Copy the contents of the remote texture to system memory. For Emissive we dont normalize buffer values, we want full HDR ranges.
 	ReadImageBuffer.Reset();
-	ReadImageBuffer.SetNumUninitialized(Width * Height);
-	//FReadSurfaceDataFlags ReadSurfaceDataFlags(RCM_MinMax);		// should we use MinMax to avoid normalization?
-	FReadSurfaceDataFlags ReadSurfaceDataFlags;
+	ReadImageBuffer.SetNumUninitialized(Dimensions.Num());
+	FReadSurfaceDataFlags ReadSurfaceDataFlags(bEmissive ? RCM_MinMax : RCM_UNorm); // Should we always use MinMax to avoid normalization?
 	ReadSurfaceDataFlags.SetLinearToGamma(false);
-	RenderTargetResource->ReadLinearColorPixels(ReadImageBuffer, ReadSurfaceDataFlags, FIntRect(0, 0, Width, Height));
+	RenderTargetResource->ReadLinearColorPixels(ReadImageBuffer, ReadSurfaceDataFlags, FIntRect(0, 0, Dimensions.GetWidth(), Dimensions.GetHeight()));
 
 	FImageView CaptureData(ReadImageBuffer.GetData(), Dimensions.GetWidth(), Dimensions.GetHeight());
 	UE::Internal::CopyCaptureDataToOutputImageFormat(CaptureData, ResultImageOut);
@@ -991,11 +900,7 @@ bool FWorldRenderCapture::CaptureFromPosition(
 
 	bool bCaptured = false;
 
-	if (CaptureType == ERenderCaptureType::Emissive)
-	{
-		bCaptured = CaptureEmissiveFromPosition(ViewFrame, HorzFOVDegrees, NearPlaneDist, ResultImageOut, Config);
-	}
-	else if (CaptureType == ERenderCaptureType::CombinedMRS)
+	if (CaptureType == ERenderCaptureType::CombinedMRS)
 	{
 		bCaptured = CaptureMRSFromPosition(ViewFrame, HorzFOVDegrees, NearPlaneDist, ResultImageOut, Config);
 	}
diff --git a/Engine/Plugins/Runtime/MeshModelingToolset/Source/ModelingComponents/Public/Scene/WorldRenderCapture.h b/Engine/Plugins/Runtime/MeshModelingToolset/Source/ModelingComponents/Public/Scene/WorldRenderCapture.h
--- a/Engine/Plugins/Runtime/MeshModelingToolset/Source/ModelingComponents/Public/Scene/WorldRenderCapture.h
+++ b/Engine/Plugins/Runtime/MeshModelingToolset/Source/ModelingComponents/Public/Scene/WorldRenderCapture.h
@@ -319,6 +319,7 @@ protected:
 	TArray<FLinearColor> ReadImageBuffer;
 
 	/** Emissive is a special case and uses different code than capture of color/property channels */
+	UE_DEPRECATED(5.4, "Please use CaptureBufferVisualizationFromPosition with VisualizationMode=PreTonemapHDRColor instead")
 	bool CaptureEmissiveFromPosition(
 		const FFrame3d& Frame,
 		double HorzFOVDegrees,
@@ -344,7 +345,7 @@ protected:
 
 	/**
 	 * The following values for VisualizationMode are supported:
-	 *   BaseColor, Metallic, Roughness, Specular, WorldNormal, Opacity, SubsurfaceColor
+	 *   BaseColor, Metallic, Roughness, Specular, WorldNormal, Opacity, SubsurfaceColor, PreTonemapHDRColor
 	 * 
 	 * Its not clear where the valid VisualizationMode FNames are defined, it may be in the [Engine.BufferVisualizationMaterials] section of
 	 * Config/BaseEngine.ini, which is read by FBufferVisualizationData::Initialize() and called when the engine initializes in UEngine::Init()
