BakeRC: Fix bug where the Opacity channel was not rebaked when a setting invalidating all baked textures was changed

diff --git a/Engine/Plugins/Runtime/MeshModelingToolset/Source/ModelingComponents/Private/Scene/WorldRenderCapture.cpp b/Engine/Plugins/Runtime/MeshModelingToolset/Source/ModelingComponents/Private/Scene/WorldRenderCapture.cpp
--- a/Engine/Plugins/Runtime/MeshModelingToolset/Source/ModelingComponents/Private/Scene/WorldRenderCapture.cpp
+++ b/Engine/Plugins/Runtime/MeshModelingToolset/Source/ModelingComponents/Private/Scene/WorldRenderCapture.cpp
@@ -35,15 +35,11 @@ static TAutoConsoleVariable<int32> CVarModelingWorldRenderCaptureWarmupFrames(
 FRenderCaptureTypeFlags FRenderCaptureTypeFlags::All(bool bCombinedMRS)
 {
 	FRenderCaptureTypeFlags Result;
-	Result.bBaseColor = true;
-	Result.bRoughness = true;
-	Result.bMetallic = true;
-	Result.bSpecular = true;
-	Result.bEmissive = true;
-	Result.bWorldNormal = true;
+	ForEachCaptureType([&Result](ERenderCaptureType CaptureType)
+	{
+		Result[CaptureType] = true;
+	});
 	Result.bCombinedMRS = bCombinedMRS;
-	Result.bSubsurfaceColor = true;
-	Result.bDeviceDepth = true;
 	return Result;
 }
 
