fix bugs that can occur on undo during mouse drags for draw spline, lattice and collision editor tools, by keeping a transaction open for the duration of the drag

also add a little utility struct to manage this kind of held-open transaction, with some error checking to help catch any errors where these transactions end up opened but not closed or vice versa

diff --git a/Engine/Plugins/Experimental/MeshModelingToolsetExp/Source/MeshModelingToolsEditorOnlyExp/Private/DrawSplineTool.cpp b/Engine/Plugins/Experimental/MeshModelingToolsetExp/Source/MeshModelingToolsEditorOnlyExp/Private/DrawSplineTool.cpp
--- a/Engine/Plugins/Experimental/MeshModelingToolsetExp/Source/MeshModelingToolsEditorOnlyExp/Private/DrawSplineTool.cpp
+++ b/Engine/Plugins/Experimental/MeshModelingToolsetExp/Source/MeshModelingToolsEditorOnlyExp/Private/DrawSplineTool.cpp
@@ -519,6 +519,8 @@ void UDrawSplineTool::TransitionOutputMode()
 
 void UDrawSplineTool::Shutdown(EToolShutdownType ShutdownType)
 {
+	LongTransactions.CloseAll(GetToolManager());
+
 	Settings->SaveProperties(this);
 
 	if (PreviousTargetActor)
@@ -830,6 +832,8 @@ void UDrawSplineTool::OnClickPress(const FInputDeviceRay& PressPos)
 	FVector3d HitLocation, HitNormal;
 	double HitT;
 
+	LongTransactions.Open(DrawSplineToolLocals::AddPointTransactionName, GetToolManager());
+
 	// Regardless of DrawMode, start by placing a point, though don't emit a transaction until mouse up
 	if (ensure(Raycast(PressPos.WorldRay, HitLocation, HitNormal, HitT)))
 	{
@@ -968,6 +972,8 @@ void UDrawSplineTool::OnTerminateDragSequence()
 		break;
 	}
 	}
+
+	LongTransactions.Close(GetToolManager());
 }
 
 
diff --git a/Engine/Plugins/Experimental/MeshModelingToolsetExp/Source/MeshModelingToolsEditorOnlyExp/Public/DrawSplineTool.h b/Engine/Plugins/Experimental/MeshModelingToolsetExp/Source/MeshModelingToolsEditorOnlyExp/Public/DrawSplineTool.h
--- a/Engine/Plugins/Experimental/MeshModelingToolsetExp/Source/MeshModelingToolsEditorOnlyExp/Public/DrawSplineTool.h
+++ b/Engine/Plugins/Experimental/MeshModelingToolsetExp/Source/MeshModelingToolsEditorOnlyExp/Public/DrawSplineTool.h
@@ -9,6 +9,7 @@
 #include "InteractiveToolBuilder.h"
 #include "InteractiveToolChange.h"
 #include "ToolContextInterfaces.h" // FViewCameraState
+#include "TransactionUtil.h"
 
 #include "DrawSplineTool.generated.h"
 
@@ -269,6 +270,9 @@ protected:
 
 	FViewCameraState CameraState;
 
+private:
+	UE::TransactionUtil::FLongTransactionTracker LongTransactions;
+
 public:
 	// Helper class for making undo/redo transactions, to avoid friending all the variations.
 	class FSplineChange : public FToolCommandChange
diff --git a/Engine/Plugins/Runtime/MeshModelingToolset/Source/ModelingComponents/Private/Mechanics/CollisionPrimitivesMechanic.cpp b/Engine/Plugins/Runtime/MeshModelingToolset/Source/ModelingComponents/Private/Mechanics/CollisionPrimitivesMechanic.cpp
--- a/Engine/Plugins/Runtime/MeshModelingToolset/Source/ModelingComponents/Private/Mechanics/CollisionPrimitivesMechanic.cpp
+++ b/Engine/Plugins/Runtime/MeshModelingToolset/Source/ModelingComponents/Private/Mechanics/CollisionPrimitivesMechanic.cpp
@@ -190,6 +190,8 @@ void UCollisionPrimitivesMechanic::SetWorld(UWorld* World)
 
 void UCollisionPrimitivesMechanic::Shutdown()
 {
+	LongTransactions.CloseAll(GetParentTool()->GetToolManager());
+
 	if (PreviewGeometry)
 	{
 		PreviewGeometry->Disconnect();
@@ -985,9 +987,12 @@ void UCollisionPrimitivesMechanic::OnUpdateModifierState(int ModifierID, bool bI
 // These get bound to the delegates on the marquee mechanic.
 void UCollisionPrimitivesMechanic::OnDragRectangleStarted()
 {
+	using namespace CollisionPrimitivesMechanicLocals;
+
 	PreDragSelection = SelectedPrimitiveIDs;
 	bIsDraggingRectangle = true;
 	UpdateGizmoVisibility();
+	LongTransactions.Open(CollisionPrimitiveSelectionTransactionText, GetParentTool()->GetToolManager());
 }
 
 void UCollisionPrimitivesMechanic::OnDragRectangleChanged(const FCameraRectangle& Rectangle)
@@ -1036,8 +1041,6 @@ void UCollisionPrimitivesMechanic::OnDragRectangleFinished(const FCameraRectangl
 
 	bIsDraggingRectangle = false;
 
-	ParentTool->GetToolManager()->BeginUndoTransaction(CollisionPrimitiveSelectionTransactionText);
-
 	if (!IsEqual(PreDragSelection, SelectedPrimitiveIDs))
 	{
 		checkSlow(CurrentActiveProxy);
@@ -1052,7 +1055,7 @@ void UCollisionPrimitivesMechanic::OnDragRectangleFinished(const FCameraRectangl
 	// We hid the gizmo at rectangle start, so it needs updating now.
 	UpdateGizmoLocation();
 
-	ParentTool->GetToolManager()->EndUndoTransaction();
+	LongTransactions.Close(GetParentTool()->GetToolManager());
 
 	UpdateDrawables();
 }
diff --git a/Engine/Plugins/Runtime/MeshModelingToolset/Source/ModelingComponents/Private/Mechanics/LatticeControlPointsMechanic.cpp b/Engine/Plugins/Runtime/MeshModelingToolset/Source/ModelingComponents/Private/Mechanics/LatticeControlPointsMechanic.cpp
--- a/Engine/Plugins/Runtime/MeshModelingToolset/Source/ModelingComponents/Private/Mechanics/LatticeControlPointsMechanic.cpp
+++ b/Engine/Plugins/Runtime/MeshModelingToolset/Source/ModelingComponents/Private/Mechanics/LatticeControlPointsMechanic.cpp
@@ -193,6 +193,8 @@ void ULatticeControlPointsMechanic::SetWorld(UWorld* World)
 
 void ULatticeControlPointsMechanic::Shutdown()
 {
+	LongTransactions.CloseAll(GetParentTool()->GetToolManager());
+
 	if (PreviewGeometryActor)
 	{
 		PreviewGeometryActor->Destroy();
@@ -626,6 +628,7 @@ void ULatticeControlPointsMechanic::OnDragRectangleStarted()
 {
 	PreDragSelection = SelectedPointIDs;
 	bIsDraggingRectangle = true;
+	LongTransactions.Open(LatticePointSelectionTransactionText, GetParentTool()->GetToolManager());
 	UpdateGizmoVisibility();
 }
 
@@ -674,7 +677,6 @@ void ULatticeControlPointsMechanic::OnDragRectangleFinished(const FCameraRectang
 
 	bIsDraggingRectangle = false;
 
-	ParentTool->GetToolManager()->BeginUndoTransaction(LatticePointSelectionTransactionText);
 
 	if (!IsEqual(PreDragSelection, SelectedPointIDs))
 	{
@@ -689,7 +691,7 @@ void ULatticeControlPointsMechanic::OnDragRectangleFinished(const FCameraRectang
 	// We hid the gizmo at rectangle start, so it needs updating now.
 	UpdateGizmoLocation();
 
-	ParentTool->GetToolManager()->EndUndoTransaction();
+	LongTransactions.Close(GetParentTool()->GetToolManager());
 
 	UpdateDrawables();
 }
diff --git a/Engine/Plugins/Runtime/MeshModelingToolset/Source/ModelingComponents/Private/TransactionUtil.cpp b/Engine/Plugins/Runtime/MeshModelingToolset/Source/ModelingComponents/Private/TransactionUtil.cpp
new file mode 100644
--- /dev/null
+++ b/Engine/Plugins/Runtime/MeshModelingToolset/Source/ModelingComponents/Private/TransactionUtil.cpp
@@ -0,0 +1,47 @@
+// Copyright Epic Games, Inc. All Rights Reserved.
+
+#include "TransactionUtil.h"
+
+#include "InteractiveToolManager.h"
+
+namespace UE::TransactionUtil
+{
+
+	void FLongTransactionTracker::Open(FText TransactionName, UInteractiveToolManager* ToolManager)
+	{
+		if (OpenCount == 0) // only actually 'begin' when opening the outer-most scope
+		{
+			ToolManager->BeginUndoTransaction(TransactionName);
+		}
+		OpenCount++;
+
+		// After opening, expect from 1 to MaxOpen to be open
+		ensureMsgf(OpenCount > 0 && OpenCount <= MaxOpen, TEXT("Opened new long transaction (%s); found unexpected current open count: %d (Expected from 1 to %d)"), *TransactionName.ToString(), OpenCount, MaxOpen);
+	}
+
+	
+	void FLongTransactionTracker::Close(UInteractiveToolManager* ToolManager)
+	{
+		// Before closing, expect from 1 to MaxOpen to be open
+		ensureMsgf(OpenCount > 0 && OpenCount <= MaxOpen, TEXT("Closing long transaction; found unexpected current open count: %d (Expected from 1 to %d)"), OpenCount, MaxOpen);
+
+		OpenCount--;
+		if (OpenCount == 0) // only actually 'end' when closing the outer-most scope 
+		{
+			ToolManager->EndUndoTransaction();
+		}
+	}
+
+	void FLongTransactionTracker::CloseAll(UInteractiveToolManager* ToolManager)
+	{
+		// When closing all, expect 0 to MaxOpen to be open
+		ensureMsgf(OpenCount >= 0 && OpenCount <= MaxOpen, TEXT("Closing any/all long transactions; found unexpected current open count: %d (Expected from 0 to %d)"), OpenCount, MaxOpen);
+
+		if (OpenCount > 0)
+		{
+			ToolManager->EndUndoTransaction();
+		}
+		OpenCount = 0;
+	}
+
+}
diff --git a/Engine/Plugins/Runtime/MeshModelingToolset/Source/ModelingComponents/Public/Mechanics/CollisionPrimitivesMechanic.h b/Engine/Plugins/Runtime/MeshModelingToolset/Source/ModelingComponents/Public/Mechanics/CollisionPrimitivesMechanic.h
--- a/Engine/Plugins/Runtime/MeshModelingToolset/Source/ModelingComponents/Public/Mechanics/CollisionPrimitivesMechanic.h
+++ b/Engine/Plugins/Runtime/MeshModelingToolset/Source/ModelingComponents/Public/Mechanics/CollisionPrimitivesMechanic.h
@@ -11,6 +11,7 @@
 #include "ToolContextInterfaces.h" //FViewCameraState
 #include "IntVectorTypes.h"
 #include "BoxTypes.h"
+#include "TransactionUtil.h"
 
 #include "CollisionPrimitivesMechanic.generated.h"
 
@@ -255,4 +256,6 @@ protected:
 		int32 ChangeStamp;	
 	};
 
+private:
+	UE::TransactionUtil::FLongTransactionTracker LongTransactions;
 };
diff --git a/Engine/Plugins/Runtime/MeshModelingToolset/Source/ModelingComponents/Public/Mechanics/LatticeControlPointsMechanic.h b/Engine/Plugins/Runtime/MeshModelingToolset/Source/ModelingComponents/Public/Mechanics/LatticeControlPointsMechanic.h
--- a/Engine/Plugins/Runtime/MeshModelingToolset/Source/ModelingComponents/Public/Mechanics/LatticeControlPointsMechanic.h
+++ b/Engine/Plugins/Runtime/MeshModelingToolset/Source/ModelingComponents/Public/Mechanics/LatticeControlPointsMechanic.h
@@ -10,6 +10,7 @@
 #include "Spatial/GeometrySet3.h"
 #include "ToolContextInterfaces.h" //FViewCameraState
 #include "IntVectorTypes.h"
+#include "TransactionUtil.h"
 
 #include "LatticeControlPointsMechanic.generated.h"
 
@@ -203,6 +204,9 @@ protected:
 
 	friend class FLatticeControlPointsMechanicSelectionChange;
 	friend class FLatticeControlPointsMechanicMovementChange;
+
+private:
+	UE::TransactionUtil::FLongTransactionTracker LongTransactions;
 };
 
 
diff --git a/Engine/Plugins/Runtime/MeshModelingToolset/Source/ModelingComponents/Public/TransactionUtil.h b/Engine/Plugins/Runtime/MeshModelingToolset/Source/ModelingComponents/Public/TransactionUtil.h
new file mode 100644
--- /dev/null
+++ b/Engine/Plugins/Runtime/MeshModelingToolset/Source/ModelingComponents/Public/TransactionUtil.h
@@ -0,0 +1,40 @@
+// Copyright Epic Games, Inc. All Rights Reserved.
+
+#pragma once
+
+#include "Internationalization/Text.h"
+
+class UInteractiveToolManager;
+
+// Utilities for managing logic around undo/redo transactions
+namespace UE::TransactionUtil
+{
+	// A simple manager class to track transactions that will be open for a duration of time, e.g. during a mouse drag, within a single tool or mechanic.
+	// Provides some error checking to help catch cases where we open without closing or close without opening, or if we nest opening more transactions than expected.
+	// One way these can be used is to keep a transaction open during a mouse drag interaction as a way to prevent the user from undoing during the mouse drag.
+	// Otherwise, undoing during a mouse drag can cause bugs in many tools if they are not careful to undo any pending work in OnTerminateDragSequence.
+	class FLongTransactionTracker
+	{
+	public:
+		FLongTransactionTracker(int32 MaxOpen = 1) : OpenCount(0), MaxOpen(MaxOpen) {}
+		~FLongTransactionTracker()
+		{
+			ensureMsgf(OpenCount == 0, TEXT("Expected all transactions to be closed, but found %d open."), OpenCount);
+		}
+
+		// Open a single long transaction.
+		MODELINGCOMPONENTS_API void Open(FText TransactionName, UInteractiveToolManager* ToolManager);
+
+		// Close a single long transaction. Expect that at least one is open.
+		MODELINGCOMPONENTS_API void Close(UInteractiveToolManager* ToolManager);
+
+		// Close any/all open transactions as part of shutdown of the owning tool/mechanic. Always safe to call on shutdown; OK if none are open.
+		MODELINGCOMPONENTS_API void CloseAll(UInteractiveToolManager* ToolManager);
+
+	private:
+		// Current number of open long transactions we are tracking.  (Note only the first will actually call BeginUndoTransaction)
+		int32 OpenCount;
+		// Maximum number of long transactions we allow to be open at once before triggering an ensure()
+		int32 MaxOpen;
+	};
+}
