add option to request tangents to the dynamic mesh provider, with custom implementations for targets that go through the mesh description (so we'll use the standard tangent compute paths in those cases)

diff --git a/Engine/Plugins/Runtime/MeshModelingToolset/Source/ModelingComponents/Private/ConversionUtils/DynamicMeshViaMeshDescriptionUtil.cpp b/Engine/Plugins/Runtime/MeshModelingToolset/Source/ModelingComponents/Private/ConversionUtils/DynamicMeshViaMeshDescriptionUtil.cpp
--- a/Engine/Plugins/Runtime/MeshModelingToolset/Source/ModelingComponents/Private/ConversionUtils/DynamicMeshViaMeshDescriptionUtil.cpp
+++ b/Engine/Plugins/Runtime/MeshModelingToolset/Source/ModelingComponents/Private/ConversionUtils/DynamicMeshViaMeshDescriptionUtil.cpp
@@ -9,12 +9,22 @@
 #include "MeshDescriptionToDynamicMesh.h"
 
 FDynamicMesh3 UE::Geometry::GetDynamicMeshViaMeshDescription(
-	IMeshDescriptionProvider& MeshDescriptionProvider)
+	IMeshDescriptionProvider& MeshDescriptionProvider, bool bRequestTangents)
 {
 	FDynamicMesh3 DynamicMesh;
 	FMeshDescriptionToDynamicMesh Converter;
 	Converter.bVIDsFromNonManifoldMeshDescriptionAttr = true;
-	Converter.Convert(MeshDescriptionProvider.GetMeshDescription(), DynamicMesh);
+	if (bRequestTangents)
+	{
+		FGetMeshParameters GetMeshParams;
+		GetMeshParams.bWantMeshTangents = true;
+		FMeshDescription MeshDescriptionCopy = MeshDescriptionProvider.GetMeshDescriptionCopy(GetMeshParams);
+		Converter.Convert(&MeshDescriptionCopy, DynamicMesh, bRequestTangents);
+	}
+	else
+	{
+		Converter.Convert(MeshDescriptionProvider.GetMeshDescription(), DynamicMesh, bRequestTangents);
+	}
 	return DynamicMesh;
 }
 
diff --git a/Engine/Plugins/Runtime/MeshModelingToolset/Source/ModelingComponents/Private/ModelingToolTargetUtil.cpp b/Engine/Plugins/Runtime/MeshModelingToolset/Source/ModelingComponents/Private/ModelingToolTargetUtil.cpp
--- a/Engine/Plugins/Runtime/MeshModelingToolset/Source/ModelingComponents/Private/ModelingToolTargetUtil.cpp
+++ b/Engine/Plugins/Runtime/MeshModelingToolset/Source/ModelingComponents/Private/ModelingToolTargetUtil.cpp
@@ -226,11 +226,10 @@ FDynamicMesh3 UE::ToolTarget::GetDynamicMeshCopy(UToolTarget* Target, bool bWant
 		return Mesh;
 	}
 
-	// TODO: Handle tangent computation. For now skip if tangents requested.
 	IDynamicMeshProvider* DynamicMeshProvider = Cast<IDynamicMeshProvider>(Target);
-	if (DynamicMeshProvider && !bWantMeshTangents)
+	if (DynamicMeshProvider)
 	{
-		return DynamicMeshProvider->GetDynamicMesh();
+		return DynamicMeshProvider->GetDynamicMesh(bWantMeshTangents);
 	}
 
 	IMeshDescriptionProvider* MeshDescriptionProvider = Cast<IMeshDescriptionProvider>(Target);
diff --git a/Engine/Plugins/Runtime/MeshModelingToolset/Source/ModelingComponents/Private/TargetInterfaces/DynamicMeshProvider.cpp b/Engine/Plugins/Runtime/MeshModelingToolset/Source/ModelingComponents/Private/TargetInterfaces/DynamicMeshProvider.cpp
new file mode 100644
--- /dev/null
+++ b/Engine/Plugins/Runtime/MeshModelingToolset/Source/ModelingComponents/Private/TargetInterfaces/DynamicMeshProvider.cpp
@@ -0,0 +1,15 @@
+// Copyright Epic Games, Inc. All Rights Reserved.
+
+#include "TargetInterfaces/DynamicMeshProvider.h"
+#include "DynamicMesh/DynamicMesh3.h"
+
+#define LOCTEXT_NAMESPACE "DynamicMeshProvider"
+
+using namespace UE::Geometry;
+
+FDynamicMesh3 IDynamicMeshProvider::GetDynamicMesh(bool bRequestTangents)
+{
+	return GetDynamicMesh();
+}
+
+#undef LOCTEXT_NAMESPACE 
\ No newline at end of file
diff --git a/Engine/Plugins/Runtime/MeshModelingToolset/Source/ModelingComponents/Public/ConversionUtils/DynamicMeshViaMeshDescriptionUtil.h b/Engine/Plugins/Runtime/MeshModelingToolset/Source/ModelingComponents/Public/ConversionUtils/DynamicMeshViaMeshDescriptionUtil.h
--- a/Engine/Plugins/Runtime/MeshModelingToolset/Source/ModelingComponents/Public/ConversionUtils/DynamicMeshViaMeshDescriptionUtil.h
+++ b/Engine/Plugins/Runtime/MeshModelingToolset/Source/ModelingComponents/Public/ConversionUtils/DynamicMeshViaMeshDescriptionUtil.h
@@ -17,7 +17,7 @@ namespace UE {
 namespace Geometry {
 
 	MODELINGCOMPONENTS_API FDynamicMesh3 GetDynamicMeshViaMeshDescription(
-		IMeshDescriptionProvider& MeshDescriptionProvider);
+		IMeshDescriptionProvider& MeshDescriptionProvider, bool bRequestTangents = false);
 
 	MODELINGCOMPONENTS_API void CommitDynamicMeshViaMeshDescription(
 		FMeshDescription&& CurrentMeshDescription,
diff --git a/Engine/Plugins/Runtime/MeshModelingToolset/Source/ModelingComponents/Public/TargetInterfaces/DynamicMeshProvider.h b/Engine/Plugins/Runtime/MeshModelingToolset/Source/ModelingComponents/Public/TargetInterfaces/DynamicMeshProvider.h
--- a/Engine/Plugins/Runtime/MeshModelingToolset/Source/ModelingComponents/Public/TargetInterfaces/DynamicMeshProvider.h
+++ b/Engine/Plugins/Runtime/MeshModelingToolset/Source/ModelingComponents/Public/TargetInterfaces/DynamicMeshProvider.h
@@ -26,4 +26,13 @@ public:
 	 * Gives a copy of a dynamic mesh for tools to operate on.
 	 */
 	virtual UE::Geometry::FDynamicMesh3 GetDynamicMesh() = 0;
+
+	/**
+	 * Gives a copy of a dynamic mesh for tools to operate on.
+	 * 
+	 * @param bRequestTangents Request tangents on the returned mesh. Not required if tangents are not on the source data and the provider does not have a standard way to generate them.
+	 *
+	 * Note: Default implementation simply returns GetDynamicMesh(). Overloaded implementations for e.g., Static and Skeletal Mesh sources will enable (and compute if needed) additional tangent data.
+	 */
+	virtual UE::Geometry::FDynamicMesh3 GetDynamicMesh(bool bRequestTangents);
 };
diff --git a/Engine/Plugins/Runtime/MeshModelingToolset/Source/ModelingComponentsEditorOnly/Private/ToolTargets/SkeletalMeshComponentToolTarget.cpp b/Engine/Plugins/Runtime/MeshModelingToolset/Source/ModelingComponentsEditorOnly/Private/ToolTargets/SkeletalMeshComponentToolTarget.cpp
--- a/Engine/Plugins/Runtime/MeshModelingToolset/Source/ModelingComponentsEditorOnly/Private/ToolTargets/SkeletalMeshComponentToolTarget.cpp
+++ b/Engine/Plugins/Runtime/MeshModelingToolset/Source/ModelingComponentsEditorOnly/Private/ToolTargets/SkeletalMeshComponentToolTarget.cpp
@@ -139,6 +139,11 @@ FDynamicMesh3 USkeletalMeshComponentReadOnlyToolTarget::GetDynamicMesh()
 	return GetDynamicMeshViaMeshDescription(*this);
 }
 
+FDynamicMesh3 USkeletalMeshComponentReadOnlyToolTarget::GetDynamicMesh(bool bRequestTangents)
+{
+	return GetDynamicMeshViaMeshDescription(*this, bRequestTangents);
+}
+
 
 USkeletalMesh* USkeletalMeshComponentReadOnlyToolTarget::GetSkeletalMesh() const
 {
diff --git a/Engine/Plugins/Runtime/MeshModelingToolset/Source/ModelingComponentsEditorOnly/Private/ToolTargets/SkeletalMeshToolTarget.cpp b/Engine/Plugins/Runtime/MeshModelingToolset/Source/ModelingComponentsEditorOnly/Private/ToolTargets/SkeletalMeshToolTarget.cpp
--- a/Engine/Plugins/Runtime/MeshModelingToolset/Source/ModelingComponentsEditorOnly/Private/ToolTargets/SkeletalMeshToolTarget.cpp
+++ b/Engine/Plugins/Runtime/MeshModelingToolset/Source/ModelingComponentsEditorOnly/Private/ToolTargets/SkeletalMeshToolTarget.cpp
@@ -160,6 +160,11 @@ FDynamicMesh3 USkeletalMeshReadOnlyToolTarget::GetDynamicMesh()
 	return GetDynamicMeshViaMeshDescription(*this);
 }
 
+FDynamicMesh3 USkeletalMeshReadOnlyToolTarget::GetDynamicMesh(bool bRequestTangents)
+{
+	return GetDynamicMeshViaMeshDescription(*this, bRequestTangents);
+}
+
 USkeletalMesh* USkeletalMeshReadOnlyToolTarget::GetSkeletalMesh() const
 {
 	return IsValid() ? SkeletalMesh.Get() : nullptr;
diff --git a/Engine/Plugins/Runtime/MeshModelingToolset/Source/ModelingComponentsEditorOnly/Private/ToolTargets/StaticMeshComponentToolTarget.cpp b/Engine/Plugins/Runtime/MeshModelingToolset/Source/ModelingComponentsEditorOnly/Private/ToolTargets/StaticMeshComponentToolTarget.cpp
--- a/Engine/Plugins/Runtime/MeshModelingToolset/Source/ModelingComponentsEditorOnly/Private/ToolTargets/StaticMeshComponentToolTarget.cpp
+++ b/Engine/Plugins/Runtime/MeshModelingToolset/Source/ModelingComponentsEditorOnly/Private/ToolTargets/StaticMeshComponentToolTarget.cpp
@@ -246,6 +246,11 @@ FDynamicMesh3 UStaticMeshComponentToolTarget::GetDynamicMesh()
 	return GetDynamicMeshViaMeshDescription(*this);
 }
 
+FDynamicMesh3 UStaticMeshComponentToolTarget::GetDynamicMesh(bool bRequestTangents)
+{
+	return GetDynamicMeshViaMeshDescription(*this, bRequestTangents);
+}
+
 void UStaticMeshComponentToolTarget::CommitDynamicMesh(const FDynamicMesh3& Mesh, const FDynamicMeshCommitInfo& CommitInfo)
 {
 	const FMeshDescription* CurrentMeshDescription = GetMeshDescription();
diff --git a/Engine/Plugins/Runtime/MeshModelingToolset/Source/ModelingComponentsEditorOnly/Private/ToolTargets/StaticMeshToolTarget.cpp b/Engine/Plugins/Runtime/MeshModelingToolset/Source/ModelingComponentsEditorOnly/Private/ToolTargets/StaticMeshToolTarget.cpp
--- a/Engine/Plugins/Runtime/MeshModelingToolset/Source/ModelingComponentsEditorOnly/Private/ToolTargets/StaticMeshToolTarget.cpp
+++ b/Engine/Plugins/Runtime/MeshModelingToolset/Source/ModelingComponentsEditorOnly/Private/ToolTargets/StaticMeshToolTarget.cpp
@@ -394,6 +394,11 @@ FDynamicMesh3 UStaticMeshToolTarget::GetDynamicMesh()
 	return GetDynamicMeshViaMeshDescription(*this);
 }
 
+FDynamicMesh3 UStaticMeshToolTarget::GetDynamicMesh(bool bRequestTangents)
+{
+	return GetDynamicMeshViaMeshDescription(*this, bRequestTangents);
+}
+
 void UStaticMeshToolTarget::CommitDynamicMesh(const FDynamicMesh3& Mesh, const FDynamicMeshCommitInfo& CommitInfo)
 {
 	FMeshDescription CurrentMeshDescription = *GetMeshDescription();
diff --git a/Engine/Plugins/Runtime/MeshModelingToolset/Source/ModelingComponentsEditorOnly/Public/ToolTargets/SkeletalMeshComponentToolTarget.h b/Engine/Plugins/Runtime/MeshModelingToolset/Source/ModelingComponentsEditorOnly/Public/ToolTargets/SkeletalMeshComponentToolTarget.h
--- a/Engine/Plugins/Runtime/MeshModelingToolset/Source/ModelingComponentsEditorOnly/Public/ToolTargets/SkeletalMeshComponentToolTarget.h
+++ b/Engine/Plugins/Runtime/MeshModelingToolset/Source/ModelingComponentsEditorOnly/Public/ToolTargets/SkeletalMeshComponentToolTarget.h
@@ -49,6 +49,7 @@ public:
 
 	// IDynamicMeshProvider
 	virtual UE::Geometry::FDynamicMesh3 GetDynamicMesh() override;
+	virtual UE::Geometry::FDynamicMesh3 GetDynamicMesh(bool bRequestTangents) override;
 
 	// ISkeletalMeshBackedTarget implementation
 	virtual USkeletalMesh* GetSkeletalMesh() const override;
diff --git a/Engine/Plugins/Runtime/MeshModelingToolset/Source/ModelingComponentsEditorOnly/Public/ToolTargets/SkeletalMeshToolTarget.h b/Engine/Plugins/Runtime/MeshModelingToolset/Source/ModelingComponentsEditorOnly/Public/ToolTargets/SkeletalMeshToolTarget.h
--- a/Engine/Plugins/Runtime/MeshModelingToolset/Source/ModelingComponentsEditorOnly/Public/ToolTargets/SkeletalMeshToolTarget.h
+++ b/Engine/Plugins/Runtime/MeshModelingToolset/Source/ModelingComponentsEditorOnly/Public/ToolTargets/SkeletalMeshToolTarget.h
@@ -49,6 +49,7 @@ public:
 
 	// IDynamicMeshProvider
 	virtual UE::Geometry::FDynamicMesh3 GetDynamicMesh() override;
+	virtual UE::Geometry::FDynamicMesh3 GetDynamicMesh(bool bRequestTangents) override;
 
 	// ISkeletalMeshBackedTarget implementation
 	virtual USkeletalMesh* GetSkeletalMesh() const override;
diff --git a/Engine/Plugins/Runtime/MeshModelingToolset/Source/ModelingComponentsEditorOnly/Public/ToolTargets/StaticMeshComponentToolTarget.h b/Engine/Plugins/Runtime/MeshModelingToolset/Source/ModelingComponentsEditorOnly/Public/ToolTargets/StaticMeshComponentToolTarget.h
--- a/Engine/Plugins/Runtime/MeshModelingToolset/Source/ModelingComponentsEditorOnly/Public/ToolTargets/StaticMeshComponentToolTarget.h
+++ b/Engine/Plugins/Runtime/MeshModelingToolset/Source/ModelingComponentsEditorOnly/Public/ToolTargets/StaticMeshComponentToolTarget.h
@@ -58,6 +58,7 @@ public:
 
 	// IDynamicMeshProvider
 	virtual UE::Geometry::FDynamicMesh3 GetDynamicMesh() override;
+	virtual UE::Geometry::FDynamicMesh3 GetDynamicMesh(bool bRequestTangents) override;
 
 	// IDynamicMeshCommitter
 	virtual void CommitDynamicMesh(const UE::Geometry::FDynamicMesh3& Mesh, const FDynamicMeshCommitInfo& CommitInfo) override;
diff --git a/Engine/Plugins/Runtime/MeshModelingToolset/Source/ModelingComponentsEditorOnly/Public/ToolTargets/StaticMeshToolTarget.h b/Engine/Plugins/Runtime/MeshModelingToolset/Source/ModelingComponentsEditorOnly/Public/ToolTargets/StaticMeshToolTarget.h
--- a/Engine/Plugins/Runtime/MeshModelingToolset/Source/ModelingComponentsEditorOnly/Public/ToolTargets/StaticMeshToolTarget.h
+++ b/Engine/Plugins/Runtime/MeshModelingToolset/Source/ModelingComponentsEditorOnly/Public/ToolTargets/StaticMeshToolTarget.h
@@ -62,6 +62,7 @@ public:
 
 	// IDynamicMeshProvider
 	virtual UE::Geometry::FDynamicMesh3 GetDynamicMesh() override;
+	virtual UE::Geometry::FDynamicMesh3 GetDynamicMesh(bool bRequestTangents) override;
 
 	// IDynamicMeshCommitter
 	virtual void CommitDynamicMesh(const UE::Geometry::FDynamicMesh3& Mesh, const FDynamicMeshCommitInfo& CommitInfo) override;
