Fix coordinate frame used for applying scale for the lattice control point gizmo

diff --git a/Engine/Plugins/Runtime/MeshModelingToolset/Source/ModelingComponents/Private/Mechanics/LatticeControlPointsMechanic.cpp b/Engine/Plugins/Runtime/MeshModelingToolset/Source/ModelingComponents/Private/Mechanics/LatticeControlPointsMechanic.cpp
--- a/Engine/Plugins/Runtime/MeshModelingToolset/Source/ModelingComponents/Private/Mechanics/LatticeControlPointsMechanic.cpp
+++ b/Engine/Plugins/Runtime/MeshModelingToolset/Source/ModelingComponents/Private/Mechanics/LatticeControlPointsMechanic.cpp
@@ -325,7 +325,6 @@ void ULatticeControlPointsMechanic::GizmoTransformChanged(UTransformProxy* Proxy
 	FVector DeltaScale = Transform.GetScale3D() / GizmoStartScale;
 
 	FTransformSRT3d DeltaTransform;
-	DeltaTransform.SetScale((FVector3d)DeltaScale);
 	DeltaTransform.SetRotation((FQuaterniond)DeltaRotation);
 	DeltaTransform.SetTranslation((FVector3d)Transform.GetTranslation());
 
@@ -335,6 +334,8 @@ void ULatticeControlPointsMechanic::GizmoTransformChanged(UTransformProxy* Proxy
 
 		// Translate to origin, scale, rotate, and translate back (DeltaTransform has "translate back" baked in.)
 		PointPosition -= (FVector3d)GizmoStartPosition;
+		// Align the scale to the gizmo orientation
+		PointPosition = GizmoStartRotation * (DeltaScale * (GizmoStartRotation.Inverse() * PointPosition));
 		PointPosition = DeltaTransform.TransformPosition(PointPosition);
 
 		ControlPoints[PointID] = PointPosition;
