on volume->dynamic mesh conversion, when normals are requested and groups are not requested, still keep the groups until they can be used to compute normals and then discard them.

this fixes an ensure on trying to place a new modeling tools primitive mesh in a level when a volume is in the level and volume snapping is enabled

diff --git a/Engine/Plugins/Runtime/MeshModelingToolset/Source/ModelingComponents/Private/ConversionUtils/VolumeToDynamicMesh.cpp b/Engine/Plugins/Runtime/MeshModelingToolset/Source/ModelingComponents/Private/ConversionUtils/VolumeToDynamicMesh.cpp
--- a/Engine/Plugins/Runtime/MeshModelingToolset/Source/ModelingComponents/Private/ConversionUtils/VolumeToDynamicMesh.cpp
+++ b/Engine/Plugins/Runtime/MeshModelingToolset/Source/ModelingComponents/Private/ConversionUtils/VolumeToDynamicMesh.cpp
@@ -92,7 +92,9 @@ void BrushToDynamicMesh(UModel& BrushModel, UE::Geometry::FDynamicMesh3& Mesh, c
 	using namespace VolumeToDynamicMeshLocals;
 
 	Mesh.Clear();
-	if (Options.bSetGroups)
+	// Note: Groups are also used to compute normals matching the input polygons
+	// If normals are requested but groups are not, groups will be discarded after normals are computed
+	if (Options.bSetGroups || Options.bGenerateNormals)
 	{
 		Mesh.EnableTriangleGroups();
 	}
@@ -225,6 +227,10 @@ void BrushToDynamicMesh(UModel& BrushModel, UE::Geometry::FDynamicMesh3& Mesh, c
 		FDynamicMeshNormalOverlay* Normals = Mesh.Attributes()->PrimaryNormals();
 		FMeshNormals::InitializeOverlayTopologyFromFaceGroups(&Mesh, Normals);
 		FMeshNormals::QuickRecomputeOverlayNormals(Mesh);
+		if (!Options.bSetGroups)
+		{
+			Mesh.DiscardTriangleGroups();
+		}
 	}
 }
 
