add a solid collision visualization toggle to the mesh to collision and inspect collision tools

diff --git a/Engine/Plugins/Experimental/MeshModelingToolsetExp/Source/MeshModelingToolsExp/Private/Physics/CollisionGeometryVisualization.cpp b/Engine/Plugins/Experimental/MeshModelingToolsetExp/Source/MeshModelingToolsExp/Private/Physics/CollisionGeometryVisualization.cpp
--- a/Engine/Plugins/Experimental/MeshModelingToolsetExp/Source/MeshModelingToolsExp/Private/Physics/CollisionGeometryVisualization.cpp
+++ b/Engine/Plugins/Experimental/MeshModelingToolsetExp/Source/MeshModelingToolsExp/Private/Physics/CollisionGeometryVisualization.cpp
@@ -1,8 +1,11 @@
 // Copyright Epic Games, Inc. All Rights Reserved.
 
 #include "Physics/CollisionGeometryVisualization.h"
-#include "Physics/PhysicsDataCollection.h"
+
+#include "DynamicMesh/DynamicMeshAttributeSet.h"
 #include "Physics/CollisionPropertySets.h"
+#include "Physics/ComponentCollisionUtil.h"
+#include "Physics/PhysicsDataCollection.h"
 #include "Generators/LineSegmentGenerators.h"
 #include "Drawing/PreviewGeometryActor.h"
 
@@ -11,6 +14,71 @@ using namespace UE::Geometry;
 namespace
 {
 
+void InitializePreviewGeometrySolid(
+	const FPhysicsDataCollection& PhysicsData,
+	UPreviewGeometry* PreviewGeom,
+	UMaterialInterface* SolidMaterial,
+	TFunctionRef<FColor(int32 GeoSetIndex)> GeoSetIndexToColorFunc,
+	bool bVisible,
+	int32 SphereStepsResolution,
+	int32 FirstGeoSetIndex)
+{
+	check(PreviewGeom);
+	check(SolidMaterial);
+
+	int32 CircleSteps = FMath::Max(4, SphereStepsResolution);
+	int32 GeoSetIndex = FirstGeoSetIndex;
+
+	const FKAggregateGeom& AggGeom = PhysicsData.AggGeom;
+	FSimpleCollisionTriangulationSettings TriangulationSettings;
+	TriangulationSettings.InitFromSphereResolution(CircleSteps);
+
+	UE::Geometry::ConvertSimpleCollisionToDynamicMeshes(
+		AggGeom,
+		[&](int32 ShapeIndex, const FKShapeElem& ShapeElem, const FDynamicMesh3& Mesh)
+		{
+			FColor Color = GeoSetIndexToColorFunc(ShapeIndex);
+			if (UTriangleSetComponent* TriangleSetComponent = PreviewGeom->CreateOrUpdateTriangleSet(FString::Printf(TEXT("Shape %d"), ShapeIndex), 1, [&](int32 Index, TArray<FRenderableTriangle>& TrisOut)
+				{
+					Mesh.TriangleCount();
+					check(Mesh.HasAttributes());
+					const UE::Geometry::FDynamicMeshNormalOverlay* PrimaryNormals = Mesh.Attributes()->PrimaryNormals();
+					const UE::Geometry::FDynamicMeshUVOverlay* PrimaryUV = Mesh.Attributes()->PrimaryUV();
+
+					TrisOut.Reserve(Mesh.TriangleCount());
+					for (int32 TID : Mesh.TriangleIndicesItr())
+					{
+						FIndex3i MeshTri = Mesh.GetTriangle(TID);
+						FRenderableTriangle& RenderTri = TrisOut.Emplace_GetRef();
+						RenderTri.Material = SolidMaterial;
+						FRenderableTriangleVertex* TriVerts[3]{ &RenderTri.Vertex0, &RenderTri.Vertex1, &RenderTri.Vertex2 };
+						for (int32 SubIdx = 0; SubIdx < 3; ++SubIdx)
+						{
+							int32 VID = MeshTri[SubIdx];
+							Mesh.GetVertex(VID);
+							FVector3f Normal;
+							PrimaryNormals->GetTriElement(TID, SubIdx, Normal);
+							FVector2f UV(0, 0);
+							if (PrimaryUV && PrimaryUV->IsSetTriangle(TID))
+							{
+								PrimaryUV->GetTriElement(TID, SubIdx, UV);
+							}
+							TriVerts[SubIdx]->Position = Mesh.GetVertex(VID);
+							TriVerts[SubIdx]->Color = Color;
+							TriVerts[SubIdx]->Normal = (FVector)Normal;
+							TriVerts[SubIdx]->UV = (FVector2D)UV;
+						}
+					}
+				}, Mesh.TriangleCount()))
+			{
+				TriangleSetComponent->SetVisibility(bVisible);
+			}
+		},
+		TriangulationSettings,
+		false /*bSetToPerTriangleNormals*/, false /*bInitializeConvexAndLevelSetUVs*/
+	);
+}
+
 void InitializePreviewGeometryLines(
 	const FPhysicsDataCollection& PhysicsData,
 	UPreviewGeometry* PreviewGeom,
@@ -33,7 +101,7 @@ void InitializePreviewGeometryLines(
 	// spheres are draw as 3 orthogonal circles
 	for (int32 Index = 0; Index < AggGeom.SphereElems.Num(); Index++)
 	{
-		PreviewGeom->CreateOrUpdateLineSet(FString::Printf(TEXT("Spheres %d"), Index), 1, [&](int32 UnusedIndex, TArray<FRenderableLine>& LinesOut)
+		if (ULineSetComponent* LineSetComponent = PreviewGeom->CreateOrUpdateLineSet(FString::Printf(TEXT("Spheres %d"), Index), 1, [&](int32 UnusedIndex, TArray<FRenderableLine>& LinesOut)
 		{
 			FColor Color = LineSetIndexToColorFunc(LineSetIndex++);
 
@@ -48,14 +116,18 @@ void InitializePreviewGeometryLines(
 				[&](const FVector3f& A, const FVector3f& B) { LinesOut.Add(FRenderableLine((FVector)A, (FVector)B, Color, LineThickness, DepthBias)); });
 			UE::Geometry::GenerateCircleSegments<float>(CircleSteps, Radius, FVector3f::Zero(), FVector3f::UnitY(), FVector3f::UnitZ(), ElemTransformf,
 				[&](const FVector3f& A, const FVector3f& B) { LinesOut.Add(FRenderableLine((FVector)A, (FVector)B, Color, LineThickness, DepthBias)); });
-		});
+		}))
+		{
+			LineSetComponent->SetVisibility(bVisible);
+			LineSetComponent->SetLineMaterial(LineMaterial);
+		}
 	}
 
 
 	// boxes are drawn as boxes
 	for (int32 Index = 0; Index < AggGeom.BoxElems.Num(); Index++)
 	{
-		PreviewGeom->CreateOrUpdateLineSet(FString::Printf(TEXT("Boxes %d"), Index), 1, [&](int32 UnusedIndex, TArray<FRenderableLine>& LinesOut)
+		if (ULineSetComponent* LineSetComponent = PreviewGeom->CreateOrUpdateLineSet(FString::Printf(TEXT("Boxes %d"), Index), 1, [&](int32 UnusedIndex, TArray<FRenderableLine>& LinesOut)
 		{
 			FColor Color = LineSetIndexToColorFunc(LineSetIndex++);
 
@@ -69,14 +141,18 @@ void InitializePreviewGeometryLines(
 				PhysicsData.ExternalScale3D.Z * Box.Z * 0.5f);
 			UE::Geometry::GenerateBoxSegments<float>(HalfDimensions, FVector3f::Zero(), FVector3f::UnitX(), FVector3f::UnitY(), FVector3f::UnitZ(), ElemTransformf,
 				[&](const FVector3f& A, const FVector3f& B) { LinesOut.Add(FRenderableLine((FVector)A, (FVector)B, Color, LineThickness, DepthBias)); });
-		});
+		}))
+		{
+			LineSetComponent->SetVisibility(bVisible);
+			LineSetComponent->SetLineMaterial(LineMaterial);
+		}
 	}
 
 
 	// capsules are draw as two hemispheres (with 3 intersecting arcs/circles) and connecting lines
 	for (int32 Index = 0; Index < AggGeom.SphylElems.Num(); Index++)
 	{
-		PreviewGeom->CreateOrUpdateLineSet(FString::Printf(TEXT("Capsules %d"), Index), 1, [&](int32 UnusedIndex, TArray<FRenderableLine>& LinesOut)
+		if (ULineSetComponent* LineSetComponent = PreviewGeom->CreateOrUpdateLineSet(FString::Printf(TEXT("Capsules %d"), Index), 1, [&](int32 UnusedIndex, TArray<FRenderableLine>& LinesOut)
 		{
 			FColor Color = LineSetIndexToColorFunc(LineSetIndex++);
 
@@ -118,13 +194,17 @@ void InitializePreviewGeometryLines(
 					ElemTransform.TransformPosition((FVector)Top + DY),
 					ElemTransform.TransformPosition((FVector)Bottom + DY), Color, LineThickness, DepthBias));
 			}
-		});
+		}))
+		{
+			LineSetComponent->SetVisibility(bVisible);
+			LineSetComponent->SetLineMaterial(LineMaterial);
+		}
 	}
 
 	// convexes are drawn as mesh edges
 	for (int32 Index = 0; Index < AggGeom.ConvexElems.Num(); Index++)
 	{
-		PreviewGeom->CreateOrUpdateLineSet(FString::Printf(TEXT("Convex %d"), Index), 1, [&](int32 UnusedIndex, TArray<FRenderableLine>& LinesOut)
+		if (ULineSetComponent* LineSetComponent = PreviewGeom->CreateOrUpdateLineSet(FString::Printf(TEXT("Convex %d"), Index), 1, [&](int32 UnusedIndex, TArray<FRenderableLine>& LinesOut)
 		{
 			FColor Color = LineSetIndexToColorFunc(LineSetIndex++);
 
@@ -142,13 +222,17 @@ void InitializePreviewGeometryLines(
 				LinesOut.Add(FRenderableLine(B, C, Color, LineThickness, DepthBias));
 				LinesOut.Add(FRenderableLine(C, A, Color, LineThickness, DepthBias));
 			}
-		});
+		}))
+		{
+			LineSetComponent->SetVisibility(bVisible);
+			LineSetComponent->SetLineMaterial(LineMaterial);
+		}
 	}
 
 	// for Level Sets draw the grid cells where phi < 0
 	for (int32 Index = 0; Index < AggGeom.LevelSetElems.Num(); Index++)
 	{
-		PreviewGeom->CreateOrUpdateLineSet(FString::Printf(TEXT("Level Set %d"), Index), 1, [&](int32 UnusedIndex, TArray<FRenderableLine>& LinesOut)
+		if (ULineSetComponent* LineSetComponent = PreviewGeom->CreateOrUpdateLineSet(FString::Printf(TEXT("Level Set %d"), Index), 1, [&](int32 UnusedIndex, TArray<FRenderableLine>& LinesOut)
 		{
 			FColor Color = LineSetIndexToColorFunc(LineSetIndex++);
 			const FKLevelSetElem& LevelSet = AggGeom.LevelSetElems[Index];
@@ -178,34 +262,44 @@ void InitializePreviewGeometryLines(
 				GenerateBoxSegmentsFromFBox(CellBox);
 			}
 
-		});
+		}))
+		{
+			LineSetComponent->SetVisibility(bVisible);
+			LineSetComponent->SetLineMaterial(LineMaterial);
+		}
 	}
 
 	// Unclear whether we actually use these in the Engine, for UBodySetup? Does not appear to be supported by UxX import system,
 	// and online documentation suggests they may only be supported for cloth?
 	ensure(AggGeom.TaperedCapsuleElems.Num() == 0);
-
-	PreviewGeom->SetAllLineSetsMaterial(LineMaterial);
-	PreviewGeom->SetAllVisible(bVisible);
 }
 
-void UpdatePreviewGeometryLines(
+void UpdatePreviewGeometryForCollision(
 	UPreviewGeometry* PartialPreviewGeom,
 	UCollisionGeometryVisualizationProperties* Settings,
-	int32 FirstLineSetIndex = 0)
+	int32 FirstColorIndex)
 {
 	check(PartialPreviewGeom);
 	check(Settings);
 
-	int32 LineSetIndex = FirstLineSetIndex;
+	int32 LineColorIndex = FirstColorIndex;
 	PartialPreviewGeom->UpdateAllLineSets([&](ULineSetComponent* LineSet)
 	{
-		FColor LineColor = Settings->GetLineSetColor(LineSetIndex++);
+		FColor LineColor = Settings->GetLineSetColor(LineColorIndex++);
 		LineSet->SetAllLinesColor(LineColor);
 		LineSet->SetAllLinesThickness(Settings->LineThickness);
 		LineSet->SetLineMaterial(Settings->GetLineMaterial());
 		LineSet->SetVisibility(Settings->bShowCollision);
 	});
+
+	int32 TriColorIndex = FirstColorIndex;
+	PartialPreviewGeom->UpdateAllTriangleSets([&](UTriangleSetComponent* TriangleSet)
+	{
+		FColor TriangleColor = Settings->GetTriangleSetColor(TriColorIndex++);
+		TriangleSet->SetAllTrianglesColor(TriangleColor);
+		TriangleSet->SetAllTrianglesMaterial(Settings->GetSolidMaterial());
+		TriangleSet->SetVisibility(Settings->bShowCollision && Settings->bShowSolid);
+	});
 }
 
 } // end namespace
@@ -220,11 +314,17 @@ void UE::PhysicsTools::InitializeCollisionGeometryVisualization(
 	UCollisionGeometryVisualizationProperties* Settings,
 	const FPhysicsDataCollection& PhysicsData,
 	float DepthBias,
-	int32 CircleStepResolution)
+	int32 CircleStepResolution,
+	bool bClearExistingLinesAndTriangles)
 {
 	check(PreviewGeom);
 	check(Settings);
 
+	if (bClearExistingLinesAndTriangles)
+	{
+		PreviewGeom->RemoveAllLineSets();
+		PreviewGeom->RemoveAllTriangleSets();
+	}
 	InitializePreviewGeometryLines(
 		PhysicsData,
 		PreviewGeom,
@@ -235,6 +335,11 @@ void UE::PhysicsTools::InitializeCollisionGeometryVisualization(
 		DepthBias,
 		CircleStepResolution,
 		0);
+	InitializePreviewGeometrySolid(
+		PhysicsData, PreviewGeom, Settings->GetSolidMaterial(), [&Settings](int SolidSetIndex) { return Settings->GetTriangleSetColor(SolidSetIndex); },
+		Settings->bShowCollision && Settings->bShowSolid,
+		CircleStepResolution,
+		0);
 
 	Settings->bVisualizationDirty = false;
 }
@@ -248,7 +353,7 @@ void UE::PhysicsTools::UpdateCollisionGeometryVisualization(
 
 	if (Settings->bVisualizationDirty)
 	{
-		UpdatePreviewGeometryLines(PreviewGeom, Settings);
+		UpdatePreviewGeometryForCollision(PreviewGeom, Settings, 0);
 		Settings->bVisualizationDirty = false;
 	}
 }
@@ -262,7 +367,7 @@ void UE::PhysicsTools::PartiallyInitializeCollisionGeometryVisualization(
 	UPreviewGeometry* PreviewGeom,
 	UCollisionGeometryVisualizationProperties* Settings,
 	const FPhysicsDataCollection& PhysicsData,
-	int32 FirstLineSetIndex,
+	int32 ColorIndex,
 	float DepthBias,
 	int32 CircleStepResolution)
 {
@@ -278,17 +383,24 @@ void UE::PhysicsTools::PartiallyInitializeCollisionGeometryVisualization(
 		Settings->bShowCollision,
 		DepthBias,
 		CircleStepResolution,
-		FirstLineSetIndex);
+		ColorIndex);
+
+	InitializePreviewGeometrySolid(
+		PhysicsData, PreviewGeom, Settings->GetSolidMaterial(), [&Settings](int SolidSetIndex) { return Settings->GetTriangleSetColor(SolidSetIndex); },
+		Settings->bShowCollision && Settings->bShowSolid,
+		CircleStepResolution,
+		ColorIndex);
 }
 
 void UE::PhysicsTools::PartiallyUpdateCollisionGeometryVisualization(
 	UPreviewGeometry* PartialPreviewGeom,
 	UCollisionGeometryVisualizationProperties* Settings,
-	int32 FirstLineSetIndex)
+	int32 ColorIndex)
 {
 	check(PartialPreviewGeom);
 	check(Settings);
 
-	UpdatePreviewGeometryLines(PartialPreviewGeom, Settings, FirstLineSetIndex);
+	// Note: Solid geometry uses the same coloring as line sets
+	UpdatePreviewGeometryForCollision(PartialPreviewGeom, Settings, ColorIndex);
 }
 
diff --git a/Engine/Plugins/Experimental/MeshModelingToolsetExp/Source/MeshModelingToolsExp/Private/Physics/CollisionPropertySets.cpp b/Engine/Plugins/Experimental/MeshModelingToolsetExp/Source/MeshModelingToolsExp/Private/Physics/CollisionPropertySets.cpp
--- a/Engine/Plugins/Experimental/MeshModelingToolsetExp/Source/MeshModelingToolsExp/Private/Physics/CollisionPropertySets.cpp
+++ b/Engine/Plugins/Experimental/MeshModelingToolsetExp/Source/MeshModelingToolsExp/Private/Physics/CollisionPropertySets.cpp
@@ -18,9 +18,11 @@ void UCollisionGeometryVisualizationProperties::Initialize(UInteractiveTool* Too
 	WatchProperty(Color, [this](FColor NewValue) { bVisualizationDirty = true; });
 	WatchProperty(bRandomColors, [this](bool bNewValue) { bVisualizationDirty = true; });
 	WatchProperty(bShowHidden, [this](bool bNewValue) { bVisualizationDirty = true; });
+	WatchProperty(bShowSolid, [this](bool bNewValue) { bVisualizationDirty = true; });
 
 	LineMaterial = ToolSetupUtil::GetDefaultLineComponentMaterial(Tool->GetToolManager(), true);
 	LineMaterialShowingHidden = ToolSetupUtil::GetDefaultLineComponentMaterial(Tool->GetToolManager(), false);
+	TriangleMaterial = ToolSetupUtil::GetVertexColorMaterial(Tool->GetToolManager(), false);
 }
 
 FColor UCollisionGeometryVisualizationProperties::GetLineSetColor(int32 LineSetIndex) const
diff --git a/Engine/Plugins/Experimental/MeshModelingToolsetExp/Source/MeshModelingToolsExp/Private/Physics/SetCollisionGeometryTool.cpp b/Engine/Plugins/Experimental/MeshModelingToolsetExp/Source/MeshModelingToolsExp/Private/Physics/SetCollisionGeometryTool.cpp
--- a/Engine/Plugins/Experimental/MeshModelingToolsetExp/Source/MeshModelingToolsExp/Private/Physics/SetCollisionGeometryTool.cpp
+++ b/Engine/Plugins/Experimental/MeshModelingToolsetExp/Source/MeshModelingToolsExp/Private/Physics/SetCollisionGeometryTool.cpp
@@ -598,8 +598,6 @@ void USetCollisionGeometryTool::OnTick(float DeltaTime)
 				VizSettings->bVisualizationDirty = true;
 
 				// update visualization
-				PreviewGeom->RemoveAllLineSets();
-
 				UE::PhysicsTools::InitializeCollisionGeometryVisualization(PreviewGeom, VizSettings, *GeneratedCollision);
 
 				// update property set
diff --git a/Engine/Plugins/Experimental/MeshModelingToolsetExp/Source/MeshModelingToolsExp/Public/Physics/CollisionGeometryVisualization.h b/Engine/Plugins/Experimental/MeshModelingToolsetExp/Source/MeshModelingToolsExp/Public/Physics/CollisionGeometryVisualization.h
--- a/Engine/Plugins/Experimental/MeshModelingToolsetExp/Source/MeshModelingToolsExp/Public/Physics/CollisionGeometryVisualization.h
+++ b/Engine/Plugins/Experimental/MeshModelingToolsetExp/Source/MeshModelingToolsExp/Public/Physics/CollisionGeometryVisualization.h
@@ -12,15 +12,17 @@ class UMaterialInterface;
 namespace UE::PhysicsTools
 {
 	/**
-	 * Create line sets in a UPreviewGeometry for all the elements in a Physics Data Collection.
-	 * Spheres and Capsules are drawn as 3-axis wireframes. Convexes are added as wireframes.
+	 * Create triangle sets (solids) and line sets (wireframes) in a UPreviewGeometry for all the elements in a Physics Data Collection.
+	 * In line sets, Spheres and Capsules are drawn as 3-axis wireframes. Convexes are added as wireframes.
+	 * By default, clear existing lines and triangles from the preview
 	 */
 	void MESHMODELINGTOOLSEXP_API InitializeCollisionGeometryVisualization(
 		UPreviewGeometry* PreviewGeom,
 		UCollisionGeometryVisualizationProperties* Settings,
 		const FPhysicsDataCollection& PhysicsData,
 		float DepthBias = 0.f,
-		int32 CircleStepResolution = 16);
+		int32 CircleStepResolution = 16,
+		bool bClearExistingLinesAndTriangles = true);
 
 	/**
 	 * Update line sets in a UPreviewGeometry.
@@ -35,29 +37,29 @@ namespace UE::PhysicsTools
 	/**
 	 * Like InitializeCollisionGeometryVisualization but only initializes part of the collision geometry visualization.
 	 * 
-	 * The FirstLineSetIndex parameter is used implement Settings->bRandomColors since random colors are determined
-	 * using unique integers for each line set, which may be stored in separate UPreviewGeometry objects. Since the
+	 * The ColorIndex parameter is used to implement Settings->bRandomColors since pseudorandom colors are determined
+	 * using unique integers for each line/triangle set, which may be stored in separate UPreviewGeometry objects. Since the
 	 * initialization is partial this function doesn't set Settings->bVisualizationDirty, the caller should do that.
 	 */
 	void MESHMODELINGTOOLSEXP_API PartiallyInitializeCollisionGeometryVisualization(
 		UPreviewGeometry* PartialPreviewGeom,
 		UCollisionGeometryVisualizationProperties* Settings,
 		const FPhysicsDataCollection& PhysicsData,
-		int32 FirstLineSetIndex,
+		int32 ColorIndex,
 		float DepthBias = 0.f,
 		int32 CircleStepResolution = 16);
 
 	/**
 	 * Like UpdateCollisionGeometryVisualization but only updates part of the collision geometry visualization.
 	 * 
-	 * The FirstLineSetIndex parameter is used implement Settings->bRandomColors since random colors are determined
-	 * using unique integers for each line set, which may be stored in separate UPreviewGeometry objects. Since the
+	 * The ColorIndex parameter is used implement Settings->bRandomColors since pseudorandom colors are determined
+	 * using unique integers for each line/triangle set, which may be stored in separate UPreviewGeometry objects. Since the
 	 * update is partial this function doesn't check or set Settings->bVisualizationDirty, the caller should do that. 
 	 */
 	void MESHMODELINGTOOLSEXP_API PartiallyUpdateCollisionGeometryVisualization(
 		UPreviewGeometry* PartialPreviewGeom,
 		UCollisionGeometryVisualizationProperties* Settings,
-		int32 FirstLineSetIndex);
+		int32 ColorIndex);
 }
 
 
diff --git a/Engine/Plugins/Experimental/MeshModelingToolsetExp/Source/MeshModelingToolsExp/Public/Physics/CollisionPropertySets.h b/Engine/Plugins/Experimental/MeshModelingToolsetExp/Source/MeshModelingToolsExp/Public/Physics/CollisionPropertySets.h
--- a/Engine/Plugins/Experimental/MeshModelingToolsetExp/Source/MeshModelingToolsExp/Public/Physics/CollisionPropertySets.h
+++ b/Engine/Plugins/Experimental/MeshModelingToolsetExp/Source/MeshModelingToolsExp/Public/Physics/CollisionPropertySets.h
@@ -149,29 +149,46 @@ public:
 
 	/**
 	 * Returns the line set color.
-	 * If bRandomColors == true  LineSetIndex is used to return a random color
+	 * If bRandomColors == true  LineSetIndex is used to choose the color
 	 * If bRandomColors == false LineSetIndex is ignored and this->Color is returned
 	 */
 	FColor GetLineSetColor(int32 LineSetIndex = 0) const;
 
+	/** @return The color used for solid rendering. Currently uses the same colors as the lines. */
+	FColor GetTriangleSetColor(int32 TriangleSetIndex = 0) const
+	{
+		return GetLineSetColor(TriangleSetIndex);
+	}
+
 	UMaterialInterface* GetLineMaterial() const
 	{
 		ensure(LineMaterial);
 		ensure(LineMaterialShowingHidden);
 		return bShowHidden ? LineMaterialShowingHidden : LineMaterial;
 	}
+
+	/** @return the material to use for triangle rendering, when bShowSolid is true */
+	UMaterialInterface* GetSolidMaterial() const
+	{
+		ensure(TriangleMaterial);
+		return TriangleMaterial;
+	}
 	
 	/** Show/hide collision geometry */
 	UPROPERTY(EditAnywhere, Category = "Collision Visualization",
 		meta = (EditCondition = "bEnableShowCollision", EditConditionHides, HideEditConditionToggle))
 	bool bShowCollision = true;
 
+	/** Whether to show solid shapes in addition to wireframes */
+	UPROPERTY(EditAnywhere, Category = "Collision Visualization")
+	bool bShowSolid = false;
+
 	/** Thickness of lines used to visualize collision shapes */
 	UPROPERTY(EditAnywhere, Category = "Collision Visualization")
 	float LineThickness = 3.0f;
 
 	/** Show occluded parts of the collision geometry, rendered with dashed lines */
-	UPROPERTY(EditAnywhere, Category = "Collision Visualization")
+	UPROPERTY(EditAnywhere, DisplayName = "Show Hidden Lines", Category = "Collision Visualization")
 	bool bShowHidden = false;
 
 	/** Render each collision geometry with a randomly-chosen color */
@@ -191,6 +208,10 @@ public:
 	UPROPERTY(Transient, meta=(TransientToolProperty))
 	TObjectPtr<UMaterialInterface> LineMaterialShowingHidden = nullptr;
 
+	//~Used if bShowHidden is false
+	UPROPERTY(Transient, meta = (TransientToolProperty))
+	TObjectPtr<UMaterialInterface> TriangleMaterial = nullptr;
+
 	//~Some tools will want showing collision geometry to be non-optional
 	UPROPERTY(Transient, meta=(TransientToolProperty))
 	bool bEnableShowCollision = true;
diff --git a/Engine/Plugins/Runtime/MeshModelingToolset/Source/ModelingComponents/Private/Drawing/PreviewGeometryActor.cpp b/Engine/Plugins/Runtime/MeshModelingToolset/Source/ModelingComponents/Private/Drawing/PreviewGeometryActor.cpp
--- a/Engine/Plugins/Runtime/MeshModelingToolset/Source/ModelingComponents/Private/Drawing/PreviewGeometryActor.cpp
+++ b/Engine/Plugins/Runtime/MeshModelingToolset/Source/ModelingComponents/Private/Drawing/PreviewGeometryActor.cpp
@@ -113,7 +113,7 @@ UTriangleSetComponent* UPreviewGeometry::FindTriangleSet(const FString& Triangle
 
 
 
-void UPreviewGeometry::CreateOrUpdateTriangleSet(const FString& TriangleSetIdentifier, int32 NumIndices,
+UTriangleSetComponent* UPreviewGeometry::CreateOrUpdateTriangleSet(const FString& TriangleSetIdentifier, int32 NumIndices,
 	TFunctionRef<void(int32 Index, TArray<FRenderableTriangle>& TrianglesOut)> TriangleGenFunc,
 	int32 TrianglesPerIndexHint)
 {
@@ -124,12 +124,13 @@ void UPreviewGeometry::CreateOrUpdateTriangleSet(const FString& TriangleSetIdent
 		if (TriangleSet == nullptr)
 		{
 			check(false);
-			return;
+			return nullptr;
 		}
 	}
 
 	TriangleSet->Clear();
 	TriangleSet->AddTriangles(NumIndices, TriangleGenFunc, TrianglesPerIndexHint);
+	return TriangleSet;
 }
 
 
@@ -180,6 +181,7 @@ bool UPreviewGeometry::RemoveLineSet(const FString& LineSetIdentifier, bool bDes
 			LineSet->DestroyComponent();
 			LineSet = nullptr;
 		}
+		return true;
 	}
 	return false;
 }
@@ -200,6 +202,41 @@ void UPreviewGeometry::RemoveAllLineSets(bool bDestroy)
 }
 
 
+bool UPreviewGeometry::RemoveTriangleSet(const FString& TriangleSetIdentifier, bool bDestroy)
+{
+	TObjectPtr<UTriangleSetComponent>* Found = TriangleSets.Find(TriangleSetIdentifier);
+	if (Found != nullptr)
+	{
+		UTriangleSetComponent* TriangleSet = *Found;
+
+		TriangleSets.Remove(TriangleSetIdentifier);
+		if (bDestroy)
+		{
+			TriangleSet->UnregisterComponent();
+			TriangleSet->DestroyComponent();
+			TriangleSet = nullptr;
+		}
+		return true;
+	}
+	return false;
+}
+
+
+
+void UPreviewGeometry::RemoveAllTriangleSets(bool bDestroy)
+{
+	if (bDestroy)
+	{
+		for (TPair<FString, TObjectPtr<UTriangleSetComponent>> Entry : TriangleSets)
+		{
+			Entry.Value->UnregisterComponent();
+			Entry.Value->DestroyComponent();
+		}
+	}
+	TriangleSets.Reset();
+}
+
+
 
 bool UPreviewGeometry::SetLineSetVisibility(const FString& LineSetIdentifier, bool bVisible)
 {
@@ -235,7 +272,7 @@ void UPreviewGeometry::SetAllLineSetsMaterial(UMaterialInterface* Material)
 
 
 
-void UPreviewGeometry::CreateOrUpdateLineSet(const FString& LineSetIdentifier, int32 NumIndices,
+ULineSetComponent* UPreviewGeometry::CreateOrUpdateLineSet(const FString& LineSetIdentifier, int32 NumIndices,
 	TFunctionRef<void(int32 Index, TArray<FRenderableLine>& LinesOut)> LineGenFunc,
 	int32 LinesPerIndexHint)
 {
@@ -246,12 +283,13 @@ void UPreviewGeometry::CreateOrUpdateLineSet(const FString& LineSetIdentifier, i
 		if (LineSet == nullptr)
 		{
 			check(false);
-			return;
+			return nullptr;
 		}
 	}
 
 	LineSet->Clear();
 	LineSet->AddLines(NumIndices, LineGenFunc, LinesPerIndexHint);
+	return LineSet;
 }
 
 
@@ -320,6 +358,7 @@ bool UPreviewGeometry::RemovePointSet(const FString& PointSetIdentifier, bool bD
 			PointSet->DestroyComponent();
 			PointSet = nullptr;
 		}
+		return true;
 	}
 	return false;
 }
diff --git a/Engine/Plugins/Runtime/MeshModelingToolset/Source/ModelingComponents/Private/Physics/ComponentCollisionUtil.cpp b/Engine/Plugins/Runtime/MeshModelingToolset/Source/ModelingComponents/Private/Physics/ComponentCollisionUtil.cpp
--- a/Engine/Plugins/Runtime/MeshModelingToolset/Source/ModelingComponents/Private/Physics/ComponentCollisionUtil.cpp
+++ b/Engine/Plugins/Runtime/MeshModelingToolset/Source/ModelingComponents/Private/Physics/ComponentCollisionUtil.cpp
@@ -3,15 +3,18 @@
 #include "Physics/ComponentCollisionUtil.h"
 
 #include "DynamicMesh/DynamicMesh3.h"
-#include "ShapeApproximation/SimpleShapeSet3.h"
 #include "DynamicMeshEditor.h"
 #include "DynamicMesh/MeshNormals.h"
+#include "DynamicMesh/MeshTransforms.h"
+#include "DynamicMesh/Operations/MergeCoincidentMeshEdges.h"
 #include "Generators/SphereGenerator.h"
+#include "Generators/BoxSphereGenerator.h"
+#include "Generators/MarchingCubes.h"
 #include "Generators/MinimalBoxMeshGenerator.h"
 #include "Generators/CapsuleGenerator.h"
-#include "DynamicMesh/MeshTransforms.h"
 #include "Parameterization/DynamicMeshUVEditor.h"
-#include "DynamicMesh/Operations/MergeCoincidentMeshEdges.h"
+#include "ShapeApproximation/SimpleShapeSet3.h"
+#include "VectorUtil.h"
 
 #include "Physics/PhysicsDataCollection.h"
 #include "Physics/CollisionGeometryConversion.h"
@@ -22,6 +25,7 @@
 #include "Engine/StaticMesh.h"
 #include "PhysicsEngine/AggregateGeom.h"
 #include "PhysicsEngine/BodySetup.h"
+#include "PhysicsEngine/ShapeElem.h"
 #include "UObject/UObjectIterator.h"
 
 using namespace UE::Geometry;
@@ -320,157 +324,248 @@ bool UE::Geometry::AppendSimpleCollision(
 	return true;
 }
 
-
-void UE::Geometry::ConvertSimpleCollisionToMeshes(
-	const FKAggregateGeom& AggGeom,
-	FDynamicMesh3& MeshOut,
-	const FTransformSequence3d& TransformSeqeuence,
-	int32 SphereResolution,
-	bool bSetToPerTriangleNormals,
-	bool bInitializeConvexUVs,
-	TFunction<void(int, const FDynamicMesh3&)> PerElementMeshCallback )
+namespace UE::Private::ConvertCollisionInternal
 {
-	FDynamicMeshEditor Editor(&MeshOut);
-
-	bool bTransformInverts = TransformSeqeuence.WillInvert();
-
-	for (const FKSphereElem& Sphere : AggGeom.SphereElems)
-	{
-		FSphereGenerator SphereGen;
-		SphereGen.Radius = Sphere.Radius;
-		SphereGen.NumPhi = SphereGen.NumTheta = SphereResolution;
-		SphereGen.bPolygroupPerQuad = false;
-		SphereGen.Generate();
-		FDynamicMesh3 SphereMesh(&SphereGen);
-
-		MeshTransforms::Translate(SphereMesh, FVector3d(Sphere.Center));
+	static void ConvertSimpleCollisionToMeshesHelper(
+		const FKAggregateGeom& AggGeom,
+		FSimpleCollisionTriangulationSettings TriangulationSettings,
+		bool bSetToPerTriangleNormals,
+		bool bInitializeLevelSetAndConvexUVs,
+		TFunctionRef<void(int32 Index, const FKShapeElem& ShapeElem, const FDynamicMesh3&)> DynamicMeshCallback)
+	{
+		// Helpers to transform vertex and normal arrays / invert triangle buffers
+		auto TransformVerts = [](FTransform Transform, TArrayView<FVector3d> Vertices, TArrayView<FVector3f> Normals = TArrayView<FVector3f>())
+		{
+			for (FVector3d& V : Vertices)
+			{
+				V = Transform.TransformPosition(V);
+			}
+			for (FVector3f& N : Normals)
+			{
+				N = (FVector3f)VectorUtil::TransformNormal<double>(Transform, (FVector3d)N);
+			}
+		};
+		auto TranslateVerts = [](FVector3d Offset, TArrayView<FVector3d> Vertices)
+		{
+			for (FVector3d& V : Vertices)
+			{
+				V += Offset;
+			}
+		};
+		auto InvertTris = [](TArrayView<FIndex3i> Tris)
+		{
+			for (FIndex3i& Tri : Tris)
+			{
+				Swap(Tri.B, Tri.C);
+			}
+		};
 
-		if (bTransformInverts)
+		auto RunCallbacks = 
+			[&DynamicMeshCallback, bSetToPerTriangleNormals]
+			(int32 Index, const FKShapeElem& ShapeElem, FMeshShapeGenerator& Gen)
 		{
-			SphereMesh.ReverseOrientation(false);
+			FDynamicMesh3 ShapeMesh(&Gen);
+			if (bSetToPerTriangleNormals)
+			{
+				FMeshNormals::InitializeMeshToPerTriangleNormals(&ShapeMesh);
+			} // otherwise allow the generator's normals to pass through
+			DynamicMeshCallback(Index, ShapeElem, ShapeMesh);
+		};
+			
+		int32 ShapeIndex = 0;
+		for (const FKSphereElem& Sphere : AggGeom.SphereElems)
+		{
+			if (TriangulationSettings.bUseBoxSphere)
+			{
+				FBoxSphereGenerator BoxSphereGen;
+				BoxSphereGen.Box.Frame.Origin = Sphere.Center;
+				BoxSphereGen.Radius = FMath::Max(FMathf::ZeroTolerance, Sphere.Radius);
+				int32 StepsPerSide = FMath::Max(1, TriangulationSettings.BoxSphereStepsPerSide);
+				BoxSphereGen.EdgeVertices = FIndex3i(StepsPerSide, StepsPerSide, StepsPerSide);
+				BoxSphereGen.Generate();
+				RunCallbacks(ShapeIndex++, Sphere, BoxSphereGen);
+			}
+			else
+			{
+				FSphereGenerator LatLongSphereGen;
+				LatLongSphereGen.Radius = FMath::Max(FMathf::ZeroTolerance, Sphere.Radius);
+				LatLongSphereGen.NumPhi = LatLongSphereGen.NumTheta = FMath::Max(1, TriangulationSettings.LatLongSphereSteps);
+				LatLongSphereGen.bPolygroupPerQuad = false;
+				LatLongSphereGen.Generate();
+				TranslateVerts(FVector3d(Sphere.Center), LatLongSphereGen.Vertices);
+				RunCallbacks(ShapeIndex++, Sphere, LatLongSphereGen);
+			}
 		}
 
-		FMeshIndexMappings Mappings;
-		Editor.AppendMesh(&SphereMesh, Mappings,
-			[&](int32 vid, const FVector3d& P) { return TransformSeqeuence.TransformPosition(P); },
-			[&](int32 vid, const FVector3d& N) { return TransformSeqeuence.TransformNormal(N); });
+		for (const FKBoxElem& Box : AggGeom.BoxElems)
+		{
+			FMinimalBoxMeshGenerator BoxGen;
+			BoxGen.Box = UE::Geometry::FOrientedBox3d(
+				FFrame3d(FVector3d(Box.Center), FQuaterniond(Box.Rotation.Quaternion())),
+				0.5 * FVector3d(Box.X, Box.Y, Box.Z));
+			BoxGen.Generate();
+			RunCallbacks(ShapeIndex++, Box, BoxGen);
+		}
 
-		if (PerElementMeshCallback)
+		for (const FKSphylElem& Capsule : AggGeom.SphylElems)
 		{
-			PerElementMeshCallback(0, SphereMesh);
+			FCapsuleGenerator CapsuleGen;
+			CapsuleGen.Radius = Capsule.Radius;
+			CapsuleGen.SegmentLength = Capsule.Length;
+			CapsuleGen.NumHemisphereArcSteps = TriangulationSettings.CapsuleHemisphereSteps;
+			CapsuleGen.NumCircleSteps = TriangulationSettings.CapsuleCircleSteps;
+			CapsuleGen.bPolygroupPerQuad = false;
+			CapsuleGen.Generate();
+			TranslateVerts(FVector3d(0, 0, -0.5 * Capsule.Length), CapsuleGen.Vertices);
+			FTransform Transform = Capsule.GetTransform();
+			// Note: Capsule transform cannot invert; it is generated from a rotation and translation
+			TransformVerts(Transform, CapsuleGen.Vertices, CapsuleGen.Normals);
+			RunCallbacks(ShapeIndex++, Capsule, CapsuleGen);
 		}
-	}
 
-	for (const FKBoxElem& Box : AggGeom.BoxElems)
-	{
-		FMinimalBoxMeshGenerator BoxGen;
-		BoxGen.Box = UE::Geometry::FOrientedBox3d(
-			FFrame3d(FVector3d(Box.Center), FQuaterniond(Box.Rotation.Quaternion())),
-			0.5*FVector3d(Box.X, Box.Y, Box.Z));
-		BoxGen.Generate();
-		FDynamicMesh3 BoxMesh(&BoxGen);
+		for (const FKConvexElem& Convex : AggGeom.ConvexElems)
+		{
+			FTransform ElemTransform = Convex.GetTransform();
+			const int32 NumVertices = Convex.VertexData.Num();
+			const int32 NumTris = Convex.IndexData.Num() / 3;
+			bool bInvert = ElemTransform.GetDeterminant() < 0;
 
-		// transform not applied because it is just the Center/Rotation
+			// Create the convex hull mesh from the FKConvexElem vertex/index buffer data
 
-		if (bTransformInverts)
-		{
-			BoxMesh.ReverseOrientation(false);
-		}
+			FDynamicMesh3 ConvexMesh(EMeshComponents::None);
+			for (int32 k = 0; k < NumVertices; ++k)
+			{
+				ConvexMesh.AppendVertex(ElemTransform.TransformPosition(FVector3d(Convex.VertexData[k])));
+			}
+			int32 NumTriangles = Convex.IndexData.Num() / 3;
+			for (int32 k = 0; k < NumTriangles; ++k)
+			{
+				// Note: Invert tri indices because chaos stores them inverted
+				ConvexMesh.AppendTriangle(Convex.IndexData[3 * k], Convex.IndexData[3 * k + 2], Convex.IndexData[3 * k + 1]);
+			}
+			if (bInvert)
+			{
+				ConvexMesh.ReverseOrientation(false /*bFlipNormals*/);
+			}
+
+			ConvexMesh.EnableTriangleGroups(0);
+			ConvexMesh.EnableAttributes();
+			if (bSetToPerTriangleNormals)
+			{
+				FMeshNormals::InitializeMeshToPerTriangleNormals(&ConvexMesh);
+			}
+			else
+			{
+				FMeshNormals::InitializeOverlayToPerVertexNormals(ConvexMesh.Attributes()->PrimaryNormals(), false);
+			}
+			if (bInitializeLevelSetAndConvexUVs)
+			{
+				FDynamicMeshUVEditor UVEditor(&ConvexMesh, 0, true);
+				UVEditor.SetPerTriangleUVs();
+			}
+			DynamicMeshCallback(ShapeIndex, Convex, ConvexMesh);
 
-		FMeshIndexMappings Mappings;
-		Editor.AppendMesh(&BoxMesh, Mappings,
-			[&](int32 vid, const FVector3d& P) { return TransformSeqeuence.TransformPosition(P); },
-			[&](int32 vid, const FVector3d& N) { return TransformSeqeuence.TransformNormal(N); });	
+			ShapeIndex++;
+		}
 
-		if (PerElementMeshCallback)
+		TArray<FVector3f> LevelSetVerts;
+		TArray<FIntVector> LevelSetTris;
+		for (const FKLevelSetElem& LevelSet : AggGeom.LevelSetElems)
 		{
-			PerElementMeshCallback(1, BoxMesh);
-		}
-	}
+			FTransform ElemTransform = LevelSet.GetTransform();
+			bool bInvert = ElemTransform.GetDeterminant() < 0;
 
+			LevelSetVerts.Reset();
+			LevelSetTris.Reset();
+			LevelSet.GetZeroIsosurfaceGridCellFaces(LevelSetVerts, LevelSetTris);
 
-	for (const FKSphylElem& Capsule: AggGeom.SphylElems)
-	{
-		FCapsuleGenerator CapsuleGen;
-		CapsuleGen.Radius = Capsule.Radius;
-		CapsuleGen.SegmentLength = Capsule.Length;
-		CapsuleGen.NumHemisphereArcSteps = SphereResolution/4+1;
-		CapsuleGen.NumCircleSteps = SphereResolution;
-		CapsuleGen.bPolygroupPerQuad = false;
-		CapsuleGen.Generate();
-		FDynamicMesh3 CapsuleMesh(&CapsuleGen);
-
-		MeshTransforms::Translate(CapsuleMesh, FVector3d(0,0,-0.5*Capsule.Length) );
-
-		FTransformSRT3d Transform(Capsule.GetTransform());
-		MeshTransforms::ApplyTransform(CapsuleMesh, Transform, false);
-		if (Transform.GetDeterminant() < 0 != bTransformInverts)
-		{
-			CapsuleMesh.ReverseOrientation(false);
-		}
+			FDynamicMesh3 LevelSetMesh(EMeshComponents::None);
+			for (int32 Idx = 0; Idx < LevelSetVerts.Num(); ++Idx)
+			{
+				LevelSetMesh.AppendVertex(ElemTransform.TransformPosition(FVector3d(LevelSetVerts[Idx])));
+			}
+			for (int32 Idx = 0; Idx < LevelSetTris.Num(); ++Idx)
+			{
+				FIntVector Tri = LevelSetTris[Idx];
+				LevelSetMesh.AppendTriangle(Tri.X, Tri.Y, Tri.Z);
+			}
+			if (bInvert)
+			{
+				LevelSetMesh.ReverseOrientation(false /*bFlipNormals*/);
+			}
 
-		FMeshIndexMappings Mappings;
-		Editor.AppendMesh(&CapsuleMesh, Mappings,
-			[&](int32 vid, const FVector3d& P) { return TransformSeqeuence.TransformPosition(P); },
-			[&](int32 vid, const FVector3d& N) { return TransformSeqeuence.TransformNormal(N); });	
+			LevelSetMesh.EnableTriangleGroups(0);
+			LevelSetMesh.EnableAttributes();
+			if (bSetToPerTriangleNormals)
+			{
+				FMeshNormals::InitializeMeshToPerTriangleNormals(&LevelSetMesh);
+			}
+			else
+			{
+				FMeshNormals::InitializeOverlayToPerVertexNormals(LevelSetMesh.Attributes()->PrimaryNormals(), false);
+			}
+			if (bInitializeLevelSetAndConvexUVs)
+			{
+				FDynamicMeshUVEditor UVEditor(&LevelSetMesh, 0, true);
+				UVEditor.SetPerTriangleUVs();
+			}
+			DynamicMeshCallback(ShapeIndex, LevelSet, LevelSetMesh);
 
-		if (PerElementMeshCallback)
-		{
-			PerElementMeshCallback(2, CapsuleMesh);
+			ShapeIndex++;
 		}
 	}
+}
 
+void UE::Geometry::ConvertSimpleCollisionToDynamicMeshes(
+	const FKAggregateGeom& AggGeom,
+	TFunctionRef<void(int32, const FKShapeElem&, const FDynamicMesh3&)> PerElementMeshCallback,
+	FSimpleCollisionTriangulationSettings TriangulationSettings,
+	bool bSetToPerTriangleNormals,
+	bool bInitializeConvexAndLevelSetUVs)
+{
+	UE::Private::ConvertCollisionInternal::ConvertSimpleCollisionToMeshesHelper(
+		AggGeom,
+		TriangulationSettings,
+		bSetToPerTriangleNormals,
+		bInitializeConvexAndLevelSetUVs,
+		PerElementMeshCallback);
+}
 
-	for (const FKConvexElem& Convex : AggGeom.ConvexElems)
-	{
-		FTransformSRT3d ElemTransform(Convex.GetTransform());
-		FDynamicMesh3 ConvexMesh(EMeshComponents::None);
-		int32 NumVertices = Convex.VertexData.Num();
-		for (int32 k = 0; k < NumVertices; ++k)
-		{
-			ConvexMesh.AppendVertex(ElemTransform.TransformPosition(FVector3d(Convex.VertexData[k])) );
-		}
-		int32 NumTriangles = Convex.IndexData.Num() / 3;
-		for (int32 k = 0; k < NumTriangles; ++k)
-		{
-			ConvexMesh.AppendTriangle(Convex.IndexData[3*k], Convex.IndexData[3*k+1], Convex.IndexData[3*k+2]);
-		}
-		
-		// Note we need to reverse the orientation if no transform inverts, or if both invert,
-		// because ConvexMesh has reversed-orientation triangles normally
-		if (ElemTransform.GetDeterminant() < 0 == bTransformInverts)
-		{
-			ConvexMesh.ReverseOrientation();
-		}
-		ConvexMesh.EnableTriangleGroups(0);
-		ConvexMesh.EnableAttributes();
-		if (bInitializeConvexUVs)
-		{
-			FDynamicMeshUVEditor UVEditor(&ConvexMesh, 0, true);
-			UVEditor.SetPerTriangleUVs();
-		}
 
-		FMeshIndexMappings Mappings;
-		Editor.AppendMesh(&ConvexMesh, Mappings,
-			[&](int32 vid, const FVector3d& P) { return TransformSeqeuence.TransformPosition(P); },
-			[&](int32 vid, const FVector3d& N) { return TransformSeqeuence.TransformNormal(N); });	
+void UE::Geometry::ConvertSimpleCollisionToMeshes(
+	const FKAggregateGeom& AggGeom,
+	FDynamicMesh3& MeshOut,
+	const FTransformSequence3d& TransformSequence,
+	int32 SphereResolution,
+	bool bSetToPerTriangleNormals,
+	bool bInitializeConvexUVs,
+	TFunction<void(int, const FDynamicMesh3&)> PerElementMeshCallback )
+{
+	FDynamicMeshEditor Editor(&MeshOut);
+
+	bool bTransformInverts = TransformSequence.WillInvert();
 
-		if (PerElementMeshCallback)
-		{
-			PerElementMeshCallback(3, ConvexMesh);
-		}
-	}
+	FSimpleCollisionTriangulationSettings TriangulationSettings;
+	TriangulationSettings.InitFromSphereResolution(SphereResolution);
 
-	if (MeshOut.HasAttributes() && MeshOut.Attributes()->PrimaryNormals() != nullptr)
-	{
-		if (bSetToPerTriangleNormals)
+	FMeshIndexMappings Mappings;
+	UE::Private::ConvertCollisionInternal::ConvertSimpleCollisionToMeshesHelper(
+		AggGeom,
+		TriangulationSettings,
+		bSetToPerTriangleNormals,
+		bInitializeConvexUVs,
+		[&](int32 Index, const FKShapeElem& Elem, const UE::Geometry::FDynamicMesh3& Mesh)
 		{
-			FMeshNormals::InitializeMeshToPerTriangleNormals(&MeshOut);
-		}
-		else
-		{
-			FMeshNormals::InitializeOverlayToPerVertexNormals(MeshOut.Attributes()->PrimaryNormals(), false);
+			PerElementMeshCallback((int)Elem.GetShapeType(), Mesh);
+			Mappings.Reset();
+			Editor.AppendMesh(&Mesh, Mappings,
+				[&TransformSequence](int32 vid, const FVector3d& P) { return TransformSequence.TransformPosition(P); },
+				[&TransformSequence](int32 vid, const FVector3d& N) { return TransformSequence.TransformNormal(N); });
 		}
+	);
+	if (bTransformInverts)
+	{
+		MeshOut.ReverseOrientation(false);
 	}
 }
 
diff --git a/Engine/Plugins/Runtime/MeshModelingToolset/Source/ModelingComponents/Public/Drawing/PreviewGeometryActor.h b/Engine/Plugins/Runtime/MeshModelingToolset/Source/ModelingComponents/Public/Drawing/PreviewGeometryActor.h
--- a/Engine/Plugins/Runtime/MeshModelingToolset/Source/ModelingComponents/Public/Drawing/PreviewGeometryActor.h
+++ b/Engine/Plugins/Runtime/MeshModelingToolset/Source/ModelingComponents/Public/Drawing/PreviewGeometryActor.h
@@ -107,13 +107,41 @@ public:
 		}
 	}
 
+	/**
+	 * call UpdateFunc for all existing Triangle Sets
+	 */
+	template<typename UpdateFuncType>
+	void UpdateAllTriangleSets(UpdateFuncType UpdateFunc)
+	{
+		for (TPair<FString, TObjectPtr<UTriangleSetComponent>> Entry : TriangleSets)
+		{
+			UpdateFunc(Entry.Value);
+		}
+	}
+
 	/**
 	 * Add a set of triangles produced by calling TriangleGenFunc for each index in range [0,NumIndices)
+	 * @return the created or updated triangle set
 	 */
-	void CreateOrUpdateTriangleSet(const FString& TriangleSetIdentifier, int32 NumIndices,
+	UTriangleSetComponent* CreateOrUpdateTriangleSet(const FString& TriangleSetIdentifier, int32 NumIndices,
 		TFunctionRef<void(int32 Index, TArray<FRenderableTriangle>& TrianglesOut)> TriangleGenFunc,
 		int32 TrianglesPerIndexHint = -1);
 
+	/**
+	 * Remove the TriangleSetComponent with the given TriangleSetIdentifier
+	 * @param bDestroy if true, component will unregistered and destroyed.
+	 * @return true if the TriangleSetComponent was found and removed
+	 */
+	UFUNCTION()
+	bool RemoveTriangleSet(const FString& TriangleSetIdentifier, bool bDestroy = true);
+
+	/**
+	 * Remove all TriangleSetComponents
+	 * @param bDestroy if true, the components will unregistered and destroyed.
+	 */
+	UFUNCTION()
+	void RemoveAllTriangleSets(bool bDestroy = true);
+
 
 	//
 	// Line Sets
@@ -195,8 +223,9 @@ public:
 
 	/**
 	 * Add a set of lines produced by calling LineGenFunc for each index in range [0,NumIndices)
+	 * @return the created or updated line set
 	 */
-	void CreateOrUpdateLineSet(const FString& LineSetIdentifier, int32 NumIndices,
+	ULineSetComponent* CreateOrUpdateLineSet(const FString& LineSetIdentifier, int32 NumIndices,
 		TFunctionRef<void(int32 Index, TArray<FRenderableLine>& LinesOut)> LineGenFunc,
 		int32 LinesPerIndexHint = -1);
 
diff --git a/Engine/Plugins/Runtime/MeshModelingToolset/Source/ModelingComponents/Public/Physics/ComponentCollisionUtil.h b/Engine/Plugins/Runtime/MeshModelingToolset/Source/ModelingComponents/Public/Physics/ComponentCollisionUtil.h
--- a/Engine/Plugins/Runtime/MeshModelingToolset/Source/ModelingComponents/Public/Physics/ComponentCollisionUtil.h
+++ b/Engine/Plugins/Runtime/MeshModelingToolset/Source/ModelingComponents/Public/Physics/ComponentCollisionUtil.h
@@ -117,20 +117,56 @@ MODELINGCOMPONENTS_API UBodySetup* GetBodySetup(UPrimitiveComponent* SourceCompo
 
 /**
  * Extract the simple collision geometry from AggGeom as meshes (ie spheres and capsules are tessellated) and
- * accumulate into MeshOut. TransformSequence is applied to the meshes as they are accumulated.
+ * accumulate into MeshOut. TransformSequence is applied to the meshes as they are accumulated, and not to the meshes passed to the PerElementMeshCallback.
  * @param bSetToPerTriangleNormals if true and the mesh has a normals attribute overlay,  the mesh is set to face normals, otherwise averaged vertex normals
  * @param bInitializeConvexUVs if true convex hulls have their UVs initialized to per-face planar projections, otherwise no UVs are set
- * @param PerElementMeshCallback if provided, called with each element mesh before transforming/appending to MeshOut. The int parameter is the shape type, 0=Sphere, 1=Box, 2=Capsule, 3=Convex
+ * @param PerElementMeshCallback if provided, called with each element mesh before transforming/appending to MeshOut. The int parameter is the shape type (EAggCollisionShape::Type), 0=Sphere, 1=Box, 2=Capsule, 3=Convex
  */
 MODELINGCOMPONENTS_API void ConvertSimpleCollisionToMeshes(
 	const FKAggregateGeom& AggGeom,
 	UE::Geometry::FDynamicMesh3& MeshOut,
-	const FTransformSequence3d& TransformSeqeuence,
+	const FTransformSequence3d& TransformSequence,
 	int32 SphereResolution = 16,
 	bool bSetToPerTriangleNormals = false,
 	bool bInitializeConvexUVs = false,
 	TFunction<void(int, const FDynamicMesh3&)> PerElementMeshCallback = nullptr );
 
+// Settings to define how simple collision shapes are triangulated
+struct FSimpleCollisionTriangulationSettings
+{
+	// Steps to use per side if bUseBoxSphere is true
+	int32 BoxSphereStepsPerSide = 5;
+	// Steps to use along circumference if bUseBoxSphere is false
+	int32 LatLongSphereSteps = 16;
+	// Steps to use along the capsule hemisphere arc
+	int32 CapsuleHemisphereSteps = 5;
+	// Steps to use radially on the capsule
+	int32 CapsuleCircleSteps = 16;
+
+	// Whether to use a box-sphere triangulation for spheres (otherwise, uses a lat/long sphere triangulation)
+	bool bUseBoxSphere = false;
+
+	// TODO: Add options for controlling level set triangulation as well
+
+	// Init sphere-related settings from a single resolution parameter
+	void InitFromSphereResolution(int32 SphereResolution)
+	{
+		bUseBoxSphere = false;
+		LatLongSphereSteps = SphereResolution;
+		CapsuleHemisphereSteps = SphereResolution / 4 + 1;
+		CapsuleCircleSteps = SphereResolution;
+	}
+};
+
+// Similar to ConvertSimpleCollisionToMeshes but without aggregating to an output mesh.
+// @param PerElementMeshCallback	Called with each shape element and the corresponding mesh
+MODELINGCOMPONENTS_API void ConvertSimpleCollisionToDynamicMeshes(
+	const FKAggregateGeom& AggGeom,
+	TFunctionRef<void(int32 Index, const FKShapeElem&, const FDynamicMesh3& Mesh)> PerElementMeshCallback,
+	FSimpleCollisionTriangulationSettings TriangulationSettings,
+	bool bSetToPerTriangleNormals = false,
+	bool bInitializeConvexAndLevelSetUVs = false);
+
 /**
  * Extract the complex collision geometry from CollisionDataProvider as a dynamic mesh.
  * TransformSequence is applied to the mesh.
