Make the modeling components' scene snapping volumes filter also filter out shape components, so e.g. a trigger box is filtered (or not) the same way as trigger volume

diff --git a/Engine/Plugins/Runtime/MeshModelingToolset/Source/ModelingComponents/Private/Snapping/ModelingSceneSnappingManager.cpp b/Engine/Plugins/Runtime/MeshModelingToolset/Source/ModelingComponents/Private/Snapping/ModelingSceneSnappingManager.cpp
--- a/Engine/Plugins/Runtime/MeshModelingToolset/Source/ModelingComponents/Private/Snapping/ModelingSceneSnappingManager.cpp
+++ b/Engine/Plugins/Runtime/MeshModelingToolset/Source/ModelingComponents/Private/Snapping/ModelingSceneSnappingManager.cpp
@@ -14,10 +14,11 @@
 #include "GameFramework/Actor.h"
 #include "EngineUtils.h" // for TActorIterator<>
 #include "Engine/StaticMesh.h"
-#include "Components/PrimitiveComponent.h"
-#include "Components/StaticMeshComponent.h"
 #include "Components/BrushComponent.h"
 #include "Components/DynamicMeshComponent.h"
+#include "Components/PrimitiveComponent.h"
+#include "Components/StaticMeshComponent.h"
+#include "Components/ShapeComponent.h"
 #include "RawIndexBuffer.h"
 #include "StaticMeshResources.h"
 #include "UObject/UObjectGlobals.h"
@@ -339,7 +340,8 @@ static bool FindNearestVisibleObjectHit_Internal(
 		}
 
 		// filtering out any volume hits here will disable volume snapping
-		if (bEnableVolumes == false && Cast<UBrushComponent>(CurResult.GetComponent()) != nullptr)
+		// shape components are also commonly used to represent simple volumes, so if we are filtering out volume hits, we filter out both brush and shape components
+		if (bEnableVolumes == false && (Cast<UBrushComponent>(CurResult.GetComponent()) != nullptr || Cast<UShapeComponent>(CurResult.GetComponent()) != nullptr))
 		{
 			continue;
 		}
