Texture preventative fixes; clarify some mip locking code and fix 64 bit issues

#rnx

diff --git a/Engine/Plugins/FX/Niagara/Source/Niagara/Private/NiagaraDataInterfaceCurveBase.cpp b/Engine/Plugins/FX/Niagara/Source/Niagara/Private/NiagaraDataInterfaceCurveBase.cpp
--- a/Engine/Plugins/FX/Niagara/Source/Niagara/Private/NiagaraDataInterfaceCurveBase.cpp
+++ b/Engine/Plugins/FX/Niagara/Source/Niagara/Private/NiagaraDataInterfaceCurveBase.cpp
@@ -386,7 +386,8 @@ void UNiagaraDataInterfaceCurveBase::UpdateExposedTexture()
 		ExposedTexture->AddressY = TA_Clamp;
 		ExposedTexture->LODGroup = TEXTUREGROUP_EffectsNotFiltered;
 	}
-
+	
+	ExposedTexture->PreEditChange(nullptr);
 	FFloat16Color* TexData = reinterpret_cast<FFloat16Color*>(ExposedTexture->Source.LockMip(0));
 
 	const int32 NumElements = GetCurveNumElems();
@@ -402,7 +403,8 @@ void UNiagaraDataInterfaceCurveBase::UpdateExposedTexture()
 	}
 
 	ExposedTexture->Source.UnlockMip(0);
-	ExposedTexture->Source.UseHashAsGuid();
+	// UseHashAsGuid is done by UnlockMip
+	//ExposedTexture->Source.UseHashAsGuid();
 	ExposedTexture->PostEditChange();
 
 	// PostEditChange() will assign a random GUID to the texture, which leads to non-deterministic builds.
diff --git a/Engine/Plugins/Runtime/MeshModelingToolset/Source/ModelingComponents/Private/AssetUtils/Texture2DBuilder.cpp b/Engine/Plugins/Runtime/MeshModelingToolset/Source/ModelingComponents/Private/AssetUtils/Texture2DBuilder.cpp
--- a/Engine/Plugins/Runtime/MeshModelingToolset/Source/ModelingComponents/Private/AssetUtils/Texture2DBuilder.cpp
+++ b/Engine/Plugins/Runtime/MeshModelingToolset/Source/ModelingComponents/Private/AssetUtils/Texture2DBuilder.cpp
@@ -314,7 +314,7 @@ void FTexture2DBuilder::Clear(const FColor& ClearColor)
 {
 	if (ensure(IsEditable() && IsByteTexture()))
 	{
-		int32 Num = Dimensions.Num();
+		int64 Num = Dimensions.Num();
 		for (int64 k = 0; k < Num; ++k)
 		{
 			CurrentMipData[k] = ClearColor;
@@ -329,7 +329,7 @@ void FTexture2DBuilder::Clear(const FFloat16Color& ClearColor)
 {
 	if (ensure(IsEditable() && IsFloat16Texture()))
 	{
-		int32 Num = Dimensions.Num();
+		int64 Num = Dimensions.Num();
 		for (int64 k = 0; k < Num; ++k)
 		{
 			CurrentMipDataFloat16[k] = ClearColor;
@@ -350,14 +350,11 @@ bool FTexture2DBuilder::Copy(const TImageBuilder<FVector3f>& SourceImage, const
 	}
 
 	int64 Num = Dimensions.Num();
-	for (int32 i = 0; i < Num; ++i)
+	for (int64 i = 0; i < Num; ++i)
 	{
 		FVector3f Pixel = SourceImage.GetPixel(i);
 		if (IsByteTexture())
 		{
-			Pixel.X = FMathf::Clamp(Pixel.X, 0.0, 1.0);
-			Pixel.Y = FMathf::Clamp(Pixel.Y, 0.0, 1.0);
-			Pixel.Z = FMathf::Clamp(Pixel.Z, 0.0, 1.0);
 			FColor Texel = ToLinearColor(Pixel).ToFColor(bConvertToSRGB);
 			SetTexel(i, Texel);
 		}
@@ -383,15 +380,11 @@ bool FTexture2DBuilder::Copy(const TImageBuilder<FVector4f>& SourceImage, const
 	}
 
 	int64 Num = Dimensions.Num();
-	for (int32 i = 0; i < Num; ++i)
+	for (int64 i = 0; i < Num; ++i)
 	{
 		FVector4f Pixel = SourceImage.GetPixel(i);
 		if (IsByteTexture())
 		{
-			Pixel.X = FMathf::Clamp(Pixel.X, 0.0, 1.0);
-			Pixel.Y = FMathf::Clamp(Pixel.Y, 0.0, 1.0);
-			Pixel.Z = FMathf::Clamp(Pixel.Z, 0.0, 1.0);
-			Pixel.W = FMathf::Clamp(Pixel.W, 0.0, 1.0);
 			FColor Texel = ToLinearColor(Pixel).ToFColor(bConvertToSRGB);
 			SetTexel(i, Texel);
 		}
@@ -436,13 +429,9 @@ bool FTexture2DBuilder::CopyImageToSourceData(const TImageBuilder<FVector4f>& So
 	if (SourceDataFormat == TSF_BGRA8)
 	{
 		FColor* DestData = reinterpret_cast<FColor*>(RawTexture2D->Source.LockMip(0));
-		for (int32 i = 0; i < Num; ++i)
+		for (int64 i = 0; i < Num; ++i)
 		{
 			FVector4f Pixel = SourceImage.GetPixel(i);
-			Pixel.X = FMathf::Clamp(Pixel.X, 0.0, 1.0);
-			Pixel.Y = FMathf::Clamp(Pixel.Y, 0.0, 1.0);
-			Pixel.Z = FMathf::Clamp(Pixel.Z, 0.0, 1.0);
-			Pixel.W = FMathf::Clamp(Pixel.W, 0.0, 1.0);
 			const FColor Texel = ToLinearColor(Pixel).ToFColor(bConvertToSRGB);
 			DestData[i] = Texel;
 		}
@@ -450,7 +439,7 @@ bool FTexture2DBuilder::CopyImageToSourceData(const TImageBuilder<FVector4f>& So
 	else if (SourceDataFormat == TSF_RGBA16F)
 	{
 		FFloat16Color* DestData = reinterpret_cast<FFloat16Color*>(RawTexture2D->Source.LockMip(0));
-		for (int32 i = 0; i < Num; ++i)
+		for (int64 i = 0; i < Num; ++i)
 		{
 			FVector4f Pixel = SourceImage.GetPixel(i);
 			const FFloat16Color Texel(ToLinearColor(Pixel));
@@ -472,7 +461,7 @@ bool FTexture2DBuilder::CopyTo(TImageBuilder<FVector4f>& DestImage) const
 		return false;
 	}
 	int64 Num = Dimensions.Num();
-	for (int32 i = 0; i < Num; ++i)
+	for (int64 i = 0; i < Num; ++i)
 	{
 		if (IsByteTexture())
 		{
diff --git a/Engine/Source/Editor/UnrealEd/Private/WorldPartition/WorldPartitionMiniMapBuilder.cpp b/Engine/Source/Editor/UnrealEd/Private/WorldPartition/WorldPartitionMiniMapBuilder.cpp
--- a/Engine/Source/Editor/UnrealEd/Private/WorldPartition/WorldPartitionMiniMapBuilder.cpp
+++ b/Engine/Source/Editor/UnrealEd/Private/WorldPartition/WorldPartitionMiniMapBuilder.cpp
@@ -139,6 +139,7 @@ bool UWorldPartitionMiniMapBuilder::RunInternal(UWorld* World, const FCellInfo&
 		const uint8* SrcDataPtr = TileTexture->Source.LockMipReadOnly(0);
 		check(SrcDataPtr);
 
+		// PreEditChange was called before in PreRun
 		uint8* const MiniMapDstPtr = WorldMiniMap->MiniMapTexture->Source.LockMip(0);
 		check(MiniMapDstPtr);
 
@@ -202,7 +203,7 @@ bool UWorldPartitionMiniMapBuilder::PostRun(UWorld* World, FPackageSourceControl
 		WorldMiniMap->MiniMapTexture->AdjustMinAlpha = 1.f;
 		WorldMiniMap->MiniMapTexture->LODGroup = TEXTUREGROUP_UI;
 		WorldMiniMap->MiniMapTexture->VirtualTextureStreaming = true;
-		WorldMiniMap->MiniMapTexture->UpdateResource();
+		//WorldMiniMap->MiniMapTexture->UpdateResource(); // @@?? <- UpdateResource before PostEditChange looks wrong
 		WorldMiniMap->MiniMapTexture->PostEditChange();
 	}
 
diff --git a/Engine/Source/Runtime/Engine/Private/ImageUtils.cpp b/Engine/Source/Runtime/Engine/Private/ImageUtils.cpp
--- a/Engine/Source/Runtime/Engine/Private/ImageUtils.cpp
+++ b/Engine/Source/Runtime/Engine/Private/ImageUtils.cpp
@@ -640,8 +640,9 @@ UTexture2D* FImageUtils::CreateTexture2D(int32 SrcWidth, int32 SrcHeight, const
 	uint8* MipData = Tex2D->Source.LockMip(0);
 	for( int32 y=0; y<SrcHeight; y++ )
 	{
-		uint8* DestPtr = &MipData[(SrcHeight - 1 - y) * SrcWidth * sizeof(FColor)];
-		const FColor* SrcPtr = &SrcData[(SrcHeight - 1 - y) * SrcWidth];
+		// when UseAlpha is true, this could/should just be a memcpy of the whole array
+		uint8* DestPtr = &MipData[(int64) y * SrcWidth * sizeof(FColor)];
+		const FColor* SrcPtr = &SrcData[(int64) y * SrcWidth];
 		for( int32 x=0; x<SrcWidth; x++ )
 		{
 			*DestPtr++ = SrcPtr->B;
diff --git a/Engine/Source/Runtime/GeometryCore/Public/Image/ImageDimensions.h b/Engine/Source/Runtime/GeometryCore/Public/Image/ImageDimensions.h
--- a/Engine/Source/Runtime/GeometryCore/Public/Image/ImageDimensions.h
+++ b/Engine/Source/Runtime/GeometryCore/Public/Image/ImageDimensions.h
@@ -67,7 +67,7 @@ public:
 	/** @return linear index into image from 2D coordinates */
 	int64 GetIndex(int32 X, int32 Y) const
 	{
-		return (Y * Width) + X;
+		return (int64) Y * Width + X;
 	}
 
 	/** @return linear index into image from 2D coordinates */
diff --git a/Engine/Source/Runtime/Landscape/Private/LandscapeDataAccess.cpp b/Engine/Source/Runtime/Landscape/Private/LandscapeDataAccess.cpp
--- a/Engine/Source/Runtime/Landscape/Private/LandscapeDataAccess.cpp
+++ b/Engine/Source/Runtime/Landscape/Private/LandscapeDataAccess.cpp
@@ -114,12 +114,13 @@ LANDSCAPE_API bool FLandscapeComponentDataInterface::GetWeightmapTextureData(ULa
 	OutData.Empty(FMath::Square(WeightmapSize));
 	OutData.AddUninitialized(FMath::Square(WeightmapSize));
 
-	FColor* WeightMipData = (FColor*)DataInterface.LockMip(ComponentWeightmapTextures[ComponentWeightmapLayerAllocations[LayerIdx].WeightmapTextureIndex], MipLevel);
+	// DataInterface Lock is a LockMipReadOnly on the texture
+	const FColor* WeightMipData = (const FColor*)DataInterface.LockMip(ComponentWeightmapTextures[ComponentWeightmapLayerAllocations[LayerIdx].WeightmapTextureIndex], MipLevel);
 
 	// Channel remapping
 	int32 ChannelOffsets[4] = { (int32)STRUCT_OFFSET(FColor, R), (int32)STRUCT_OFFSET(FColor, G), (int32)STRUCT_OFFSET(FColor, B), (int32)STRUCT_OFFSET(FColor, A) };
 
-	uint8* SrcTextureData = (uint8*)WeightMipData + ChannelOffsets[ComponentWeightmapLayerAllocations[LayerIdx].WeightmapTextureChannel];
+	const uint8* SrcTextureData = (const uint8*)WeightMipData + ChannelOffsets[ComponentWeightmapLayerAllocations[LayerIdx].WeightmapTextureChannel];
 
 	for (int32 i = 0; i < FMath::Square(WeightmapSize); i++)
 	{
