better simple collision level set visualization: add optional marching-cubes-based level set triangulation to the simple-collision->mesh utility functions and use the new marching cubes triangulation in tools + optionally in geometry script's append simple collision function.

also do some cleanup / generalization of the utility functions and change geometry script's append simple collision function to use them.

diff --git a/Engine/Plugins/Experimental/GeometryScripting/Source/GeometryScriptingCore/Private/MeshPrimitiveFunctions.cpp b/Engine/Plugins/Experimental/GeometryScripting/Source/GeometryScriptingCore/Private/MeshPrimitiveFunctions.cpp
--- a/Engine/Plugins/Experimental/GeometryScripting/Source/GeometryScriptingCore/Private/MeshPrimitiveFunctions.cpp
+++ b/Engine/Plugins/Experimental/GeometryScripting/Source/GeometryScriptingCore/Private/MeshPrimitiveFunctions.cpp
@@ -23,6 +23,8 @@
 #include "CompGeom/Delaunay2.h"
 #include "CompGeom/ConvexDecomposition3.h"
 
+#include "Physics/ComponentCollisionUtil.h"
+
 #include UE_INLINE_GENERATED_CPP_BY_NAME(MeshPrimitiveFunctions)
 
 using namespace UE::Geometry;
@@ -1515,61 +1517,29 @@ UDynamicMesh* UGeometryScriptLibrary_MeshPrimitiveFunctions::AppendSimpleCollisi
 		return TargetMesh;
 	}
 
-	for (const FKBoxElem& Box : SimpleCollision.AggGeom.BoxElems)
-	{
-		FGridBoxMeshGenerator GridBoxGenerator;
-		GridBoxGenerator.Box.Extents = FVector(Box.X * .5, Box.Y * .5, Box.Z * .5);
-		GridBoxGenerator.Box.Frame.Origin = Box.Center;
-		GridBoxGenerator.Box.Frame.Rotation = (FQuaterniond)Box.Rotation;
-		GridBoxGenerator.EdgeVertices = FIndex3i(2, 2, 2);
-		GridBoxGenerator.bPolygroupPerQuad = (PrimitiveOptions.PolygroupMode == EGeometryScriptPrimitivePolygroupMode::PerQuad);
-		GridBoxGenerator.Generate();
-		AppendPrimitive(TargetMesh, &GridBoxGenerator, Transform, PrimitiveOptions);
-	}
-	for (const FKSphereElem& Sphere : SimpleCollision.AggGeom.SphereElems)
-	{
-		FBoxSphereGenerator SphereGenerator;
-		SphereGenerator.Box.Frame.Origin = Sphere.Center;
-		SphereGenerator.Radius = FMath::Max(FMathf::ZeroTolerance, Sphere.Radius);
-		int32 StepsPerSide = FMath::Max(1, TriangulationOptions.SphereStepsPerSide);
-		SphereGenerator.EdgeVertices = FIndex3i(StepsPerSide, StepsPerSide, StepsPerSide);
-		SphereGenerator.bPolygroupPerQuad = (PrimitiveOptions.PolygroupMode == EGeometryScriptPrimitivePolygroupMode::PerQuad);
-		SphereGenerator.Generate();
-		AppendPrimitive(TargetMesh, &SphereGenerator, Transform, PrimitiveOptions);
-	}
-	for (const FKSphylElem& Capsule : SimpleCollision.AggGeom.SphylElems)
-	{
-		FCapsuleGenerator CapsuleGenerator;
-		CapsuleGenerator.Radius = Capsule.Radius;
-		CapsuleGenerator.SegmentLength = Capsule.Length;
-		CapsuleGenerator.NumHemisphereArcSteps = FMath::Max(2, TriangulationOptions.CapsuleHemisphereSteps);
-		CapsuleGenerator.NumCircleSteps = FMath::Max(3, TriangulationOptions.CapsuleCircleSteps);
-		CapsuleGenerator.bPolygroupPerQuad = (PrimitiveOptions.PolygroupMode == EGeometryScriptPrimitivePolygroupMode::PerQuad);
-		CapsuleGenerator.Generate();
-		
-		FQuaterniond Rotation(Capsule.Rotation);
-		AppendPrimitive(TargetMesh, &CapsuleGenerator, Transform, PrimitiveOptions, Capsule.Center + Rotation * FVector(0, 0, -Capsule.Length*.5), Rotation);
-	}
-	for (const FKConvexElem& Convex : SimpleCollision.AggGeom.ConvexElems)
-	{
-		FDynamicMesh3 ConvexMesh;
-		FTransform ConvexTransform = Convex.GetTransform();
-		for (FVector V : Convex.VertexData)
-		{
-			ConvexMesh.AppendVertex(ConvexTransform.TransformPosition(V));
-		}
-		for (int32 TriStart = 0; TriStart + 2 < Convex.IndexData.Num(); TriStart += 3)
+	FSimpleCollisionTriangulationSettings UseTriOptions;
+	UseTriOptions.bUseBoxSphere = true;
+	UseTriOptions.BoxSphereStepsPerSide = TriangulationOptions.SphereStepsPerSide;
+	UseTriOptions.CapsuleCircleSteps = TriangulationOptions.CapsuleCircleSteps;
+	UseTriOptions.CapsuleHemisphereSteps = TriangulationOptions.CapsuleHemisphereSteps;
+	UseTriOptions.bApproximateLevelSetWithCubes = TriangulationOptions.bApproximateLevelSetsWithCubes;
+
+	FSimpleCollisionToMeshAttributeSettings UseAttributeOptions(true /*enable attributes*/, false /*per-triangle normals*/, true /*initialize UVs*/);
+
+	ConvertSimpleCollisionToDynamicMeshes(
+		SimpleCollision.AggGeom,
+		[&](int32 Index, const FKShapeElem& ShapeElem, FDynamicMesh3& Mesh)
 		{
-			ConvexMesh.AppendTriangle(Convex.IndexData[TriStart], Convex.IndexData[TriStart + 2], Convex.IndexData[TriStart + 1]);
-		}
-		ConvexMesh.EnableAttributes();
-		FMeshNormals::InitializeOverlayToPerTriangleNormals(ConvexMesh.Attributes()->PrimaryNormals());
-		AppendPrimitiveMesh(TargetMesh, ConvexMesh, Transform, PrimitiveOptions);
-	}
-	if (!SimpleCollision.AggGeom.TaperedCapsuleElems.IsEmpty() || !SimpleCollision.AggGeom.SkinnedLevelSetElems.IsEmpty() || !SimpleCollision.AggGeom.LevelSetElems.IsEmpty())
+			AppendPrimitiveMesh(TargetMesh, Mesh, Transform, PrimitiveOptions);
+		},
+		UseTriOptions,
+		UseAttributeOptions
+	);
+
+	if (!SimpleCollision.AggGeom.TaperedCapsuleElems.IsEmpty() || !SimpleCollision.AggGeom.SkinnedLevelSetElems.IsEmpty())
 	{
 		// Tapered capsules and level sets are not supported yet
-		UE::Geometry::AppendWarning(Debug, EGeometryScriptErrorType::OperationFailed, LOCTEXT("PrimitiveFunctions_AppendSimpleCollisionShapes Tapered Capsules and Level Sets Unsupported", "AppendSimpleCollisionShapes: Tapered Capsules and Level Sets are not supported and will be skipped"));
+		UE::Geometry::AppendWarning(Debug, EGeometryScriptErrorType::OperationFailed, LOCTEXT("PrimitiveFunctions_AppendSimpleCollisionShapes Tapered Capsules Unsupported", "AppendSimpleCollisionShapes: Tapered Capsules are not supported and will be skipped"));
 	}
 
 	return TargetMesh;
diff --git a/Engine/Plugins/Experimental/GeometryScripting/Source/GeometryScriptingCore/Public/GeometryScript/GeometryScriptTypes.h b/Engine/Plugins/Experimental/GeometryScripting/Source/GeometryScriptingCore/Public/GeometryScript/GeometryScriptTypes.h
--- a/Engine/Plugins/Experimental/GeometryScripting/Source/GeometryScriptingCore/Public/GeometryScript/GeometryScriptTypes.h
+++ b/Engine/Plugins/Experimental/GeometryScripting/Source/GeometryScriptingCore/Public/GeometryScript/GeometryScriptTypes.h
@@ -162,14 +162,22 @@ struct GEOMETRYSCRIPTINGCORE_API FGeometryScriptSimpleCollisionTriangulationOpti
 	GENERATED_BODY()
 public:
 
+	// When triangulating a sphere by deforming a cube to the sphere, number of vertices to use along each edge of the cube
 	UPROPERTY(EditAnywhere, BlueprintReadWrite, Category = Options)
 	int32 SphereStepsPerSide = 4;
 
+	// When triangulating a capsule's spherical endcaps, number of vertices to use on the arcs across the endcaps.
 	UPROPERTY(EditAnywhere, BlueprintReadWrite, Category = Options)
 	int32 CapsuleHemisphereSteps = 5;
 
+	// When triangulating a capsule, number of vertices to use for the circular cross-sections
 	UPROPERTY(EditAnywhere, BlueprintReadWrite, Category = Options)
 	int32 CapsuleCircleSteps = 8;
+
+	// Whether to cheaply approximate level sets with cubes. Otherwise, will use marching cubes.
+	UPROPERTY(EditAnywhere, BlueprintReadWrite, Category = Options)
+	bool bApproximateLevelSetsWithCubes = false;
+
 };
 
 
diff --git a/Engine/Plugins/Experimental/GeometryScripting/Source/GeometryScriptingCore/Public/GeometryScript/MeshPrimitiveFunctions.h b/Engine/Plugins/Experimental/GeometryScripting/Source/GeometryScriptingCore/Public/GeometryScript/MeshPrimitiveFunctions.h
--- a/Engine/Plugins/Experimental/GeometryScripting/Source/GeometryScriptingCore/Public/GeometryScript/MeshPrimitiveFunctions.h
+++ b/Engine/Plugins/Experimental/GeometryScripting/Source/GeometryScriptingCore/Public/GeometryScript/MeshPrimitiveFunctions.h
@@ -591,7 +591,7 @@ public:
 	/**
 	 * Appends Simple Collision shapes to the Target Mesh, triangulated as specified by Triangulation Options
 	 */
-	UFUNCTION(BlueprintCallable, Category = "GeometryScript|Primitives")
+	UFUNCTION(BlueprintCallable, DisplayName = "Append Simple Collision Shapes to Mesh", Category = "GeometryScript|Primitives")
 	static UPARAM(DisplayName = "Target Mesh") UDynamicMesh*
 	AppendSimpleCollisionShapes(
 		UDynamicMesh* TargetMesh,
diff --git a/Engine/Plugins/Experimental/MeshModelingToolsetExp/Source/MeshModelingToolsExp/Private/Physics/CollisionGeometryVisualization.cpp b/Engine/Plugins/Experimental/MeshModelingToolsetExp/Source/MeshModelingToolsExp/Private/Physics/CollisionGeometryVisualization.cpp
--- a/Engine/Plugins/Experimental/MeshModelingToolsetExp/Source/MeshModelingToolsExp/Private/Physics/CollisionGeometryVisualization.cpp
+++ b/Engine/Plugins/Experimental/MeshModelingToolsetExp/Source/MeshModelingToolsExp/Private/Physics/CollisionGeometryVisualization.cpp
@@ -32,10 +32,11 @@ void InitializePreviewGeometrySolid(
 	const FKAggregateGeom& AggGeom = PhysicsData.AggGeom;
 	FSimpleCollisionTriangulationSettings TriangulationSettings;
 	TriangulationSettings.InitFromSphereResolution(CircleSteps);
+	TriangulationSettings.bApproximateLevelSetWithCubes = false;
 
 	UE::Geometry::ConvertSimpleCollisionToDynamicMeshes(
 		AggGeom,
-		[&](int32 ShapeIndex, const FKShapeElem& ShapeElem, const FDynamicMesh3& Mesh)
+		[&](int32 ShapeIndex, const FKShapeElem& ShapeElem, FDynamicMesh3& Mesh)
 		{
 			FColor Color = GeoSetIndexToColorFunc(ShapeIndex);
 			if (UTriangleSetComponent* TriangleSetComponent = PreviewGeom->CreateOrUpdateTriangleSet(FString::Printf(TEXT("Shape %d"), ShapeIndex), 1, [&](int32 Index, TArray<FRenderableTriangle>& TrisOut)
@@ -74,8 +75,7 @@ void InitializePreviewGeometrySolid(
 				TriangleSetComponent->SetVisibility(bVisible);
 			}
 		},
-		TriangulationSettings,
-		false /*bSetToPerTriangleNormals*/, false /*bInitializeConvexAndLevelSetUVs*/
+		TriangulationSettings
 	);
 }
 
diff --git a/Engine/Plugins/Experimental/MeshModelingToolsetExp/Source/MeshModelingToolsExp/Private/Physics/ExtractCollisionGeometryTool.cpp b/Engine/Plugins/Experimental/MeshModelingToolsetExp/Source/MeshModelingToolsExp/Private/Physics/ExtractCollisionGeometryTool.cpp
--- a/Engine/Plugins/Experimental/MeshModelingToolsetExp/Source/MeshModelingToolsExp/Private/Physics/ExtractCollisionGeometryTool.cpp
+++ b/Engine/Plugins/Experimental/MeshModelingToolsetExp/Source/MeshModelingToolsExp/Private/Physics/ExtractCollisionGeometryTool.cpp
@@ -245,7 +245,8 @@ void UExtractCollisionGeometryTool::RecalculateMesh_Simple()
 		FTransformSequence3d(), SphereResolution, true, true,
 		[&](int32 ElemType, const FDynamicMesh3& ElemMesh) {
 			CurrentMeshParts.Add(MakeShared<FDynamicMesh3>(ElemMesh));
-		});
+		},
+		false /*bApproximateLevelSetWithCubes*/);
 
 	for ( int32 k = 0; k < CurrentMeshParts.Num(); ++k)
 	{
diff --git a/Engine/Plugins/Runtime/MeshModelingToolset/Source/ModelingComponents/ModelingComponents.Build.cs b/Engine/Plugins/Runtime/MeshModelingToolset/Source/ModelingComponents/ModelingComponents.Build.cs
--- a/Engine/Plugins/Runtime/MeshModelingToolset/Source/ModelingComponents/ModelingComponents.Build.cs
+++ b/Engine/Plugins/Runtime/MeshModelingToolset/Source/ModelingComponents/ModelingComponents.Build.cs
@@ -29,6 +29,7 @@ public class ModelingComponents : ModuleRules
 			new string[]
 			{
 				"CoreUObject",
+				"Chaos",
 				"Engine",
 				"RenderCore",
 				"RHI",
diff --git a/Engine/Plugins/Runtime/MeshModelingToolset/Source/ModelingComponents/Private/Physics/ComponentCollisionUtil.cpp b/Engine/Plugins/Runtime/MeshModelingToolset/Source/ModelingComponents/Private/Physics/ComponentCollisionUtil.cpp
--- a/Engine/Plugins/Runtime/MeshModelingToolset/Source/ModelingComponents/Private/Physics/ComponentCollisionUtil.cpp
+++ b/Engine/Plugins/Runtime/MeshModelingToolset/Source/ModelingComponents/Private/Physics/ComponentCollisionUtil.cpp
@@ -328,11 +328,15 @@ namespace UE::Private::ConvertCollisionInternal
 {
 	static void ConvertSimpleCollisionToMeshesHelper(
 		const FKAggregateGeom& AggGeom,
-		FSimpleCollisionTriangulationSettings TriangulationSettings,
-		bool bSetToPerTriangleNormals,
-		bool bInitializeLevelSetAndConvexUVs,
-		TFunctionRef<void(int32 Index, const FKShapeElem& ShapeElem, const FDynamicMesh3&)> DynamicMeshCallback)
+		const FSimpleCollisionTriangulationSettings& TriangulationSettings,
+		const FSimpleCollisionToMeshAttributeSettings& MeshAttributeSettings,
+		TFunctionRef<void(int32 Index, const FKShapeElem& ShapeElem, FDynamicMesh3&)> DynamicMeshCallback,
+		TFunctionRef<bool(const FKShapeElem&)> IncludeElement)
 	{
+		// Clear other attribute related settings if we're not enabling attributes
+		bool bInitializeConvexAndLevelSetUVs = MeshAttributeSettings.bEnableAttributes && MeshAttributeSettings.bInitializeConvexAndLevelSetUVs;
+		bool bSetToPerTriangleNormals = MeshAttributeSettings.bEnableAttributes && MeshAttributeSettings.bSetToPerTriangleNormals;
+
 		// Helpers to transform vertex and normal arrays / invert triangle buffers
 		auto TransformVerts = [](FTransform Transform, TArrayView<FVector3d> Vertices, TArrayView<FVector3f> Normals = TArrayView<FVector3f>())
 		{
@@ -361,9 +365,13 @@ namespace UE::Private::ConvertCollisionInternal
 		};
 
 		auto RunCallbacks = 
-			[&DynamicMeshCallback, bSetToPerTriangleNormals]
+			[&MeshAttributeSettings, &DynamicMeshCallback, bSetToPerTriangleNormals]
 			(int32 Index, const FKShapeElem& ShapeElem, FMeshShapeGenerator& Gen)
 		{
+			if (!MeshAttributeSettings.bEnableAttributes)
+			{
+				Gen.ResetAttributes(true);
+			}
 			FDynamicMesh3 ShapeMesh(&Gen);
 			if (bSetToPerTriangleNormals)
 			{
@@ -375,6 +383,10 @@ namespace UE::Private::ConvertCollisionInternal
 		int32 ShapeIndex = 0;
 		for (const FKSphereElem& Sphere : AggGeom.SphereElems)
 		{
+			if (!IncludeElement(Sphere))
+			{
+				continue;
+			}
 			if (TriangulationSettings.bUseBoxSphere)
 			{
 				FBoxSphereGenerator BoxSphereGen;
@@ -399,6 +411,10 @@ namespace UE::Private::ConvertCollisionInternal
 
 		for (const FKBoxElem& Box : AggGeom.BoxElems)
 		{
+			if (!IncludeElement(Box))
+			{
+				continue;
+			}
 			FMinimalBoxMeshGenerator BoxGen;
 			BoxGen.Box = UE::Geometry::FOrientedBox3d(
 				FFrame3d(FVector3d(Box.Center), FQuaterniond(Box.Rotation.Quaternion())),
@@ -409,6 +425,10 @@ namespace UE::Private::ConvertCollisionInternal
 
 		for (const FKSphylElem& Capsule : AggGeom.SphylElems)
 		{
+			if (!IncludeElement(Capsule))
+			{
+				continue;
+			}
 			FCapsuleGenerator CapsuleGen;
 			CapsuleGen.Radius = Capsule.Radius;
 			CapsuleGen.SegmentLength = Capsule.Length;
@@ -425,6 +445,10 @@ namespace UE::Private::ConvertCollisionInternal
 
 		for (const FKConvexElem& Convex : AggGeom.ConvexElems)
 		{
+			if (!IncludeElement(Convex))
+			{
+				continue;
+			}
 			FTransform ElemTransform = Convex.GetTransform();
 			const int32 NumVertices = Convex.VertexData.Num();
 			const int32 NumTris = Convex.IndexData.Num() / 3;
@@ -448,20 +472,24 @@ namespace UE::Private::ConvertCollisionInternal
 				ConvexMesh.ReverseOrientation(false /*bFlipNormals*/);
 			}
 
-			ConvexMesh.EnableTriangleGroups(0);
-			ConvexMesh.EnableAttributes();
-			if (bSetToPerTriangleNormals)
-			{
-				FMeshNormals::InitializeMeshToPerTriangleNormals(&ConvexMesh);
-			}
-			else
+			if (MeshAttributeSettings.bEnableAttributes)
 			{
-				FMeshNormals::InitializeOverlayToPerVertexNormals(ConvexMesh.Attributes()->PrimaryNormals(), false);
-			}
-			if (bInitializeLevelSetAndConvexUVs)
-			{
-				FDynamicMeshUVEditor UVEditor(&ConvexMesh, 0, true);
-				UVEditor.SetPerTriangleUVs();
+				ConvexMesh.EnableTriangleGroups(0);
+				ConvexMesh.EnableAttributes();
+				if (bSetToPerTriangleNormals)
+				{
+					FMeshNormals::InitializeMeshToPerTriangleNormals(&ConvexMesh);
+				}
+				else
+				{
+					FMeshNormals::InitializeOverlayToPerVertexNormals(ConvexMesh.Attributes()->PrimaryNormals(), false);
+				}
+				if (bInitializeConvexAndLevelSetUVs)
+				{
+					FDynamicMeshUVEditor UVEditor(&ConvexMesh, 0, true);
+					// TODO: Consider adding more options to control the generated UVs (e.g., scaling, box projection?) here and for level sets below
+					UVEditor.SetPerTriangleUVs();
+				}
 			}
 			DynamicMeshCallback(ShapeIndex, Convex, ConvexMesh);
 
@@ -472,42 +500,75 @@ namespace UE::Private::ConvertCollisionInternal
 		TArray<FIntVector> LevelSetTris;
 		for (const FKLevelSetElem& LevelSet : AggGeom.LevelSetElems)
 		{
+			if (!IncludeElement(LevelSet))
+			{
+				continue;
+			}
 			FTransform ElemTransform = LevelSet.GetTransform();
 			bool bInvert = ElemTransform.GetDeterminant() < 0;
 
-			LevelSetVerts.Reset();
-			LevelSetTris.Reset();
-			LevelSet.GetZeroIsosurfaceGridCellFaces(LevelSetVerts, LevelSetTris);
-
 			FDynamicMesh3 LevelSetMesh(EMeshComponents::None);
-			for (int32 Idx = 0; Idx < LevelSetVerts.Num(); ++Idx)
+
+			if (!TriangulationSettings.bApproximateLevelSetWithCubes)
 			{
-				LevelSetMesh.AppendVertex(ElemTransform.TransformPosition(FVector3d(LevelSetVerts[Idx])));
+				FMarchingCubes MarchingCubes;
+				MarchingCubes.CubeSize = LevelSet.GetLevelSet()->GetGrid().Dx().GetMin() * .5;
+				MarchingCubes.IsoValue = 0.0f;
+				Chaos::FAABB3 Box = LevelSet.GetLevelSet()->BoundingBox();
+				MarchingCubes.Bounds = FAxisAlignedBox3d(Box.Min(), Box.Max());
+				MarchingCubes.Bounds.Expand(UE_DOUBLE_KINDA_SMALL_NUMBER);
+				MarchingCubes.RootMode = ERootfindingModes::SingleLerp;
+
+				MarchingCubes.Implicit = [&LevelSet](const FVector3d& Pt) { return -LevelSet.GetLevelSet()->SignedDistance(Pt); };
+				MarchingCubes.Generate();
+
+				if (!MeshAttributeSettings.bEnableAttributes)
+				{
+					MarchingCubes.ResetAttributes(true);
+				}
+				TransformVerts(ElemTransform, MarchingCubes.Vertices, MarchingCubes.Normals);
+
+				LevelSetMesh.Copy(&MarchingCubes);
 			}
-			for (int32 Idx = 0; Idx < LevelSetTris.Num(); ++Idx)
+			else // bApproximateLevelSetWithCubes
 			{
-				FIntVector Tri = LevelSetTris[Idx];
-				LevelSetMesh.AppendTriangle(Tri.X, Tri.Y, Tri.Z);
+				LevelSetVerts.Reset();
+				LevelSetTris.Reset();
+				LevelSet.GetZeroIsosurfaceGridCellFaces(LevelSetVerts, LevelSetTris);
+
+				for (int32 Idx = 0; Idx < LevelSetVerts.Num(); ++Idx)
+				{
+					LevelSetMesh.AppendVertex(ElemTransform.TransformPosition(FVector3d(LevelSetVerts[Idx])));
+				}
+				for (int32 Idx = 0; Idx < LevelSetTris.Num(); ++Idx)
+				{
+					FIntVector Tri = LevelSetTris[Idx];
+					LevelSetMesh.AppendTriangle(Tri.X, Tri.Y, Tri.Z);
+				}
 			}
+
 			if (bInvert)
 			{
 				LevelSetMesh.ReverseOrientation(false /*bFlipNormals*/);
 			}
 
-			LevelSetMesh.EnableTriangleGroups(0);
-			LevelSetMesh.EnableAttributes();
-			if (bSetToPerTriangleNormals)
+			if (MeshAttributeSettings.bEnableAttributes)
 			{
-				FMeshNormals::InitializeMeshToPerTriangleNormals(&LevelSetMesh);
-			}
-			else
-			{
-				FMeshNormals::InitializeOverlayToPerVertexNormals(LevelSetMesh.Attributes()->PrimaryNormals(), false);
-			}
-			if (bInitializeLevelSetAndConvexUVs)
-			{
-				FDynamicMeshUVEditor UVEditor(&LevelSetMesh, 0, true);
-				UVEditor.SetPerTriangleUVs();
+				LevelSetMesh.EnableTriangleGroups(0);
+				LevelSetMesh.EnableAttributes();
+				if (bSetToPerTriangleNormals)
+				{
+					FMeshNormals::InitializeMeshToPerTriangleNormals(&LevelSetMesh);
+				}
+				else
+				{
+					FMeshNormals::InitializeOverlayToPerVertexNormals(LevelSetMesh.Attributes()->PrimaryNormals(), false);
+				}
+				if (bInitializeConvexAndLevelSetUVs)
+				{
+					FDynamicMeshUVEditor UVEditor(&LevelSetMesh, 0, true);
+					UVEditor.SetPerTriangleUVs();
+				}
 			}
 			DynamicMeshCallback(ShapeIndex, LevelSet, LevelSetMesh);
 
@@ -518,17 +579,31 @@ namespace UE::Private::ConvertCollisionInternal
 
 void UE::Geometry::ConvertSimpleCollisionToDynamicMeshes(
 	const FKAggregateGeom& AggGeom,
-	TFunctionRef<void(int32, const FKShapeElem&, const FDynamicMesh3&)> PerElementMeshCallback,
-	FSimpleCollisionTriangulationSettings TriangulationSettings,
-	bool bSetToPerTriangleNormals,
-	bool bInitializeConvexAndLevelSetUVs)
+	TFunctionRef<void(int32, const FKShapeElem&, FDynamicMesh3&)> PerElementMeshCallback,
+	const FSimpleCollisionTriangulationSettings& TriangulationSettings,
+	const FSimpleCollisionToMeshAttributeSettings& MeshAttributeSettings)
+{
+	UE::Private::ConvertCollisionInternal::ConvertSimpleCollisionToMeshesHelper(
+		AggGeom,
+		TriangulationSettings,
+		MeshAttributeSettings,
+		PerElementMeshCallback,
+		[](const FKShapeElem&)->bool { return true; });
+}
+
+void UE::Geometry::ConvertSimpleCollisionToDynamicMeshes(
+	const FKAggregateGeom& AggGeom,
+	TFunctionRef<void(int32, const FKShapeElem&, FDynamicMesh3&)> PerElementMeshCallback,
+	TFunctionRef<bool(const FKShapeElem&)> IncludeElement,
+	const FSimpleCollisionTriangulationSettings& TriangulationSettings,
+	const FSimpleCollisionToMeshAttributeSettings& MeshAttributeSettings)
 {
 	UE::Private::ConvertCollisionInternal::ConvertSimpleCollisionToMeshesHelper(
 		AggGeom,
 		TriangulationSettings,
-		bSetToPerTriangleNormals,
-		bInitializeConvexAndLevelSetUVs,
-		PerElementMeshCallback);
+		MeshAttributeSettings,
+		PerElementMeshCallback,
+		IncludeElement);
 }
 
 
@@ -538,8 +613,9 @@ void UE::Geometry::ConvertSimpleCollisionToMeshes(
 	const FTransformSequence3d& TransformSequence,
 	int32 SphereResolution,
 	bool bSetToPerTriangleNormals,
-	bool bInitializeConvexUVs,
-	TFunction<void(int, const FDynamicMesh3&)> PerElementMeshCallback )
+	bool bInitializeConvexAndLevelSetUVs,
+	TFunction<void(int, const FDynamicMesh3&)> PerElementMeshCallback,
+	bool bApproximateLevelSetWithCubes)
 {
 	FDynamicMeshEditor Editor(&MeshOut);
 
@@ -547,13 +623,18 @@ void UE::Geometry::ConvertSimpleCollisionToMeshes(
 
 	FSimpleCollisionTriangulationSettings TriangulationSettings;
 	TriangulationSettings.InitFromSphereResolution(SphereResolution);
+	TriangulationSettings.bApproximateLevelSetWithCubes = bApproximateLevelSetWithCubes;
+
+	FSimpleCollisionToMeshAttributeSettings MeshAttributeSettings;
+	MeshAttributeSettings.bEnableAttributes = true;
+	MeshAttributeSettings.bSetToPerTriangleNormals = bSetToPerTriangleNormals;
+	MeshAttributeSettings.bInitializeConvexAndLevelSetUVs = bInitializeConvexAndLevelSetUVs;
 
 	FMeshIndexMappings Mappings;
 	UE::Private::ConvertCollisionInternal::ConvertSimpleCollisionToMeshesHelper(
 		AggGeom,
 		TriangulationSettings,
-		bSetToPerTriangleNormals,
-		bInitializeConvexUVs,
+		MeshAttributeSettings,
 		[&](int32 Index, const FKShapeElem& Elem, const UE::Geometry::FDynamicMesh3& Mesh)
 		{
 			PerElementMeshCallback((int)Elem.GetShapeType(), Mesh);
@@ -561,7 +642,8 @@ void UE::Geometry::ConvertSimpleCollisionToMeshes(
 			Editor.AppendMesh(&Mesh, Mappings,
 				[&TransformSequence](int32 vid, const FVector3d& P) { return TransformSequence.TransformPosition(P); },
 				[&TransformSequence](int32 vid, const FVector3d& N) { return TransformSequence.TransformNormal(N); });
-		}
+		},
+		[](const FKShapeElem& Elem) { return true; }
 	);
 	if (bTransformInverts)
 	{
diff --git a/Engine/Plugins/Runtime/MeshModelingToolset/Source/ModelingComponents/Public/Physics/ComponentCollisionUtil.h b/Engine/Plugins/Runtime/MeshModelingToolset/Source/ModelingComponents/Public/Physics/ComponentCollisionUtil.h
--- a/Engine/Plugins/Runtime/MeshModelingToolset/Source/ModelingComponents/Public/Physics/ComponentCollisionUtil.h
+++ b/Engine/Plugins/Runtime/MeshModelingToolset/Source/ModelingComponents/Public/Physics/ComponentCollisionUtil.h
@@ -121,6 +121,7 @@ MODELINGCOMPONENTS_API UBodySetup* GetBodySetup(UPrimitiveComponent* SourceCompo
  * @param bSetToPerTriangleNormals if true and the mesh has a normals attribute overlay,  the mesh is set to face normals, otherwise averaged vertex normals
  * @param bInitializeConvexUVs if true convex hulls have their UVs initialized to per-face planar projections, otherwise no UVs are set
  * @param PerElementMeshCallback if provided, called with each element mesh before transforming/appending to MeshOut. The int parameter is the shape type (EAggCollisionShape::Type), 0=Sphere, 1=Box, 2=Capsule, 3=Convex
+ * @param bApproximateLevelSetWithCubes if true, level sets will be approximated with cubes; otherwise, marching cubes will be used to triangulate level set surfaces
  */
 MODELINGCOMPONENTS_API void ConvertSimpleCollisionToMeshes(
 	const FKAggregateGeom& AggGeom,
@@ -129,7 +130,8 @@ MODELINGCOMPONENTS_API void ConvertSimpleCollisionToMeshes(
 	int32 SphereResolution = 16,
 	bool bSetToPerTriangleNormals = false,
 	bool bInitializeConvexUVs = false,
-	TFunction<void(int, const FDynamicMesh3&)> PerElementMeshCallback = nullptr );
+	TFunction<void(int, const FDynamicMesh3&)> PerElementMeshCallback = nullptr,
+	bool bApproximateLevelSetWithCubes = true);
 
 // Settings to define how simple collision shapes are triangulated
 struct FSimpleCollisionTriangulationSettings
@@ -146,7 +148,8 @@ struct FSimpleCollisionTriangulationSettings
 	// Whether to use a box-sphere triangulation for spheres (otherwise, uses a lat/long sphere triangulation)
 	bool bUseBoxSphere = false;
 
-	// TODO: Add options for controlling level set triangulation as well
+	// Whether to cheaply approximate level set meshes with voxel cubes. Otherwise the level set will be triangulated using marching cubes.
+	bool bApproximateLevelSetWithCubes = true;
 
 	// Init sphere-related settings from a single resolution parameter
 	void InitFromSphereResolution(int32 SphereResolution)
@@ -158,14 +161,42 @@ struct FSimpleCollisionTriangulationSettings
 	}
 };
 
+// Settings to define how attributes such as UVs and normals are set when converting simple collision shapes to dynamic meshes
+struct FSimpleCollisionToMeshAttributeSettings
+{
+	// Whether to enable any attributes on the constructed dynamic meshes. If false, settings for normals and UVs below will be ignored.
+	bool bEnableAttributes = true;
+
+	// Whether to create faceted meshes where each triangle has a separate normal
+	bool bSetToPerTriangleNormals = false;
+
+	// Whether to generate UVs to convex hulls and level sets
+	bool bInitializeConvexAndLevelSetUVs = false;
+
+	FSimpleCollisionToMeshAttributeSettings(bool bEnableAttributes = true, bool bSetToPerTriangleNormals = false, bool bInitializeConvexAndLevelSetUVs = false)
+		: bEnableAttributes(bEnableAttributes), bSetToPerTriangleNormals(bSetToPerTriangleNormals), bInitializeConvexAndLevelSetUVs(bInitializeConvexAndLevelSetUVs)
+	{}
+};
+
 // Similar to ConvertSimpleCollisionToMeshes but without aggregating to an output mesh.
 // @param PerElementMeshCallback	Called with each shape element and the corresponding mesh
 MODELINGCOMPONENTS_API void ConvertSimpleCollisionToDynamicMeshes(
 	const FKAggregateGeom& AggGeom,
-	TFunctionRef<void(int32 Index, const FKShapeElem&, const FDynamicMesh3& Mesh)> PerElementMeshCallback,
-	FSimpleCollisionTriangulationSettings TriangulationSettings,
-	bool bSetToPerTriangleNormals = false,
-	bool bInitializeConvexAndLevelSetUVs = false);
+	TFunctionRef<void(int32 Index, const FKShapeElem&, FDynamicMesh3& Mesh)> PerElementMeshCallback,
+	const FSimpleCollisionTriangulationSettings& TriangulationSettings,
+	const FSimpleCollisionToMeshAttributeSettings& MeshAttributeSettings = FSimpleCollisionToMeshAttributeSettings()
+);
+
+// Similar to ConvertSimpleCollisionToMeshes but without aggregating to an output mesh, and with the option to filter by Shape Element
+// @param PerElementMeshCallback	Called with each shape element and the corresponding mesh
+// @param IncludeElement			Filters Shape Elements: The PerElementMeshCallback will only be called if this function returns true
+MODELINGCOMPONENTS_API void ConvertSimpleCollisionToDynamicMeshes(
+	const FKAggregateGeom& AggGeom,
+	TFunctionRef<void(int32 Index, const FKShapeElem&, FDynamicMesh3& Mesh)> PerElementMeshCallback,
+	TFunctionRef<bool(const FKShapeElem&)> IncludeElement,
+	const FSimpleCollisionTriangulationSettings& TriangulationSettings,
+	const FSimpleCollisionToMeshAttributeSettings& MeshAttributeSettings = FSimpleCollisionToMeshAttributeSettings()
+);
 
 /**
  * Extract the complex collision geometry from CollisionDataProvider as a dynamic mesh.
diff --git a/Engine/Source/Runtime/GeometryCore/Public/Generators/MeshShapeGenerator.h b/Engine/Source/Runtime/GeometryCore/Public/Generators/MeshShapeGenerator.h
--- a/Engine/Source/Runtime/GeometryCore/Public/Generators/MeshShapeGenerator.h
+++ b/Engine/Source/Runtime/GeometryCore/Public/Generators/MeshShapeGenerator.h
@@ -60,6 +60,22 @@ public:
 		return UVs.Num() > 0 || Normals.Num() > 0;
 	}
 
+	// Reset UV and normal-related on a generator, leaving vertices and triangles alone
+	// @param bResetPolygonIDs Whether to also reset polygon IDs
+	void ResetAttributes(bool bResetPolygonIDs = false)
+	{
+		UVs.Reset();
+		UVParentVertex.Reset();
+		Normals.Reset();
+		NormalParentVertex.Reset();
+		TriangleUVs.Reset();
+		TriangleNormals.Reset();
+		if (bResetPolygonIDs)
+		{
+			TrianglePolygonIDs.Reset();
+		}
+	}
+
 	/** clear arrays so that Generate() can be run again */
 	void Reset()
 	{
