Fix handling of external scaling of simple collision shapes across geometry script and modeling tools, to be consistent with how UE applies component transforms to the shapes while still keeping the invariants of the shapes (spheres/capsules remain spheres/capsules, etc)

diff --git a/Engine/Plugins/Experimental/GeometryScripting/Source/GeometryScriptingCore/Private/MeshPrimitiveFunctions.cpp b/Engine/Plugins/Experimental/GeometryScripting/Source/GeometryScriptingCore/Private/MeshPrimitiveFunctions.cpp
--- a/Engine/Plugins/Experimental/GeometryScripting/Source/GeometryScriptingCore/Private/MeshPrimitiveFunctions.cpp
+++ b/Engine/Plugins/Experimental/GeometryScripting/Source/GeometryScriptingCore/Private/MeshPrimitiveFunctions.cpp
@@ -1526,8 +1526,10 @@ UDynamicMesh* UGeometryScriptLibrary_MeshPrimitiveFunctions::AppendSimpleCollisi
 
 	FSimpleCollisionToMeshAttributeSettings UseAttributeOptions(true /*enable attributes*/, false /*per-triangle normals*/, true /*initialize UVs*/);
 
+	FVector ExternalScale = Transform.GetScale3D();
+	Transform.SetScale3D(FVector::OneVector);
 	ConvertSimpleCollisionToDynamicMeshes(
-		SimpleCollision.AggGeom,
+		SimpleCollision.AggGeom, ExternalScale,
 		[&](int32 Index, const FKShapeElem& ShapeElem, FDynamicMesh3& Mesh)
 		{
 			AppendPrimitiveMesh(TargetMesh, Mesh, Transform, PrimitiveOptions);
diff --git a/Engine/Plugins/Experimental/GeometryScripting/Source/GeometryScriptingCore/Private/SceneUtilityFunctions.cpp b/Engine/Plugins/Experimental/GeometryScripting/Source/GeometryScriptingCore/Private/SceneUtilityFunctions.cpp
--- a/Engine/Plugins/Experimental/GeometryScripting/Source/GeometryScriptingCore/Private/SceneUtilityFunctions.cpp
+++ b/Engine/Plugins/Experimental/GeometryScripting/Source/GeometryScriptingCore/Private/SceneUtilityFunctions.cpp
@@ -323,7 +323,7 @@ UDynamicMesh* UGeometryScriptLibrary_SceneUtilityFunctions::CopyCollisionMeshesF
 		bool bFoundMeshErrors = false;
 		UE::Geometry::ConvertComplexCollisionToMeshes(CollisionProvider, AccumulatedMesh, Transforms, bFoundMeshErrors, true, true);
 	}
-	else
+	else // simple collision
 	{
 		const UBodySetup* BodySetup = nullptr;
 		if (UPrimitiveComponent* AnyComponent = Cast<UPrimitiveComponent>(FromObject))
@@ -344,12 +344,16 @@ UDynamicMesh* UGeometryScriptLibrary_SceneUtilityFunctions::CopyCollisionMeshesF
 		}
 
 		FTransformSequence3d Transforms;
+		FVector ExternalScale = FVector::OneVector;
 		if (bTransformToWorld)
 		{
-			Transforms.Append(LocalToWorld);
+			ExternalScale = LocalToWorld.GetScale3D();
+			FTransform WithoutScale = LocalToWorld;
+			WithoutScale.SetScale3D(FVector::OneVector);
+			Transforms.Append(WithoutScale);
 		}
 
-		UE::Geometry::ConvertSimpleCollisionToMeshes(BodySetup->AggGeom, AccumulatedMesh, Transforms, SphereResolution, true, true);
+		UE::Geometry::ConvertSimpleCollisionToMeshes(BodySetup->AggGeom, AccumulatedMesh, Transforms, SphereResolution, true, true, nullptr, false, ExternalScale);
 	}
 
 	ToDynamicMesh->EditMesh([&](FDynamicMesh3& EditMesh)
diff --git a/Engine/Plugins/Experimental/GeometryScripting/Source/GeometryScriptingCore/Public/GeometryScript/MeshPrimitiveFunctions.h b/Engine/Plugins/Experimental/GeometryScripting/Source/GeometryScriptingCore/Public/GeometryScript/MeshPrimitiveFunctions.h
--- a/Engine/Plugins/Experimental/GeometryScripting/Source/GeometryScriptingCore/Public/GeometryScript/MeshPrimitiveFunctions.h
+++ b/Engine/Plugins/Experimental/GeometryScripting/Source/GeometryScriptingCore/Public/GeometryScript/MeshPrimitiveFunctions.h
@@ -590,6 +590,7 @@ public:
 	
 	/**
 	 * Appends Simple Collision shapes to the Target Mesh, triangulated as specified by Triangulation Options
+	 * @param Transform Transform to be applied to simple collision shapes, following the method by which simple collision shapes are transformed -- so, e.g., spheres will not be non-uniformly scaled
 	 */
 	UFUNCTION(BlueprintCallable, DisplayName = "Append Simple Collision Shapes to Mesh", Category = "GeometryScript|Primitives")
 	static UPARAM(DisplayName = "Target Mesh") UDynamicMesh*
diff --git a/Engine/Plugins/Experimental/MeshModelingToolsetExp/Source/MeshModelingToolsExp/Private/Physics/CollisionGeometryVisualization.cpp b/Engine/Plugins/Experimental/MeshModelingToolsetExp/Source/MeshModelingToolsExp/Private/Physics/CollisionGeometryVisualization.cpp
--- a/Engine/Plugins/Experimental/MeshModelingToolsetExp/Source/MeshModelingToolsExp/Private/Physics/CollisionGeometryVisualization.cpp
+++ b/Engine/Plugins/Experimental/MeshModelingToolsetExp/Source/MeshModelingToolsExp/Private/Physics/CollisionGeometryVisualization.cpp
@@ -35,7 +35,7 @@ void InitializePreviewGeometrySolid(
 	TriangulationSettings.bApproximateLevelSetWithCubes = false;
 
 	UE::Geometry::ConvertSimpleCollisionToDynamicMeshes(
-		AggGeom,
+		AggGeom, PhysicsData.ExternalScale3D,
 		[&](int32 ShapeIndex, const FKShapeElem& ShapeElem, FDynamicMesh3& Mesh)
 		{
 			FColor Color = GeoSetIndexToColorFunc(ShapeIndex);
diff --git a/Engine/Plugins/Experimental/MeshModelingToolsetExp/Source/MeshModelingToolsExp/Private/Physics/ExtractCollisionGeometryTool.cpp b/Engine/Plugins/Experimental/MeshModelingToolsetExp/Source/MeshModelingToolsExp/Private/Physics/ExtractCollisionGeometryTool.cpp
--- a/Engine/Plugins/Experimental/MeshModelingToolsetExp/Source/MeshModelingToolsExp/Private/Physics/ExtractCollisionGeometryTool.cpp
+++ b/Engine/Plugins/Experimental/MeshModelingToolsetExp/Source/MeshModelingToolsExp/Private/Physics/ExtractCollisionGeometryTool.cpp
@@ -116,8 +116,8 @@ void UExtractCollisionGeometryTool::Setup()
 
 		PreviewElements = NewObject<UPreviewGeometry>(this);
 		FTransform TargetTransform = (FTransform)UE::ToolTarget::GetLocalToWorldTransform(Target);
-		//PhysicsInfo->ExternalScale3D = TargetTransform.GetScale3D();
-		//TargetTransform.SetScale3D(FVector::OneVector);
+		PhysicsInfo->ExternalScale3D = TargetTransform.GetScale3D();
+		TargetTransform.SetScale3D(FVector::OneVector);
 		PreviewElements->CreateInWorld(UE::ToolTarget::GetTargetActor(Target)->GetWorld(), TargetTransform);
 
 		UE::PhysicsTools::InitializeCollisionGeometryVisualization(PreviewElements, VizSettings, *PhysicsInfo);
@@ -246,7 +246,8 @@ void UExtractCollisionGeometryTool::RecalculateMesh_Simple()
 		[&](int32 ElemType, const FDynamicMesh3& ElemMesh) {
 			CurrentMeshParts.Add(MakeShared<FDynamicMesh3>(ElemMesh));
 		},
-		false /*bApproximateLevelSetWithCubes*/);
+		false /*bApproximateLevelSetWithCubes*/,
+		PhysicsInfo->ExternalScale3D);
 
 	for ( int32 k = 0; k < CurrentMeshParts.Num(); ++k)
 	{
@@ -254,6 +255,9 @@ void UExtractCollisionGeometryTool::RecalculateMesh_Simple()
 		FMeshNormals::InitializeMeshToPerTriangleNormals(&MeshPart);
 	}
 
+	FTransform LocalToWorldUnscaled = (FTransform)UE::ToolTarget::GetLocalToWorldTransform(Target);
+	LocalToWorldUnscaled.SetScale3D(FVector::OneVector);
+	PreviewMesh->SetTransform(LocalToWorldUnscaled);
 	PreviewMesh->UpdatePreview(&CurrentMesh);
 
 	if (CurrentMeshParts.Num() == 0)
@@ -279,6 +283,7 @@ void UExtractCollisionGeometryTool::RecalculateMesh_Complex()
 		UE::Geometry::ConvertComplexCollisionToMeshes(CollisionProvider, CurrentMesh, FTransformSequence3d(), bMeshErrors, Settings->bWeldEdges, true);
 	}
 
+	PreviewMesh->SetTransform((FTransform)UE::ToolTarget::GetLocalToWorldTransform(Target));
 	PreviewMesh->UpdatePreview(&CurrentMesh);
 
 	if (CurrentMesh.TriangleCount() == 0)
diff --git a/Engine/Plugins/Runtime/MeshModelingToolset/Source/ModelingComponents/Private/Physics/ComponentCollisionUtil.cpp b/Engine/Plugins/Runtime/MeshModelingToolset/Source/ModelingComponents/Private/Physics/ComponentCollisionUtil.cpp
--- a/Engine/Plugins/Runtime/MeshModelingToolset/Source/ModelingComponents/Private/Physics/ComponentCollisionUtil.cpp
+++ b/Engine/Plugins/Runtime/MeshModelingToolset/Source/ModelingComponents/Private/Physics/ComponentCollisionUtil.cpp
@@ -327,7 +327,7 @@ bool UE::Geometry::AppendSimpleCollision(
 namespace UE::Private::ConvertCollisionInternal
 {
 	static void ConvertSimpleCollisionToMeshesHelper(
-		const FKAggregateGeom& AggGeom,
+		const FKAggregateGeom& AggGeom, FVector ExternalScale,
 		const FSimpleCollisionTriangulationSettings& TriangulationSettings,
 		const FSimpleCollisionToMeshAttributeSettings& MeshAttributeSettings,
 		TFunctionRef<void(int32 Index, const FKShapeElem& ShapeElem, FDynamicMesh3&)> DynamicMeshCallback,
@@ -387,11 +387,13 @@ namespace UE::Private::ConvertCollisionInternal
 			{
 				continue;
 			}
+			FVector Center = Sphere.Center * ExternalScale;
+			double Radius = FMath::Max(FMathf::ZeroTolerance, Sphere.Radius) * ExternalScale.GetAbsMin();
 			if (TriangulationSettings.bUseBoxSphere)
 			{
 				FBoxSphereGenerator BoxSphereGen;
-				BoxSphereGen.Box.Frame.Origin = Sphere.Center;
-				BoxSphereGen.Radius = FMath::Max(FMathf::ZeroTolerance, Sphere.Radius);
+				BoxSphereGen.Box.Frame.Origin = Center;
+				BoxSphereGen.Radius = Radius;
 				int32 StepsPerSide = FMath::Max(1, TriangulationSettings.BoxSphereStepsPerSide);
 				BoxSphereGen.EdgeVertices = FIndex3i(StepsPerSide, StepsPerSide, StepsPerSide);
 				BoxSphereGen.Generate();
@@ -400,11 +402,11 @@ namespace UE::Private::ConvertCollisionInternal
 			else
 			{
 				FSphereGenerator LatLongSphereGen;
-				LatLongSphereGen.Radius = FMath::Max(FMathf::ZeroTolerance, Sphere.Radius);
+				LatLongSphereGen.Radius = Radius;
 				LatLongSphereGen.NumPhi = LatLongSphereGen.NumTheta = FMath::Max(1, TriangulationSettings.LatLongSphereSteps);
 				LatLongSphereGen.bPolygroupPerQuad = false;
 				LatLongSphereGen.Generate();
-				TranslateVerts(FVector3d(Sphere.Center), LatLongSphereGen.Vertices);
+				TranslateVerts(FVector3d(Center), LatLongSphereGen.Vertices);
 				RunCallbacks(ShapeIndex++, Sphere, LatLongSphereGen);
 			}
 		}
@@ -417,8 +419,8 @@ namespace UE::Private::ConvertCollisionInternal
 			}
 			FMinimalBoxMeshGenerator BoxGen;
 			BoxGen.Box = UE::Geometry::FOrientedBox3d(
-				FFrame3d(FVector3d(Box.Center), FQuaterniond(Box.Rotation.Quaternion())),
-				0.5 * FVector3d(Box.X, Box.Y, Box.Z));
+				FFrame3d(FVector3d(Box.Center * ExternalScale), FQuaterniond(Box.Rotation.Quaternion())),
+				0.5 * FVector3d(Box.X, Box.Y, Box.Z) * ExternalScale);
 			BoxGen.Generate();
 			RunCallbacks(ShapeIndex++, Box, BoxGen);
 		}
@@ -430,14 +432,15 @@ namespace UE::Private::ConvertCollisionInternal
 				continue;
 			}
 			FCapsuleGenerator CapsuleGen;
-			CapsuleGen.Radius = Capsule.Radius;
-			CapsuleGen.SegmentLength = Capsule.Length;
+			CapsuleGen.Radius = (double)Capsule.GetScaledRadius(ExternalScale);
+			CapsuleGen.SegmentLength = (double)Capsule.GetScaledCylinderLength(ExternalScale);
 			CapsuleGen.NumHemisphereArcSteps = TriangulationSettings.CapsuleHemisphereSteps;
 			CapsuleGen.NumCircleSteps = TriangulationSettings.CapsuleCircleSteps;
 			CapsuleGen.bPolygroupPerQuad = false;
 			CapsuleGen.Generate();
-			TranslateVerts(FVector3d(0, 0, -0.5 * Capsule.Length), CapsuleGen.Vertices);
+			TranslateVerts(FVector3d(0, 0, -0.5 * CapsuleGen.SegmentLength), CapsuleGen.Vertices);
 			FTransform Transform = Capsule.GetTransform();
+			Transform.ScaleTranslation(ExternalScale);
 			// Note: Capsule transform cannot invert; it is generated from a rotation and translation
 			TransformVerts(Transform, CapsuleGen.Vertices, CapsuleGen.Normals);
 			RunCallbacks(ShapeIndex++, Capsule, CapsuleGen);
@@ -450,6 +453,8 @@ namespace UE::Private::ConvertCollisionInternal
 				continue;
 			}
 			FTransform ElemTransform = Convex.GetTransform();
+			ElemTransform.ScaleTranslation(ExternalScale);
+			ElemTransform.MultiplyScale3D(ExternalScale);
 			const int32 NumVertices = Convex.VertexData.Num();
 			const int32 NumTris = Convex.IndexData.Num() / 3;
 			bool bInvert = ElemTransform.GetDeterminant() < 0;
@@ -505,6 +510,8 @@ namespace UE::Private::ConvertCollisionInternal
 				continue;
 			}
 			FTransform ElemTransform = LevelSet.GetTransform();
+			ElemTransform.ScaleTranslation(ExternalScale);
+			ElemTransform.MultiplyScale3D(ExternalScale);
 			bool bInvert = ElemTransform.GetDeterminant() < 0;
 
 			FDynamicMesh3 LevelSetMesh(EMeshComponents::None);
@@ -578,13 +585,13 @@ namespace UE::Private::ConvertCollisionInternal
 }
 
 void UE::Geometry::ConvertSimpleCollisionToDynamicMeshes(
-	const FKAggregateGeom& AggGeom,
+	const FKAggregateGeom& AggGeom, FVector ExternalScale,
 	TFunctionRef<void(int32, const FKShapeElem&, FDynamicMesh3&)> PerElementMeshCallback,
 	const FSimpleCollisionTriangulationSettings& TriangulationSettings,
 	const FSimpleCollisionToMeshAttributeSettings& MeshAttributeSettings)
 {
 	UE::Private::ConvertCollisionInternal::ConvertSimpleCollisionToMeshesHelper(
-		AggGeom,
+		AggGeom, ExternalScale,
 		TriangulationSettings,
 		MeshAttributeSettings,
 		PerElementMeshCallback,
@@ -592,14 +599,14 @@ void UE::Geometry::ConvertSimpleCollisionToDynamicMeshes(
 }
 
 void UE::Geometry::ConvertSimpleCollisionToDynamicMeshes(
-	const FKAggregateGeom& AggGeom,
+	const FKAggregateGeom& AggGeom, FVector ExternalScale,
 	TFunctionRef<void(int32, const FKShapeElem&, FDynamicMesh3&)> PerElementMeshCallback,
 	TFunctionRef<bool(const FKShapeElem&)> IncludeElement,
 	const FSimpleCollisionTriangulationSettings& TriangulationSettings,
 	const FSimpleCollisionToMeshAttributeSettings& MeshAttributeSettings)
 {
 	UE::Private::ConvertCollisionInternal::ConvertSimpleCollisionToMeshesHelper(
-		AggGeom,
+		AggGeom, ExternalScale,
 		TriangulationSettings,
 		MeshAttributeSettings,
 		PerElementMeshCallback,
@@ -615,7 +622,8 @@ void UE::Geometry::ConvertSimpleCollisionToMeshes(
 	bool bSetToPerTriangleNormals,
 	bool bInitializeConvexAndLevelSetUVs,
 	TFunction<void(int, const FDynamicMesh3&)> PerElementMeshCallback,
-	bool bApproximateLevelSetWithCubes)
+	bool bApproximateLevelSetWithCubes,
+	FVector ExternalScale)
 {
 	FDynamicMeshEditor Editor(&MeshOut);
 
@@ -632,7 +640,7 @@ void UE::Geometry::ConvertSimpleCollisionToMeshes(
 
 	FMeshIndexMappings Mappings;
 	UE::Private::ConvertCollisionInternal::ConvertSimpleCollisionToMeshesHelper(
-		AggGeom,
+		AggGeom, ExternalScale,
 		TriangulationSettings,
 		MeshAttributeSettings,
 		[&](int32 Index, const FKShapeElem& Elem, const UE::Geometry::FDynamicMesh3& Mesh)
diff --git a/Engine/Plugins/Runtime/MeshModelingToolset/Source/ModelingComponents/Public/Physics/ComponentCollisionUtil.h b/Engine/Plugins/Runtime/MeshModelingToolset/Source/ModelingComponents/Public/Physics/ComponentCollisionUtil.h
--- a/Engine/Plugins/Runtime/MeshModelingToolset/Source/ModelingComponents/Public/Physics/ComponentCollisionUtil.h
+++ b/Engine/Plugins/Runtime/MeshModelingToolset/Source/ModelingComponents/Public/Physics/ComponentCollisionUtil.h
@@ -122,6 +122,7 @@ MODELINGCOMPONENTS_API UBodySetup* GetBodySetup(UPrimitiveComponent* SourceCompo
  * @param bInitializeConvexUVs if true convex hulls have their UVs initialized to per-face planar projections, otherwise no UVs are set
  * @param PerElementMeshCallback if provided, called with each element mesh before transforming/appending to MeshOut. The int parameter is the shape type (EAggCollisionShape::Type), 0=Sphere, 1=Box, 2=Capsule, 3=Convex
  * @param bApproximateLevelSetWithCubes if true, level sets will be approximated with cubes; otherwise, marching cubes will be used to triangulate level set surfaces
+ * @param ExternalScale if provided, apply external component scaling according to the simple collision shapes, in limited way that UE does so (e.g., so spheres w/ non-uniform scaling will still remain spheres)
  */
 MODELINGCOMPONENTS_API void ConvertSimpleCollisionToMeshes(
 	const FKAggregateGeom& AggGeom,
@@ -131,7 +132,8 @@ MODELINGCOMPONENTS_API void ConvertSimpleCollisionToMeshes(
 	bool bSetToPerTriangleNormals = false,
 	bool bInitializeConvexUVs = false,
 	TFunction<void(int, const FDynamicMesh3&)> PerElementMeshCallback = nullptr,
-	bool bApproximateLevelSetWithCubes = true);
+	bool bApproximateLevelSetWithCubes = true,
+	FVector ExternalScale = FVector::OneVector);
 
 // Settings to define how simple collision shapes are triangulated
 struct FSimpleCollisionTriangulationSettings
@@ -181,7 +183,7 @@ struct FSimpleCollisionToMeshAttributeSettings
 // Similar to ConvertSimpleCollisionToMeshes but without aggregating to an output mesh.
 // @param PerElementMeshCallback	Called with each shape element and the corresponding mesh
 MODELINGCOMPONENTS_API void ConvertSimpleCollisionToDynamicMeshes(
-	const FKAggregateGeom& AggGeom,
+	const FKAggregateGeom& AggGeom, FVector ExternalScale,
 	TFunctionRef<void(int32 Index, const FKShapeElem&, FDynamicMesh3& Mesh)> PerElementMeshCallback,
 	const FSimpleCollisionTriangulationSettings& TriangulationSettings,
 	const FSimpleCollisionToMeshAttributeSettings& MeshAttributeSettings = FSimpleCollisionToMeshAttributeSettings()
@@ -191,7 +193,7 @@ MODELINGCOMPONENTS_API void ConvertSimpleCollisionToDynamicMeshes(
 // @param PerElementMeshCallback	Called with each shape element and the corresponding mesh
 // @param IncludeElement			Filters Shape Elements: The PerElementMeshCallback will only be called if this function returns true
 MODELINGCOMPONENTS_API void ConvertSimpleCollisionToDynamicMeshes(
-	const FKAggregateGeom& AggGeom,
+	const FKAggregateGeom& AggGeom, FVector ExternalScale,
 	TFunctionRef<void(int32 Index, const FKShapeElem&, FDynamicMesh3& Mesh)> PerElementMeshCallback,
 	TFunctionRef<bool(const FKShapeElem&)> IncludeElement,
 	const FSimpleCollisionTriangulationSettings& TriangulationSettings,
