add a default-enabled option to remove degenerate triangles during dynamic mesh -> volume conversion, to make volume creation more likely to create valid BSP geometry that works for collisions.  This fixes the collision for the modeling mode cone and arrow primitives when they are created as blocking volumes.

diff --git a/Engine/Plugins/Runtime/MeshModelingToolset/Source/ModelingComponentsEditorOnly/Private/ConversionUtils/DynamicMeshToVolume.cpp b/Engine/Plugins/Runtime/MeshModelingToolset/Source/ModelingComponentsEditorOnly/Private/ConversionUtils/DynamicMeshToVolume.cpp
--- a/Engine/Plugins/Runtime/MeshModelingToolset/Source/ModelingComponentsEditorOnly/Private/ConversionUtils/DynamicMeshToVolume.cpp
+++ b/Engine/Plugins/Runtime/MeshModelingToolset/Source/ModelingComponentsEditorOnly/Private/ConversionUtils/DynamicMeshToVolume.cpp
@@ -7,6 +7,7 @@
 #include "ConstrainedDelaunay2.h"
 #include "Curve/PlanarComplex.h"
 #include "DynamicMesh/DynamicMesh3.h"
+#include "DynamicMeshEditor.h"
 #include "Engine/Polys.h"
 #include "GameFramework/Volume.h"
 #include "DynamicMesh/MeshNormals.h"
@@ -683,6 +684,46 @@ void GetPolygonFaces(const FDynamicMesh3& InputMesh, const FMeshToVolumeOptions&
 		}
 	}
 
+	if (Options.bCleanDegenerate)
+	{
+		// If the above simplification hasn't already switched us to local mesh, initialize local mesh
+		if (UseMesh != &LocalMesh)
+		{
+			LocalMesh = InputMesh;
+			LocalMesh.DiscardAttributes();
+		}
+		
+		// Delete remaining triangles with small area or small edges
+		TArray<int32> ToDelete;
+		for (int32 tid : LocalMesh.TriangleIndicesItr())
+		{
+			if (LocalMesh.GetTriArea(tid) < Options.MinTriangleArea)
+			{
+				ToDelete.Add(tid);
+			}
+		}
+		FDynamicMeshEditor Editor(&LocalMesh);
+		Editor.RemoveTriangles(ToDelete, true);
+		ToDelete.Reset();
+		const double MinEdgeSq = Options.MinEdgeLength * Options.MinEdgeLength;
+		for (int32 eid : LocalMesh.EdgeIndicesItr())
+		{
+			FDynamicMesh3::FEdge Edge = LocalMesh.GetEdge(eid);
+			if ( FVector::DistSquared(LocalMesh.GetVertex(Edge.Vert.A), LocalMesh.GetVertex(Edge.Vert.B)) < MinEdgeSq )
+			{
+				// Note: It's ok that triangles may be added multiple times; RemoveTriangles will just skip triangles were already deleted
+				ToDelete.Add(Edge.Tri.A);
+				if (Edge.Tri.B != IndexConstants::InvalidID)
+				{
+					ToDelete.Add(Edge.Tri.B);
+				}
+			}
+		}
+		Editor.RemoveTriangles(ToDelete, true);
+
+		UseMesh = &LocalMesh;
+	}
+
 	GetPolygonFaces(*UseMesh, FacesOut, Options.bRespectGroupBoundaries);
 }
 
diff --git a/Engine/Plugins/Runtime/MeshModelingToolset/Source/ModelingComponentsEditorOnly/Public/ConversionUtils/DynamicMeshToVolume.h b/Engine/Plugins/Runtime/MeshModelingToolset/Source/ModelingComponentsEditorOnly/Public/ConversionUtils/DynamicMeshToVolume.h
--- a/Engine/Plugins/Runtime/MeshModelingToolset/Source/ModelingComponentsEditorOnly/Public/ConversionUtils/DynamicMeshToVolume.h
+++ b/Engine/Plugins/Runtime/MeshModelingToolset/Source/ModelingComponentsEditorOnly/Public/ConversionUtils/DynamicMeshToVolume.h
@@ -23,7 +23,7 @@ using namespace UE::Geometry;
 struct FMeshToVolumeOptions
 {
 	/**
-	 * When true, coplanar components with different group id's will not be
+	 * When true, coplanar components with different group IDs will not be
 	 * merged into a single volume polygon.
 	 */
 	bool bRespectGroupBoundaries = true;
@@ -35,8 +35,17 @@ struct FMeshToVolumeOptions
 	 */
 	bool bAutoSimplify = false;
 
-	/** If true, simplify the mesh to the specified target triangle count */
+	/** If bAutoSimplify is true, it will simplify the mesh to this MaxTriangles count */
 	int32 MaxTriangles = 250;
+
+	/** If true, attempt to clean degenerate triangles before converting to a volume, as degenerate triangles can cause BSP construction to fail. */
+	bool bCleanDegenerate = true;
+
+	/** If bCleanDegenerate is true, defines the minimum triangle area to keep */
+	double MinTriangleArea = UE_DOUBLE_KINDA_SMALL_NUMBER;
+
+	/** If bCleanDegenerate is true, defines the minimum triangle edge length to keep. */
+	double MinEdgeLength = .01; // default value reflects how the BSP code tends to snap vertices
 };
 
 struct FDynamicMeshFace
