Rename the edit mesh polygons 'Done' UI button from 'Cancel' to 'Complete' because it often (but not always) accepts the action result
Add a comment noting that this is intentional, despite looking broken since the function calls PostAction(CancelCurrent)

diff --git a/Engine/Plugins/Runtime/MeshModelingToolset/Source/MeshModelingTools/Public/EditMeshPolygonsTool.h b/Engine/Plugins/Runtime/MeshModelingToolset/Source/MeshModelingTools/Public/EditMeshPolygonsTool.h
--- a/Engine/Plugins/Runtime/MeshModelingToolset/Source/MeshModelingTools/Public/EditMeshPolygonsTool.h
+++ b/Engine/Plugins/Runtime/MeshModelingToolset/Source/MeshModelingTools/Public/EditMeshPolygonsTool.h
@@ -461,7 +461,9 @@ class MESHMODELINGTOOLS_API UEditMeshPolygonsToolCancelAction : public UEditMesh
 {
 	GENERATED_BODY()
 public:
-	UFUNCTION(CallInEditor, Category = CurrentOperation, meta = (DisplayName = "Cancel", DisplayPriority = 1))
+	// Complete the current action
+	//~ Note this function may accept or cancel depending on the tool, despite the class name and the PostAction argument
+	UFUNCTION(CallInEditor, Category = CurrentOperation, meta = (DisplayName = "Complete", DisplayPriority = 1))
 	void Done() { PostAction(EEditMeshPolygonsToolActions::CancelCurrent); }
 };
 
