handle build scale for static meshes in the static mesh selector

diff --git a/Engine/Plugins/Runtime/MeshModelingToolset/Source/ModelingComponentsEditorOnly/Private/Selection/StaticMeshSelector.cpp b/Engine/Plugins/Runtime/MeshModelingToolset/Source/ModelingComponentsEditorOnly/Private/Selection/StaticMeshSelector.cpp
--- a/Engine/Plugins/Runtime/MeshModelingToolset/Source/ModelingComponentsEditorOnly/Private/Selection/StaticMeshSelector.cpp
+++ b/Engine/Plugins/Runtime/MeshModelingToolset/Source/ModelingComponentsEditorOnly/Private/Selection/StaticMeshSelector.cpp
@@ -18,6 +18,8 @@
 
 #include "Selection/DynamicMeshPolygroupTransformer.h"
 
+#include "StaticMeshOperations.h"
+
 #include "RenderingThread.h"
 #include "UObject/Package.h"
 
@@ -219,9 +221,14 @@ void FStaticMeshSelector::CopyFromStaticMesh()
 {
 	int32 UseLODIndex = 0;
 	const FMeshDescription* SourceMesh = nullptr;
+	FVector BuildScale = FVector::OneVector;
 	if (const UStaticMesh* StaticMesh = WeakStaticMesh.Get())
 	{
 		SourceMesh = StaticMesh->GetMeshDescription(UseLODIndex);
+		if (StaticMesh->IsSourceModelValid(UseLODIndex))
+		{
+			BuildScale = StaticMesh->GetSourceModel(UseLODIndex).BuildSettings.BuildScale3D;
+		}
 	}
 	
 	if (SourceMesh != nullptr)
@@ -229,6 +236,10 @@ void FStaticMeshSelector::CopyFromStaticMesh()
 		FDynamicMesh3 NewMesh;
 		FMeshDescriptionToDynamicMesh Converter;
 		Converter.Convert(SourceMesh, NewMesh, false /*AssetOptions.bRequestTangents*/ );
+		if (!BuildScale.Equals(FVector::OneVector))
+		{
+			MeshTransforms::Scale(NewMesh, BuildScale, FVector::ZeroVector, true);
+		}
 
 		LocalTargetMesh->EditMesh([&](FDynamicMesh3& EditMesh)
 		{
@@ -269,6 +280,18 @@ void FStaticMeshSelector::CommitMeshTransform()
 		{
 			Converter.Convert(&ReadMesh, *MeshDescription, false /*bCopyTangents*/ );
 		});
+		FVector BuildScale = StaticMesh->GetSourceModel(UseLODIndex).BuildSettings.BuildScale3D;
+		if (!BuildScale.Equals(FVector::OneVector))
+		{
+			FVector InvBuildScale;
+			for (int32 SubIdx = 0; SubIdx < 3; ++SubIdx)
+			{
+				InvBuildScale[SubIdx] = FMath::IsNearlyZero(BuildScale[SubIdx], FMathd::Epsilon) ? 1.0 : 1.0 / BuildScale[SubIdx];
+			}
+			FTransform InvScaleTransform = FTransform::Identity;
+			InvScaleTransform.SetScale3D(InvBuildScale);
+			FStaticMeshOperations::ApplyTransform(*MeshDescription, InvScaleTransform, true);
+		}
 
 		StaticMesh->CommitMeshDescription(UseLODIndex);
 		StaticMesh->PostEditChange();
