make staticmeshtooltarget and staticmeshcomponenttooltarget handle buildscale on static mesh <-> mesh description conversions, so tools using these targets will show correctly-scaled meshes

diff --git a/Engine/Plugins/Runtime/MeshModelingToolset/Source/ModelingComponentsEditorOnly/Private/ToolTargets/StaticMeshComponentToolTarget.cpp b/Engine/Plugins/Runtime/MeshModelingToolset/Source/ModelingComponentsEditorOnly/Private/ToolTargets/StaticMeshComponentToolTarget.cpp
--- a/Engine/Plugins/Runtime/MeshModelingToolset/Source/ModelingComponentsEditorOnly/Private/ToolTargets/StaticMeshComponentToolTarget.cpp
+++ b/Engine/Plugins/Runtime/MeshModelingToolset/Source/ModelingComponentsEditorOnly/Private/ToolTargets/StaticMeshComponentToolTarget.cpp
@@ -12,6 +12,7 @@
 #include "ToolTargets/StaticMeshToolTarget.h"
 #include "AssetUtils/MeshDescriptionUtil.h"
 #include "StaticMeshAttributes.h"
+#include "StaticMeshOperations.h"
 #include "UObject/Package.h"
 
 #include UE_INLINE_GENERATED_CPP_BY_NAME(StaticMeshComponentToolTarget)
@@ -131,35 +132,19 @@ bool UStaticMeshComponentToolTarget::CommitMaterialSetUpdate(const FComponentMat
 
 const FMeshDescription* UStaticMeshComponentToolTarget::GetMeshDescription(const FGetMeshParameters& GetMeshParams)
 {
-	static FMeshDescription EmptyMeshDescription;
-	static bool bFirst = true;
-	if (bFirst)
-	{
-		FStaticMeshAttributes Attributes(EmptyMeshDescription);
-		Attributes.Register();
-		bFirst = false;
-	}
-
 	if (ensure(IsValid()))
 	{
 		ensure(GetMeshParams.bWantMeshTangents == false);		// cannot support in this path because we have to modify the FoundMeshDescription pointer!
 
-		const FMeshDescription* FoundMeshDescription = nullptr;
-
-		if (UStaticMesh* StaticMesh = Cast<UStaticMeshComponent>(Component)->GetStaticMesh())
+		UStaticMesh* StaticMesh = Cast<UStaticMeshComponent>(Component)->GetStaticMesh();
+		EMeshLODIdentifier UseLOD = EditingLOD;
+		if (StaticMesh && GetMeshParams.bHaveRequestLOD)
 		{
-			EMeshLODIdentifier UseLOD = EditingLOD;
-			if (GetMeshParams.bHaveRequestLOD)
-			{
-				UseLOD = UStaticMeshToolTarget::GetValidEditingLOD(StaticMesh, GetMeshParams.RequestLOD);
-				ensure(UseLOD == GetMeshParams.RequestLOD);		// probably a bug somewhere if this is not true
-			}
-
-			FoundMeshDescription = (UseLOD == EMeshLODIdentifier::HiResSource) ?
-				StaticMesh->GetHiResMeshDescription() : StaticMesh->GetMeshDescription((int32)UseLOD);
+			UseLOD = UStaticMeshToolTarget::GetValidEditingLOD(StaticMesh, GetMeshParams.RequestLOD);
+			ensure(UseLOD == GetMeshParams.RequestLOD);		// probably a bug somewhere if this is not true
 		}
 
-		return (FoundMeshDescription != nullptr) ? FoundMeshDescription : &EmptyMeshDescription;
+		return UStaticMeshToolTarget::GetMeshDescriptionWithScaleApplied(StaticMesh, (int32)UseLOD, CachedMeshDescriptions);
 	}
 	return nullptr;
 }
@@ -174,6 +159,15 @@ FMeshDescription UStaticMeshComponentToolTarget::GetEmptyMeshDescription()
 
 FMeshDescription UStaticMeshComponentToolTarget::GetMeshDescriptionCopy(const FGetMeshParameters& GetMeshParams)
 {
+	auto ApplyBuildScaleIfNeeded = [](FMeshDescription& MeshDescription, FVector BuildScale)
+	{
+		if (!BuildScale.Equals(FVector::OneVector))
+		{
+			FTransform ScaleTransform = FTransform::Identity;
+			ScaleTransform.SetScale3D(BuildScale);
+			FStaticMeshOperations::ApplyTransform(MeshDescription, ScaleTransform, true);
+		}
+	};
 	if (ensure(IsValid()))
 	{
 		if (UStaticMesh* StaticMesh = Cast<UStaticMeshComponent>(Component)->GetStaticMesh())
@@ -190,6 +184,7 @@ FMeshDescription UStaticMeshComponentToolTarget::GetMeshDescriptionCopy(const FG
 				{
 					FMeshDescription MeshDescriptionCopy = *StaticMesh->GetHiResMeshDescription();
 					const FStaticMeshSourceModel& SourceModel = StaticMesh->GetHiResSourceModel();
+					ApplyBuildScaleIfNeeded(MeshDescriptionCopy, SourceModel.BuildSettings.BuildScale3D);
 					UE::MeshDescription::InitializeAutoGeneratedAttributes(MeshDescriptionCopy, &SourceModel.BuildSettings);
 					return MeshDescriptionCopy;
 				}
@@ -201,6 +196,7 @@ FMeshDescription UStaticMeshComponentToolTarget::GetMeshDescriptionCopy(const FG
 					if (FMeshDescription* SourceMesh = StaticMesh->GetMeshDescription(static_cast<int32>(UseLOD)))
 					{
 						FMeshDescription MeshDescriptionCopy = *SourceMesh;
+						ApplyBuildScaleIfNeeded(MeshDescriptionCopy, StaticMesh->GetSourceModel((int32)UseLOD).BuildSettings.BuildScale3D);
 						UE::MeshDescription::InitializeAutoGeneratedAttributes(MeshDescriptionCopy, StaticMesh, static_cast<int32>(UseLOD));
 						return MeshDescriptionCopy;
 					}
diff --git a/Engine/Plugins/Runtime/MeshModelingToolset/Source/ModelingComponentsEditorOnly/Private/ToolTargets/StaticMeshToolTarget.cpp b/Engine/Plugins/Runtime/MeshModelingToolset/Source/ModelingComponentsEditorOnly/Private/ToolTargets/StaticMeshToolTarget.cpp
--- a/Engine/Plugins/Runtime/MeshModelingToolset/Source/ModelingComponentsEditorOnly/Private/ToolTargets/StaticMeshToolTarget.cpp
+++ b/Engine/Plugins/Runtime/MeshModelingToolset/Source/ModelingComponentsEditorOnly/Private/ToolTargets/StaticMeshToolTarget.cpp
@@ -12,6 +12,7 @@
 #include "Widgets/Notifications/SNotificationList.h"
 #include "Framework/Notifications/NotificationManager.h"
 #include "StaticMeshAttributes.h"
+#include "StaticMeshOperations.h"
 #include "UObject/Package.h"
 
 #include UE_INLINE_GENERATED_CPP_BY_NAME(StaticMeshToolTarget)
@@ -259,6 +260,22 @@ bool UStaticMeshToolTarget::CommitMaterialSetUpdate(UStaticMesh* StaticMeshIn,
 }
 
 const FMeshDescription* UStaticMeshToolTarget::GetMeshDescription(const FGetMeshParameters& GetMeshParams)
+{
+	if (ensure(IsValid()))
+	{
+		EMeshLODIdentifier UseLOD = EditingLOD;
+		if (GetMeshParams.bHaveRequestLOD)
+		{
+			UseLOD = UStaticMeshToolTarget::GetValidEditingLOD(StaticMesh.Get(), GetMeshParams.RequestLOD);
+			ensure(UseLOD == GetMeshParams.RequestLOD);		// probably a bug somewhere if this is not true
+		}
+
+		return GetMeshDescriptionWithScaleApplied(StaticMesh.Get(), (int32)UseLOD, CachedMeshDescriptions);
+	}
+	return nullptr;
+}
+
+const FMeshDescription* UStaticMeshToolTarget::GetMeshDescriptionWithScaleApplied(UStaticMesh* StaticMesh, int32 UseLOD, FMeshDescriptionCache& CachedMeshDescriptions)
 {
 	static FMeshDescription EmptyMeshDescription;
 	static bool bFirst = true;
@@ -269,21 +286,51 @@ const FMeshDescription* UStaticMeshToolTarget::GetMeshDescription(const FGetMesh
 		bFirst = false;
 	}
 
-	if (ensure(IsValid()))
+	if (!StaticMesh)
 	{
-		EMeshLODIdentifier UseLOD = EditingLOD;
-		if (GetMeshParams.bHaveRequestLOD)
+		return &EmptyMeshDescription;
+	}
+
+	FMeshDescription* FoundMeshDescription = nullptr;
+	FVector BuildScale = FVector::One();
+	if (UseLOD == (int32)EMeshLODIdentifier::HiResSource)
+	{
+		FoundMeshDescription = StaticMesh->GetHiResMeshDescription();
+		BuildScale = StaticMesh->GetHiResSourceModel().BuildSettings.BuildScale3D;
+	}
+	else
+	{
+		FoundMeshDescription = StaticMesh->GetMeshDescription(UseLOD);
+		if (FoundMeshDescription)
 		{
-			UseLOD = UStaticMeshToolTarget::GetValidEditingLOD(StaticMesh.Get(), GetMeshParams.RequestLOD);
-			ensure(UseLOD == GetMeshParams.RequestLOD);		// probably a bug somewhere if this is not true
+			BuildScale = StaticMesh->GetSourceModel(UseLOD).BuildSettings.BuildScale3D;
 		}
+	}
 
-		FMeshDescription* FoundMeshDescription = (UseLOD == EMeshLODIdentifier::HiResSource) ?
-			StaticMesh->GetHiResMeshDescription() : StaticMesh->GetMeshDescription((int32)UseLOD);
+	if (CachedMeshDescriptions.Contains(UseLOD))
+	{
+		if (CachedMeshDescriptions[UseLOD].Source == FoundMeshDescription)
+		{
+			return CachedMeshDescriptions[UseLOD].Copy.Get();
+		}
+		else // cache was stale; clear it
+		{
+			CachedMeshDescriptions.Remove(UseLOD);
+		}
+	}
 
-		return (FoundMeshDescription != nullptr) ? FoundMeshDescription : &EmptyMeshDescription;
+	if (FoundMeshDescription && !BuildScale.Equals(FVector::OneVector))
+	{
+		FCachedMeshDescription& Cache = CachedMeshDescriptions.Emplace(UseLOD);
+		Cache.Source = FoundMeshDescription;
+		Cache.Copy = MakeUnique<FMeshDescription>(*FoundMeshDescription);
+		FTransform ScaleTransform = FTransform::Identity;
+		ScaleTransform.SetScale3D(BuildScale);
+		FStaticMeshOperations::ApplyTransform(*Cache.Copy, ScaleTransform, true);
+		FoundMeshDescription = Cache.Copy.Get();
 	}
-	return nullptr;
+
+	return (FoundMeshDescription != nullptr) ? FoundMeshDescription : &EmptyMeshDescription;
 }
 
 FMeshDescription UStaticMeshToolTarget::GetEmptyMeshDescription()
@@ -359,19 +406,37 @@ void UStaticMeshToolTarget::CommitMeshDescription(UStaticMesh* StaticMeshIn, con
 	}
 	// do we need to configure build settings for highres LOD?
 
+	FVector BuildScale = FVector::OneVector;
 	if (EditingLODIn == EMeshLODIdentifier::HiResSource)
 	{
 		StaticMeshIn->ModifyHiResMeshDescription();
+		BuildScale = StaticMeshIn->GetHiResSourceModel().BuildSettings.BuildScale3D;
 	}
 	else
 	{
 		StaticMeshIn->ModifyMeshDescription((int32)EditingLODIn);
+		if (StaticMeshIn->IsSourceModelValid((int32)EditingLODIn))
+		{
+			BuildScale = StaticMeshIn->GetSourceModel((int32)EditingLODIn).BuildSettings.BuildScale3D;
+		}
 	}
 
 	FCommitterParams CommitterParams;
 	CommitterParams.MeshDescriptionOut = UpdateMeshDescription;
 
 	Committer(CommitterParams);
+	if (!BuildScale.Equals(FVector::OneVector))
+	{
+		FTransform InverseBuildScaleTransform = FTransform::Identity;
+		FVector InverseBuildScale;
+		// Safely invert BuildScale
+		for (int32 Idx = 0; Idx < 3; ++Idx)
+		{
+			InverseBuildScale[Idx] = FMath::IsNearlyZero(BuildScale[Idx], FMathd::Epsilon) ? 1.0 : 1.0 / BuildScale[Idx];
+		}
+		InverseBuildScaleTransform.SetScale3D(InverseBuildScale);
+		FStaticMeshOperations::ApplyTransform(*UpdateMeshDescription, InverseBuildScaleTransform, true);
+	}
 
 	if (EditingLODIn == EMeshLODIdentifier::HiResSource)
 	{
diff --git a/Engine/Plugins/Runtime/MeshModelingToolset/Source/ModelingComponentsEditorOnly/Public/ToolTargets/StaticMeshComponentToolTarget.h b/Engine/Plugins/Runtime/MeshModelingToolset/Source/ModelingComponentsEditorOnly/Public/ToolTargets/StaticMeshComponentToolTarget.h
--- a/Engine/Plugins/Runtime/MeshModelingToolset/Source/ModelingComponentsEditorOnly/Public/ToolTargets/StaticMeshComponentToolTarget.h
+++ b/Engine/Plugins/Runtime/MeshModelingToolset/Source/ModelingComponentsEditorOnly/Public/ToolTargets/StaticMeshComponentToolTarget.h
@@ -12,6 +12,7 @@
 #include "TargetInterfaces/StaticMeshBackedTarget.h"
 #include "TargetInterfaces/PhysicsDataSource.h"
 #include "ToolTargets/PrimitiveComponentToolTarget.h"
+#include "ToolTargets/StaticMeshToolTarget.h" // for FMeshDescriptionCache
 #include "ComponentSourceInterfaces.h"  // for EMeshLODIdentifier
 
 #include "StaticMeshComponentToolTarget.generated.h"
@@ -85,6 +86,9 @@ protected:
 	EMeshLODIdentifier EditingLOD = EMeshLODIdentifier::LOD0;
 
 	friend class UStaticMeshComponentToolTargetFactory;
+
+private:
+	UStaticMeshToolTarget::FMeshDescriptionCache CachedMeshDescriptions;
 };
 
 
diff --git a/Engine/Plugins/Runtime/MeshModelingToolset/Source/ModelingComponentsEditorOnly/Public/ToolTargets/StaticMeshToolTarget.h b/Engine/Plugins/Runtime/MeshModelingToolset/Source/ModelingComponentsEditorOnly/Public/ToolTargets/StaticMeshToolTarget.h
--- a/Engine/Plugins/Runtime/MeshModelingToolset/Source/ModelingComponentsEditorOnly/Public/ToolTargets/StaticMeshToolTarget.h
+++ b/Engine/Plugins/Runtime/MeshModelingToolset/Source/ModelingComponentsEditorOnly/Public/ToolTargets/StaticMeshToolTarget.h
@@ -79,6 +79,17 @@ public:
 
 	// Rest provided by parent class
 
+
+	// helper struct to cache mesh descriptions, so we can provide the const ptr GetMeshDescription API while still handling the Build Scale
+	struct FCachedMeshDescription
+	{
+		FMeshDescription* Source;
+		TUniquePtr<FMeshDescription> Copy;
+	};
+	using FMeshDescriptionCache = TMap<int32, FCachedMeshDescription>;
+	// helper to get the mesh description with build scale applied, using/updating the cache as needed
+	static const FMeshDescription* GetMeshDescriptionWithScaleApplied(UStaticMesh* Mesh, int32 UseLOD, FMeshDescriptionCache& Cache);
+
 protected:
 	TWeakObjectPtr<UStaticMesh> StaticMesh = nullptr;
 
@@ -106,6 +117,11 @@ protected:
 		FComponentMaterialSet& MaterialSetOut, bool bPreferAssetMaterials);
 	static bool CommitMaterialSetUpdate(UStaticMesh* SkeletalMesh, 
 		const FComponentMaterialSet& MaterialSet, bool bApplyToAsset);
+
+private:
+	// In some cases we need to modify the mesh description that we pass to tools, e.g. to account for the BuildScale,
+	// and in those cases we retain the storage here to cover the lifetime of the pointer returned by GetMeshDescription()
+	FMeshDescriptionCache CachedMeshDescriptions;
 };
 
 
