Selection mechanic: improve UI responsivness by proposing different notification types
    - on drag: handles the rectangle change directly
    - on tick: stores a pending selection function when dragging and calls it when ticking
    - on realease: stores a pending selection function when dragging and calls it on release

diff --git a/Engine/Plugins/Runtime/MeshModelingToolset/Source/ModelingComponents/Private/Selection/MeshTopologySelectionMechanic.cpp b/Engine/Plugins/Runtime/MeshModelingToolset/Source/ModelingComponents/Private/Selection/MeshTopologySelectionMechanic.cpp
--- a/Engine/Plugins/Runtime/MeshModelingToolset/Source/ModelingComponents/Private/Selection/MeshTopologySelectionMechanic.cpp
+++ b/Engine/Plugins/Runtime/MeshModelingToolset/Source/ModelingComponents/Private/Selection/MeshTopologySelectionMechanic.cpp
@@ -130,6 +130,11 @@
 	UpdateMarqueeEnabled();
 }
 
+void UMeshTopologySelectionMechanic::SetMarqueeSelectionUpdateType(EMarqueeSelectionUpdateType InType)
+{
+	MarqueeSelectionUpdateType = InType;
+}
+
 void UMeshTopologySelectionMechanic::SetBasePriority(const FInputCapturePriority &Priority)
 {
 	BasePriority = Priority;
@@ -213,6 +218,17 @@
 	MarqueeMechanic->DrawHUD(Canvas, RenderAPI);
 }
 
+void UMeshTopologySelectionMechanic::Tick(float DeltaTime)
+{
+	Super::Tick(DeltaTime);
+
+	if (MarqueeSelectionUpdateType == EMarqueeSelectionUpdateType::OnTickAndRelease && PendingSelectionFunction)
+	{
+		PendingSelectionFunction();
+		PendingSelectionFunction.Reset();
+	}
+}
+
 void UMeshTopologySelectionMechanic::ClearHighlight()
 {
 	checkf(DrawnTriangleSetComponent != nullptr, TEXT("Initialize() not called on UMeshTopologySelectionMechanic."));
@@ -294,7 +310,47 @@
 	return true;
 }
 
+void UMeshTopologySelectionMechanic::HandleRectangleChanged(const FCameraRectangle& InRectangle)
+{
+	FGroupTopologySelection RectangleSelection;
 
+	TopoSelector->FindSelectedElement(PreDragTopoSelectorSettings, InRectangle, TargetTransform,
+		RectangleSelection, &TriIsOccludedCache);
 
+	if (ShouldAddToSelectionFunc())
+	{
+		PersistentSelection = PreDragPersistentSelection;
+		if (ShouldRemoveFromSelectionFunc())
+		{
+			PersistentSelection.Toggle(RectangleSelection);
+		}
+		else
+		{
+			PersistentSelection.Append(RectangleSelection);
+		}
+	}
+	else if (ShouldRemoveFromSelectionFunc())
+	{
+		PersistentSelection = PreDragPersistentSelection;
+		PersistentSelection.Remove(RectangleSelection);
+	}
+	else
+	{
+		// Neither key pressed.
+		PersistentSelection = RectangleSelection;
+	}
+
+	// If we modified the currently selected edges/vertices, they will be properly displayed in our
+	// Render() call. However, the mechanic is not responsible for face highlighting, so if we modified
+	// that, we need to notify the user so that they can update the highlighting (since OnSelectionChanged
+	// only gets broadcast at rectangle end).
+	if ((!PersistentSelection.SelectedGroupIDs.IsEmpty() || !LastUpdateRectangleSelection.SelectedGroupIDs.IsEmpty()) // if groups are involved
+		&& PersistentSelection != LastUpdateRectangleSelection)
+	{
+		LastUpdateRectangleSelection = PersistentSelection;
+		OnFaceSelectionPreviewChanged.Broadcast();
+	}
+}
+
 FGroupTopologySelector::FSelectionSettings UMeshTopologySelectionMechanic::GetTopoSelectorSettings(bool bUseOrthoSettings)
 {
@@ -642,6 +698,7 @@
 void UMeshTopologySelectionMechanic::OnDragRectangleStarted()
 {
 	bCurrentlyMarqueeDragging = true;
+	PendingSelectionFunction.Reset();
 
 	ParentTool->GetToolManager()->BeginUndoTransaction(LOCTEXT("SelectionChange", "Selection"));
 	BeginChange();
@@ -654,48 +711,28 @@
 
 void UMeshTopologySelectionMechanic::OnDragRectangleChanged(const FCameraRectangle& CurrentRectangle)
 {
-	FGroupTopologySelection RectangleSelection;
-
-	TopoSelector->FindSelectedElement(PreDragTopoSelectorSettings, CurrentRectangle, TargetTransform,
-		RectangleSelection, &TriIsOccludedCache);
-
-	if (ShouldAddToSelectionFunc())
+	if (MarqueeSelectionUpdateType == EMarqueeSelectionUpdateType::OnDrag)
 	{
-		PersistentSelection = PreDragPersistentSelection;
-		if (ShouldRemoveFromSelectionFunc())
-		{
-			PersistentSelection.Toggle(RectangleSelection);
-		}
-		else
-		{
-			PersistentSelection.Append(RectangleSelection);
-		}
+		HandleRectangleChanged(CurrentRectangle);
 	}
-	else if (ShouldRemoveFromSelectionFunc())
-	{
-		PersistentSelection = PreDragPersistentSelection;
-		PersistentSelection.Remove(RectangleSelection);
-	}
 	else
 	{
-		// Neither key pressed.
-		PersistentSelection = RectangleSelection;
+		// defer the selection on tick or release  
+		PendingSelectionFunction = [this, CurrentRectangle]()
+		{
+			HandleRectangleChanged(CurrentRectangle);
+		};
 	}
-
-	// If we modified the currently selected edges/vertices, they will be properly displayed in our
-	// Render() call. However, the mechanic is not responsible for face highlighting, so if we modified
-	// that, we need to notify the user so that they can update the highlighting (since OnSelectionChanged
-	// only gets broadcast at rectangle end).
-	if ((!PersistentSelection.SelectedGroupIDs.IsEmpty() || !LastUpdateRectangleSelection.SelectedGroupIDs.IsEmpty()) // if groups are involved
-		&& PersistentSelection != LastUpdateRectangleSelection)
-	{
-		LastUpdateRectangleSelection = PersistentSelection;
-		OnFaceSelectionPreviewChanged.Broadcast();
-	}
 }
 
 void UMeshTopologySelectionMechanic::OnDragRectangleFinished(const FCameraRectangle& Rectangle, bool bCancelled)
 {
+	if (PendingSelectionFunction)
+	{
+		PendingSelectionFunction();
+		PendingSelectionFunction.Reset();
+	}
+	
 	bCurrentlyMarqueeDragging = false;
 
 	TriIsOccludedCache.Reset();
diff --git a/Engine/Plugins/Runtime/MeshModelingToolset/Source/ModelingComponents/Public/Selection/MeshTopologySelectionMechanic.h b/Engine/Plugins/Runtime/MeshModelingToolset/Source/ModelingComponents/Public/Selection/MeshTopologySelectionMechanic.h
--- a/Engine/Plugins/Runtime/MeshModelingToolset/Source/ModelingComponents/Public/Selection/MeshTopologySelectionMechanic.h
+++ b/Engine/Plugins/Runtime/MeshModelingToolset/Source/ModelingComponents/Public/Selection/MeshTopologySelectionMechanic.h
@@ -101,6 +101,17 @@ protected:
 	TWeakObjectPtr<UMeshTopologySelectionMechanic> Mechanic;
 };
 
+/*
+ * Selection update type when the marquee rectangle has changed.
+ */
+
+UENUM()
+enum class EMarqueeSelectionUpdateType
+{
+	OnDrag,
+	OnTickAndRelease,
+	OnRelease
+};
 
 
 /**
@@ -134,6 +145,8 @@ public:
 	virtual void Render(IToolsContextRenderAPI* RenderAPI) override;
 	virtual void DrawHUD(FCanvas* Canvas, IToolsContextRenderAPI* RenderAPI);
 
+	virtual void Tick(float DeltaTime) override;
+
 	/**
 	 * Removes the mechanic's own click/hover handlers, which means that the parent tool
 	 * will need to call UpdateSelection(), UpdateHighlight(), ClearHighlight(), and 
@@ -150,6 +163,11 @@ public:
 	 */
 	void SetIsEnabled(bool bOn);
 
+	/**
+	 * Sets how/when the selection updates are handled.
+	 */
+	void SetMarqueeSelectionUpdateType(EMarqueeSelectionUpdateType InType);
+
 	/**
 	 * Sets the base priority so that tools can make sure that their own behaviors are higher
 	 * priority. The mechanic will not use any priority value higher than this, but it may use
@@ -387,6 +405,16 @@ protected:
 	UPROPERTY()
 	TObjectPtr<URectangleMarqueeMechanic> MarqueeMechanic;
 
+	/**
+	 * Selection update type (default is OnDrag) as it may not need to be triggered for every rectangle change
+	 * This can drastically improve the responsiveness of the UI for meshes high density meshes.
+	 * - OnDrag: calls HandleRectangleChanged when dragging
+	 * - OnTick: stores a PendingSelection function when dragging and calls it when ticking and on release (if any)
+	 * - OnRelease: stores a PendingSelection function when dragging and calls it on release (if any)
+	 */
+	UPROPERTY()
+	EMarqueeSelectionUpdateType MarqueeSelectionUpdateType = EMarqueeSelectionUpdateType::OnDrag;
+	
 	FInputCapturePriority BasePriority = FInputCapturePriority(FInputCapturePriority::DEFAULT_TOOL_PRIORITY);
 
 	// When bSelectEdgeLoops is true, this function is tested to see if we should select edge loops,
@@ -404,6 +432,12 @@ protected:
 
 	FTransform3d TargetTransform;
 
+	/** Pending selection function to be called if the selection is deferred to tick/release */
+	TFunction<void()> PendingSelectionFunction;
+
+	/** Calls actual selection using the input marquee rectangle. **/
+	void HandleRectangleChanged(const FCameraRectangle& InRectangle);
+
 	/**
 	 * Get the topology selector settings to use given the current selection settings.
 	 * 
