Make component collision util do needed 'Modify' calls before updating dynamic mesh simple collision shapes, to help avoid bugs where the collision shape update is not correctly captured in the undo transaction

diff --git a/Engine/Plugins/Runtime/MeshModelingToolset/Source/ModelingComponents/Private/Physics/ComponentCollisionUtil.cpp b/Engine/Plugins/Runtime/MeshModelingToolset/Source/ModelingComponents/Private/Physics/ComponentCollisionUtil.cpp
--- a/Engine/Plugins/Runtime/MeshModelingToolset/Source/ModelingComponents/Private/Physics/ComponentCollisionUtil.cpp
+++ b/Engine/Plugins/Runtime/MeshModelingToolset/Source/ModelingComponents/Private/Physics/ComponentCollisionUtil.cpp
@@ -241,6 +241,9 @@ bool UE::Geometry::SetSimpleCollision(
 	}
 	else if (UDynamicMeshComponent* DynamicMeshComponent = Cast<UDynamicMeshComponent>(Component))
 	{
+		DynamicMeshComponent->Modify();
+		UBodySetup* BodySetup = DynamicMeshComponent->GetBodySetup();
+		BodySetup->Modify();
 		DynamicMeshComponent->CollisionType = (ECollisionTraceFlag)CollisionSettings.CollisionTypeFlag;
 		DynamicMeshComponent->SetSimpleCollisionShapes(PhysicsData.AggGeom, true);
 		return true;
@@ -281,6 +284,9 @@ bool UE::Geometry::TransformSimpleCollision(
 	}
 	else if (UDynamicMeshComponent* DynamicMeshComponent = Cast<UDynamicMeshComponent>(Component))
 	{
+		DynamicMeshComponent->Modify();
+		UBodySetup* BodySetup = DynamicMeshComponent->GetBodySetup();
+		BodySetup->Modify();
 		DynamicMeshComponent->SetSimpleCollisionShapes(PhysicsData.AggGeom, true);
 		return true;
 	}
