Add support for copying collision shapes in the Create From Selected base tool, so we get this feature in the derived tools (e.g., the Merge tool, Vox Wrap, a few others).

Customize derived tools to exclude the collision from subtracted shapes.

diff --git a/Engine/Plugins/Experimental/MeshModelingToolsetExp/Source/MeshModelingToolsEditorOnlyExp/Public/VoxelCSGMeshesTool.h b/Engine/Plugins/Experimental/MeshModelingToolsetExp/Source/MeshModelingToolsEditorOnlyExp/Public/VoxelCSGMeshesTool.h
--- a/Engine/Plugins/Experimental/MeshModelingToolsetExp/Source/MeshModelingToolsEditorOnlyExp/Public/VoxelCSGMeshesTool.h
+++ b/Engine/Plugins/Experimental/MeshModelingToolsetExp/Source/MeshModelingToolsEditorOnlyExp/Public/VoxelCSGMeshesTool.h
@@ -84,6 +84,19 @@ protected:
 	UPROPERTY()
 	TObjectPtr<UVoxelCSGMeshesToolProperties> CSGProps;
 
+	virtual bool KeepCollisionFrom(int32 TargetIdx) const override
+	{
+		if (CSGProps->Operation == EVoxelCSGOperation::DifferenceAB)
+		{
+			return TargetIdx == 0;
+		}
+		else if (CSGProps->Operation == EVoxelCSGOperation::DifferenceBA)
+		{
+			return TargetIdx == 1;
+		}
+		return true;
+	}
+
 };
 
 
diff --git a/Engine/Plugins/Experimental/MeshModelingToolsetExp/Source/MeshModelingToolsExp/Public/VoxelBlendMeshesTool.h b/Engine/Plugins/Experimental/MeshModelingToolsetExp/Source/MeshModelingToolsExp/Public/VoxelBlendMeshesTool.h
--- a/Engine/Plugins/Experimental/MeshModelingToolsetExp/Source/MeshModelingToolsExp/Public/VoxelBlendMeshesTool.h
+++ b/Engine/Plugins/Experimental/MeshModelingToolsetExp/Source/MeshModelingToolsExp/Public/VoxelBlendMeshesTool.h
@@ -88,6 +88,15 @@ protected:
 
 	UPROPERTY()
 	TObjectPtr<UVoxelBlendMeshesToolProperties> BlendProperties;
+
+	virtual bool KeepCollisionFrom(int32 TargetIdx) const override
+	{
+		if (BlendProperties->Operation == EVoxelBlendOperation::Subtract)
+		{
+			return TargetIdx == 0;
+		}
+		return true;
+	}
 };
 
 
diff --git a/Engine/Plugins/Runtime/MeshModelingToolset/Source/MeshModelingTools/Public/CSGMeshesTool.h b/Engine/Plugins/Runtime/MeshModelingToolset/Source/MeshModelingTools/Public/CSGMeshesTool.h
--- a/Engine/Plugins/Runtime/MeshModelingToolset/Source/MeshModelingTools/Public/CSGMeshesTool.h
+++ b/Engine/Plugins/Runtime/MeshModelingToolset/Source/MeshModelingTools/Public/CSGMeshesTool.h
@@ -161,6 +161,23 @@ protected:
 
 	// update the material of ghostly preview meshes (used to show trimming or subtracting surface)
 	void UpdatePreviewsMaterial();
+
+	virtual bool KeepCollisionFrom(int32 TargetIdx) const override
+	{
+		if (bTrimMode)
+		{
+			return static_cast<int32>(TrimProperties->WhichMesh) == TargetIdx;
+		}
+		else if (CSGProperties->Operation == ECSGOperation::DifferenceAB)
+		{
+			return TargetIdx == 0;
+		}
+		else if (CSGProperties->Operation == ECSGOperation::DifferenceBA)
+		{
+			return TargetIdx == 1;
+		}
+		return true;
+	}
 };
 
 
diff --git a/Engine/Plugins/Runtime/MeshModelingToolset/Source/ModelingComponents/Private/BaseTools/BaseCreateFromSelectedTool.cpp b/Engine/Plugins/Runtime/MeshModelingToolset/Source/ModelingComponents/Private/BaseTools/BaseCreateFromSelectedTool.cpp
--- a/Engine/Plugins/Runtime/MeshModelingToolset/Source/ModelingComponents/Private/BaseTools/BaseCreateFromSelectedTool.cpp
+++ b/Engine/Plugins/Runtime/MeshModelingToolset/Source/ModelingComponents/Private/BaseTools/BaseCreateFromSelectedTool.cpp
@@ -6,12 +6,15 @@
 #include "ToolSetupUtil.h"
 #include "BaseGizmos/TransformGizmoUtil.h"
 
+#include "Components/StaticMeshComponent.h"
 #include "DynamicMesh/MeshNormals.h"
 #include "DynamicMesh/MeshTransforms.h"
 #include "DynamicMeshToMeshDescription.h"
 
 #include "ModelingObjectsCreationAPI.h"
+#include "Physics/ComponentCollisionUtil.h"
 #include "Selection/ToolSelectionUtil.h"
+#include "ShapeApproximation/SimpleShapeSet3.h"
 
 #include "TargetInterfaces/PrimitiveComponentBackedTarget.h"
 #include "TargetInterfaces/MaterialProvider.h"
@@ -75,6 +78,13 @@ void UBaseCreateFromSelectedTool::Setup()
 	HandleSourcesProperties->RestoreProperties(this);
 	AddToolPropertySource(HandleSourcesProperties);
 
+	if (SupportsCollisionTransfer())
+	{
+		CollisionProperties = NewObject<UBaseCreateFromSelectedCollisionProperties>(this);
+		CollisionProperties->RestoreProperties(this);
+		AddToolPropertySource(CollisionProperties);
+	}
+
 	Preview = NewObject<UMeshOpPreviewWithBackgroundCompute>(this);
 	Preview->Setup(GetTargetWorld(), this);
 	ToolSetupUtil::ApplyRenderingConfigurationToPreview(Preview->PreviewMesh, nullptr);
@@ -227,6 +237,43 @@ void UBaseCreateFromSelectedTool::GenerateAsset(const FDynamicMeshOpResult& OpRe
 	{
 		OutputTypeProperties->ConfigureCreateMeshObjectParams(NewMeshObjectParams);
 	}
+
+	bool bOutputSupportsCustomCollision = NewMeshObjectParams.TypeHint != ECreateObjectTypeHint::Volume;
+	if (SupportsCollisionTransfer() && CollisionProperties->bTransferCollision && bOutputSupportsCustomCollision)
+	{
+		for (int32 TargetIdx = 0; TargetIdx < Targets.Num(); ++TargetIdx)
+		{
+			if (!KeepCollisionFrom(TargetIdx))
+			{
+				continue;
+			}
+			UPrimitiveComponent* SourceComponent = UE::ToolTarget::GetTargetComponent(Targets[TargetIdx]);
+			if (UE::Geometry::ComponentTypeSupportsCollision(SourceComponent, UE::Geometry::EComponentCollisionSupportLevel::ReadOnly))
+			{
+				NewMeshObjectParams.bEnableCollision = true;
+				FComponentCollisionSettings CollisionSettings = UE::Geometry::GetCollisionSettings(SourceComponent);
+				// flag will be set to match the last selection that had collision
+				NewMeshObjectParams.CollisionMode = (ECollisionTraceFlag)CollisionSettings.CollisionTypeFlag;
+
+				FSimpleShapeSet3d ShapeSet;
+				FTransform Transform = FTransform::Identity;
+				if (Targets.Num() > 1)
+				{
+					// Note: NewTransform in the multi-target case is constructed to have uniform scale, so should have a valid inverse
+					Transform = TransformProxies[TargetIdx]->GetTransform() * NewTransform.Inverse();
+				}
+				if (UE::Geometry::GetCollisionShapes(SourceComponent, ShapeSet))
+				{
+					if (!NewMeshObjectParams.CollisionShapeSet.IsSet())
+					{
+						NewMeshObjectParams.CollisionShapeSet.Emplace();
+					}
+					NewMeshObjectParams.CollisionShapeSet->Append(ShapeSet, Transform);
+				}
+			}
+		}
+	}
+
 	FCreateMeshObjectResult Result = UE::Modeling::CreateMeshObject(GetToolManager(), MoveTemp(NewMeshObjectParams));
 	if (Result.IsOK() && Result.NewActor != nullptr)
 	{
@@ -291,6 +338,10 @@ void UBaseCreateFromSelectedTool::OnShutdown(EToolShutdownType ShutdownType)
 {
 	SaveProperties();
 	HandleSourcesProperties->SaveProperties(this);
+	if (SupportsCollisionTransfer())
+	{
+		CollisionProperties->SaveProperties(this);
+	}
 	OutputTypeProperties->SaveProperties(this, TEXT("OutputTypeFromInputTool"));
 	TransformProperties->SaveProperties(this);
 
@@ -303,7 +354,68 @@ void UBaseCreateFromSelectedTool::OnShutdown(EToolShutdownType ShutdownType)
 
 	if (ShutdownType == EToolShutdownType::Accept)
 	{
-		GetToolManager()->BeginUndoTransaction(GetActionName());
+		// Test if we need to pre-update the collision shapes, if we are updating an asset and have multiple targets
+		// Note: For static meshes, this must be done as a separate transaction due to a long-standing bug where undo/redo 
+		// will crash if an asset updates both its collision and geometry in the same transaction
+		// TODO: If/when that static mesh bug is fixed, this collision updating code can be separated out into the UpdateAsset() method
+		bool bAlreadyOpenedMainTransaction = false;
+		if (SupportsCollisionTransfer() && HandleSourcesProperties->OutputWriteTo != EBaseCreateFromSelectedTargetType::NewObject && Targets.Num() > 1)
+		{
+			int32 KeepTargetIndex = (HandleSourcesProperties->OutputWriteTo == EBaseCreateFromSelectedTargetType::FirstInputObject) ? 0 : (Targets.Num() - 1);
+			UPrimitiveComponent* TargetComponent = UE::ToolTarget::GetTargetComponent(Targets[KeepTargetIndex]);
+
+			if (CollisionProperties->bTransferCollision && UE::Geometry::ComponentTypeSupportsCollision(TargetComponent, UE::Geometry::EComponentCollisionSupportLevel::ReadWrite))
+			{
+				bool bHasAddedShapes = false;
+				FSimpleShapeSet3d AccumulateShapes;
+				UE::Geometry::GetCollisionShapes(TargetComponent, AccumulateShapes);
+
+				FTransform3d TargetToWorld = UE::ToolTarget::GetLocalToWorldTransform(Targets[KeepTargetIndex]);
+				FTransformSequence3d InvKeepTargetTransform;
+				InvKeepTargetTransform.AppendInverse(TargetToWorld);
+
+				for (int32 TargetIdx = 0; TargetIdx < Targets.Num(); ++TargetIdx)
+				{
+					if (TargetIdx == KeepTargetIndex || !KeepCollisionFrom(TargetIdx))
+					{
+						continue;
+					}
+					UPrimitiveComponent* SourceComponent = UE::ToolTarget::GetTargetComponent(Targets[TargetIdx]);
+					FSimpleShapeSet3d ShapeSet;
+					if (UE::Geometry::ComponentTypeSupportsCollision(SourceComponent, UE::Geometry::EComponentCollisionSupportLevel::ReadOnly) &&
+						UE::Geometry::GetCollisionShapes(SourceComponent, ShapeSet))
+					{
+						bHasAddedShapes = true;
+						// Create a best-effort transform sequence to take the collision shapes from their proxy transform space to world, and then back to the local space of the target
+						// Note because transforms cannot be fully applied to some collision shape types (e.g. spheres), this will still give weird results in non-uniform scale cases (but this is somewhat unavoidable)
+						FTransformSequence3d TransformSeq;
+						TransformSeq.Append(TransformProxies[TargetIdx]->GetTransform());
+						TransformSeq.Append(InvKeepTargetTransform);
+						AccumulateShapes.Append(ShapeSet, TransformSeq);
+					}
+				}
+
+				if (bHasAddedShapes)
+				{
+					if (!Cast<UStaticMeshComponent>(TargetComponent))
+					{
+						// If we're not operating on a static mesh, we can open the main transaction here so the collision transaction becomes part of it ...
+						// for static mesh, we need the collision update to be a separate transaction to avoid crashing on undo/redo :(
+						bAlreadyOpenedMainTransaction = true;
+						GetToolManager()->BeginUndoTransaction(GetActionName());
+					}
+					GetToolManager()->BeginUndoTransaction(FText::Format(LOCTEXT("SplitTransactionCollisionUpdatePart", "Update Collision ({0})"), GetActionName()));
+					TargetComponent->Modify();
+					FComponentCollisionSettings Settings = GetCollisionSettings(TargetComponent);
+					UE::Geometry::SetSimpleCollision(TargetComponent, &AccumulateShapes, Settings);
+					GetToolManager()->EndUndoTransaction();
+				}
+			}
+		}
+		if (!bAlreadyOpenedMainTransaction)
+		{
+			GetToolManager()->BeginUndoTransaction(GetActionName());
+		}
 
 		// Generate the result
 		AActor* KeepActor = nullptr;
diff --git a/Engine/Plugins/Runtime/MeshModelingToolset/Source/ModelingComponents/Public/BaseTools/BaseCreateFromSelectedTool.h b/Engine/Plugins/Runtime/MeshModelingToolset/Source/ModelingComponents/Public/BaseTools/BaseCreateFromSelectedTool.h
--- a/Engine/Plugins/Runtime/MeshModelingToolset/Source/ModelingComponents/Public/BaseTools/BaseCreateFromSelectedTool.h
+++ b/Engine/Plugins/Runtime/MeshModelingToolset/Source/ModelingComponents/Public/BaseTools/BaseCreateFromSelectedTool.h
@@ -68,9 +68,15 @@ public:
 	FString OutputExistingName;
 };
 
-
-
-
+UCLASS()
+class MODELINGCOMPONENTS_API UBaseCreateFromSelectedCollisionProperties : public UInteractiveToolPropertySet
+{
+	GENERATED_BODY()
+public:
+	/** Whether to transfer collision settings and any simple collision shapes from the source object(s) to the new object. */
+	UPROPERTY(EditAnywhere, Category = OutputObject)
+	bool bTransferCollision = true;
+};
 
 /**
  * Properties of UI to adjust input meshes
@@ -150,6 +156,17 @@ protected:
 	/** Return the materials to be used on the output mesh on tool accept; defaults to the materials set on the preview */
 	virtual TArray<UMaterialInterface*> GetOutputMaterials() const;
 
+	// Override this to control whether the Transfer Collision setting is available
+	virtual bool SupportsCollisionTransfer() const
+	{
+		return true;
+	}
+
+	// Override this to control which inputs should transfer collision to the output (if collision transfer is enabled)
+	virtual bool KeepCollisionFrom(int32 TargetIndex) const
+	{
+		return true;
+	}
 
 
 	/**
@@ -191,6 +208,9 @@ protected:
 	UPROPERTY()
 	TObjectPtr<UBaseCreateFromSelectedHandleSourceProperties> HandleSourcesProperties;
 
+	UPROPERTY()
+	TObjectPtr<UBaseCreateFromSelectedCollisionProperties> CollisionProperties;
+
 	UPROPERTY()
 	TObjectPtr<UMeshOpPreviewWithBackgroundCompute> Preview;
 
