Make the convert tool support a larger range of scene component sources, by switching away from tool targets to instead use a general scene component->mesh conversion function adapted from geometry script. Make the geometry script node also now use this function.

diff --git a/Engine/Plugins/Experimental/GeometryScripting/Source/GeometryScriptingCore/Private/SceneUtilityFunctions.cpp b/Engine/Plugins/Experimental/GeometryScripting/Source/GeometryScriptingCore/Private/SceneUtilityFunctions.cpp
--- a/Engine/Plugins/Experimental/GeometryScripting/Source/GeometryScriptingCore/Private/SceneUtilityFunctions.cpp
+++ b/Engine/Plugins/Experimental/GeometryScripting/Source/GeometryScriptingCore/Private/SceneUtilityFunctions.cpp
@@ -22,6 +22,7 @@
 #include "ConversionUtils/VolumeToDynamicMesh.h"
 #include "ConversionUtils/SkinnedMeshToDynamicMesh.h"
 #include "ConversionUtils/SplineComponentDeformDynamicMesh.h"
+#include "ConversionUtils/SceneComponentToDynamicMesh.h"
 #include "GeometryCollection/GeometryCollectionComponent.h"
 #include "GeometryCollection/GeometryCollectionObject.h"
 #include "Physics/ComponentCollisionUtil.h"
@@ -50,197 +51,38 @@ UDynamicMesh* UGeometryScriptLibrary_SceneUtilityFunctions::CopyMeshFromComponen
 {
 	Outcome = EGeometryScriptOutcomePins::Failure;
 
-	if (USkinnedMeshComponent* SkinnedMeshComponent = Cast<USkinnedMeshComponent>(Component))
+	UE::Conversion::FToMeshOptions ToMeshOptions;
+	ToMeshOptions.bUseClosestLOD = false;
+	ToMeshOptions.LODIndex = Options.RequestedLOD.LODIndex;
+	auto SafeConvertLODType = [](EGeometryScriptLODType GSLODType) -> UE::Conversion::EMeshLODType
 	{
-		LocalToWorld = SkinnedMeshComponent->GetComponentTransform();
-
-		const int32 NumLODs = SkinnedMeshComponent->GetNumLODs();
-		const int32 RequestedLOD = Options.RequestedLOD.LODIndex;
-		if (RequestedLOD < 0 || RequestedLOD > NumLODs - 1)
-		{
-			UE::Geometry::AppendError(Debug, EGeometryScriptErrorType::InvalidInputs, LOCTEXT("CopyMeshFromComponent_MissingSkinnedMeshComponentLOD", "CopyMeshFromComponent: SkinnedMeshComponent requested LOD does not exist"));
-		}
-		else
-		{ 
-			USkinnedAsset* SkinnedAsset = SkinnedMeshComponent->GetSkinnedAsset();
-			if (SkinnedAsset)
-			{
-				FDynamicMesh3 NewMesh;
-				UE::Conversion::SkinnedMeshComponentToDynamicMesh(*SkinnedMeshComponent, NewMesh, RequestedLOD, Options.bWantTangents);
-				NewMesh.DiscardTriangleGroups();
-				ToDynamicMesh->SetMesh(MoveTemp(NewMesh));
-				Outcome = EGeometryScriptOutcomePins::Success;
-			}
-			else
-			{
-				UE::Geometry::AppendError(Debug, EGeometryScriptErrorType::InvalidInputs, LOCTEXT("CopyMeshFromComponent_MissingSkinnedAsset", "CopyMeshFromComponent: SkinnedMeshComponent has a null SkinnedAsset"));
-			}
-		}
-
-	}
-	else if (USplineMeshComponent* SplineMeshComponent = Cast<USplineMeshComponent>(Component))
-	{
-		LocalToWorld = SplineMeshComponent->GetComponentTransform();
-		UStaticMesh* StaticMesh = SplineMeshComponent->GetStaticMesh();
-		if (StaticMesh)
-		{
-			FGeometryScriptCopyMeshFromAssetOptions AssetOptions;
-			AssetOptions.bApplyBuildSettings = (Options.bWantNormals || Options.bWantTangents);
-			AssetOptions.bRequestTangents = Options.bWantTangents;
-			UGeometryScriptLibrary_StaticMeshFunctions::CopyMeshFromStaticMesh(
-				StaticMesh, ToDynamicMesh, AssetOptions, Options.RequestedLOD, Outcome, Debug);	// will set Outcome pin
-
-			// deform the dynamic mesh and its tangent space with the spline
-			constexpr bool bUpdateTangentSpace = true;
-			UE::Geometry::SplineDeformDynamicMesh(*SplineMeshComponent, ToDynamicMesh->GetMeshRef(), bUpdateTangentSpace);
-
-		}
-		else
-		{
-			UE::Geometry::AppendError(Debug, EGeometryScriptErrorType::InvalidInputs, LOCTEXT("CopyMeshFromSplineMeshComponent_MissingStaticMesh", "CopyMeshFromComponent: SplineMeshComponent has a null StaticMesh"));
-		}
-	}
-	else if (UStaticMeshComponent* StaticMeshComponent = Cast<UStaticMeshComponent>(Component))
-	{
-		LocalToWorld = StaticMeshComponent->GetComponentTransform();
-		UStaticMesh* StaticMesh = StaticMeshComponent->GetStaticMesh();
-		if (StaticMesh)
-		{
-			FGeometryScriptCopyMeshFromAssetOptions AssetOptions;
-			AssetOptions.bApplyBuildSettings = (Options.bWantNormals || Options.bWantTangents);
-			AssetOptions.bRequestTangents = Options.bWantTangents;
-			UGeometryScriptLibrary_StaticMeshFunctions::CopyMeshFromStaticMesh(
-				StaticMesh, ToDynamicMesh, AssetOptions, Options.RequestedLOD, Outcome, Debug);	// will set Outcome pin
-
-			// if we have an ISMC, append instances
-			if (UInstancedStaticMeshComponent* ISMComponent = Cast<UInstancedStaticMeshComponent>(StaticMeshComponent))
-			{
-				FDynamicMesh3 InstancedMesh;
-				ToDynamicMesh->EditMesh([&](FDynamicMesh3& EditMesh) { InstancedMesh = MoveTemp(EditMesh); EditMesh.Clear(); },
-					EDynamicMeshChangeType::MeshChange, EDynamicMeshAttributeChangeFlags::Unknown, true);
-				
-				FDynamicMesh3 AccumMesh;
-				AccumMesh.EnableMatchingAttributes(InstancedMesh);
-				FDynamicMeshEditor Editor(&AccumMesh);
-				FMeshIndexMappings Mappings;
-
-				int32 NumInstances = ISMComponent->GetInstanceCount();
-				for (int32 InstanceIdx = 0; InstanceIdx < NumInstances; ++InstanceIdx)
-				{
-					if (ISMComponent->IsValidInstance(InstanceIdx))
-					{
-						FTransform InstanceTransform;
-						ISMComponent->GetInstanceTransform(InstanceIdx, InstanceTransform, /*bWorldSpace=*/false);
-						FTransformSRT3d XForm(InstanceTransform);
-
-						Mappings.Reset();
-						Editor.AppendMesh(&InstancedMesh, Mappings,
-							[&](int, const FVector3d& Position) { return XForm.TransformPosition(Position); },
-							[&](int, const FVector3d& Normal) { return XForm.TransformNormal(Normal); });
-					}
-				}
-
-				ToDynamicMesh->EditMesh([&](FDynamicMesh3& EditMesh) { EditMesh = MoveTemp(AccumMesh); },
-					EDynamicMeshChangeType::MeshChange, EDynamicMeshAttributeChangeFlags::Unknown, true);
-			}
-		}
-		else
-		{
-			UE::Geometry::AppendError(Debug, EGeometryScriptErrorType::InvalidInputs, LOCTEXT("CopyMeshFromComponent_MissingStaticMesh", "CopyMeshFromComponent: StaticMeshComponent has a null StaticMesh"));
-		}
-	}
-	else if (UDynamicMeshComponent* DynamicMeshComponent = Cast<UDynamicMeshComponent>(Component))
-	{
-		LocalToWorld = DynamicMeshComponent->GetComponentTransform();
-		UDynamicMesh* CopyDynamicMesh = DynamicMeshComponent->GetDynamicMesh();
-		if (CopyDynamicMesh)
-		{
-			CopyDynamicMesh->ProcessMesh([&](const FDynamicMesh3& Mesh)
-			{
-				ToDynamicMesh->SetMesh(Mesh);
-			});
-			Outcome = EGeometryScriptOutcomePins::Success;
-		}
-		else
-		{
-			UE::Geometry::AppendError(Debug, EGeometryScriptErrorType::InvalidInputs, LOCTEXT("CopyMeshFromComponent_MissingDynamicMesh", "CopyMeshFromComponent: DynamicMeshComponent has a null DynamicMesh"));
-		}
-	}
-	else if (UBrushComponent* BrushComponent = Cast<UBrushComponent>(Component))
-	{
-		LocalToWorld = BrushComponent->GetComponentTransform();
-
-		UE::Conversion::FVolumeToMeshOptions VolOptions;
-		VolOptions.bMergeVertices = true;
-		VolOptions.bAutoRepairMesh = true;
-		VolOptions.bOptimizeMesh = true;
-		VolOptions.bSetGroups = true;
-
-		FDynamicMesh3 ConvertedMesh(EMeshComponents::FaceGroups);
-		UE::Conversion::BrushComponentToDynamicMesh(BrushComponent, ConvertedMesh, VolOptions);
-
-		// compute normals for current polygroup topology
-		ConvertedMesh.EnableAttributes();
-		if (Options.bWantNormals)
-		{
-			FDynamicMeshNormalOverlay* Normals = ConvertedMesh.Attributes()->PrimaryNormals();
-			FMeshNormals::InitializeOverlayTopologyFromFaceGroups(&ConvertedMesh, Normals);
-			FMeshNormals::QuickRecomputeOverlayNormals(ConvertedMesh);
-		}
-
-		if (ConvertedMesh.TriangleCount() > 0)
-		{
-			ToDynamicMesh->SetMesh(MoveTemp(ConvertedMesh));
-			Outcome = EGeometryScriptOutcomePins::Success;
-		}
-		else
+		switch (GSLODType)
 		{
-			UE::Geometry::AppendError(Debug, EGeometryScriptErrorType::InvalidInputs, LOCTEXT("CopyMeshFromComponent_InvalidBrushConversion", "CopyMeshFromComponent: BrushComponent conversion produced 0 triangles"));
+		case EGeometryScriptLODType::MaxAvailable:
+			return UE::Conversion::EMeshLODType::MaxAvailable;
+		case EGeometryScriptLODType::HiResSourceModel:
+			return UE::Conversion::EMeshLODType::HiResSourceModel;
+		case EGeometryScriptLODType::SourceModel:
+			return UE::Conversion::EMeshLODType::SourceModel;
+		case EGeometryScriptLODType::RenderData:
+			return UE::Conversion::EMeshLODType::RenderData;
 		}
-	}
-	else if (UGeometryCollectionComponent* GeometryCollectionComponent = Cast<UGeometryCollectionComponent>(Component))
+		checkNoEntry();
+		return UE::Conversion::EMeshLODType::MaxAvailable;
+	};
+	ToMeshOptions.LODType = SafeConvertLODType(Options.RequestedLOD.LODType);
+	ToMeshOptions.bWantNormals = Options.bWantNormals;
+	ToMeshOptions.bWantTangents = Options.bWantTangents;
+	UE::Geometry::FDynamicMesh3 NewMesh;
+	FText ErrorMessage;
+	bool bSuccess = UE::Conversion::SceneComponentToDynamicMesh(Component, ToMeshOptions, bTransformToWorld, NewMesh, LocalToWorld, ErrorMessage);
+	if (bSuccess)
 	{
-		LocalToWorld = GeometryCollectionComponent->GetComponentTransform();
-
-		if (const UGeometryCollection* RestCollection = GeometryCollectionComponent->GetRestCollection())
-		{
-			if (const FGeometryCollection* Collection = RestCollection->GetGeometryCollection().Get())
-			{
-				ToDynamicMesh->EditMesh([&](FDynamicMesh3& EditMesh)
-				{
-					FTransform UnusedTransform;
-					const TArray<FTransform3f>& DynamicTransforms = GeometryCollectionComponent->GetComponentSpaceTransforms3f();
-					if (!DynamicTransforms.IsEmpty())
-					{
-						ConvertGeometryCollectionToDynamicMesh(EditMesh, UnusedTransform, false, *Collection, true, DynamicTransforms, false, Collection->TransformIndex.GetConstArray());
-					}
-					else
-					{
-						ConvertGeometryCollectionToDynamicMesh(EditMesh, UnusedTransform, false, *Collection, true, TArrayView<const FTransform3f>(Collection->Transform.GetConstArray()), true, Collection->TransformIndex.GetConstArray());
-					}
-				},
-				EDynamicMeshChangeType::GeneralEdit, EDynamicMeshAttributeChangeFlags::Unknown, true);
-				Outcome = EGeometryScriptOutcomePins::Success;
-			}
-			else
-			{
-				UE::Geometry::AppendError(Debug, EGeometryScriptErrorType::InvalidInputs, LOCTEXT("CopyMeshFromComponent_MissingCollectionData", "CopyMeshFromComponent: GeometryCollectionComponent has null Geometry Collection data"));
-			}
-		}
-		else
-		{
-			UE::Geometry::AppendError(Debug, EGeometryScriptErrorType::InvalidInputs, LOCTEXT("CopyMeshFromComponent_MissingRestCollection", "CopyMeshFromComponent: GeometryCollectionComponent has null Rest Collection object"));
-		}
+		ToDynamicMesh->SetMesh(MoveTemp(NewMesh));
 	}
-
-	// transform mesh to world
-	if (Outcome == EGeometryScriptOutcomePins::Success && bTransformToWorld)
+	else // failed
 	{
-		ToDynamicMesh->EditMesh([&](FDynamicMesh3& EditMesh) 
-		{
-			MeshTransforms::ApplyTransform(EditMesh, (FTransformSRT3d)LocalToWorld, true);
-
-		}, EDynamicMeshChangeType::GeneralEdit, EDynamicMeshAttributeChangeFlags::Unknown, false);	
+		UE::Geometry::AppendError(Debug, EGeometryScriptErrorType::InvalidInputs, FText::Format(LOCTEXT("CopyMeshFromComponent_ConversionError", "CopyMeshFromComponent Error: {0}"), ErrorMessage));
 	}
 
 	return ToDynamicMesh;
diff --git a/Engine/Plugins/Experimental/MeshModelingToolsetExp/Source/MeshModelingToolsExp/MeshModelingToolsExp.Build.cs b/Engine/Plugins/Experimental/MeshModelingToolsetExp/Source/MeshModelingToolsExp/MeshModelingToolsExp.Build.cs
--- a/Engine/Plugins/Experimental/MeshModelingToolsetExp/Source/MeshModelingToolsExp/MeshModelingToolsExp.Build.cs
+++ b/Engine/Plugins/Experimental/MeshModelingToolsetExp/Source/MeshModelingToolsExp/MeshModelingToolsExp.Build.cs
@@ -53,6 +53,15 @@ public class MeshModelingToolsExp : ModuleRules
 			}
 		);
 
+		if (Target.bCompileAgainstEditor) // #if WITH_EDITOR
+		{
+			PrivateDependencyModuleNames.AddRange(
+				new string[]
+				{
+					"ModelingComponentsEditorOnly"
+				});
+		}
+
 		DynamicallyLoadedModuleNames.AddRange(
 			new string[]
 			{
diff --git a/Engine/Plugins/Experimental/MeshModelingToolsetExp/Source/MeshModelingToolsExp/Private/ConvertMeshesTool.cpp b/Engine/Plugins/Experimental/MeshModelingToolsetExp/Source/MeshModelingToolsExp/Private/ConvertMeshesTool.cpp
--- a/Engine/Plugins/Experimental/MeshModelingToolsetExp/Source/MeshModelingToolsExp/Private/ConvertMeshesTool.cpp
+++ b/Engine/Plugins/Experimental/MeshModelingToolsetExp/Source/MeshModelingToolsExp/Private/ConvertMeshesTool.cpp
@@ -1,20 +1,24 @@
 // Copyright Epic Games, Inc. All Rights Reserved.
 
 #include "ConvertMeshesTool.h"
+
+#include "Algo/Accumulate.h"
+#include "Algo/Count.h"
+#include "ConversionUtils/SceneComponentToDynamicMesh.h"
 #include "ComponentSourceInterfaces.h"
 #include "InteractiveToolManager.h"
 #include "ToolTargetManager.h"
 #include "ToolBuilderUtil.h"
 #include "ToolSetupUtil.h"
-#include "ModelingToolTargetUtil.h"
 #include "ModelingObjectsCreationAPI.h"
 #include "Selection/ToolSelectionUtil.h"
 #include "Physics/ComponentCollisionUtil.h"
 #include "ShapeApproximation/SimpleShapeSet3.h"
 #include "DynamicMesh/DynamicMeshAttributeSet.h"
 
-#include "TargetInterfaces/MeshDescriptionProvider.h"
-#include "TargetInterfaces/PrimitiveComponentBackedTarget.h"
+#if WITH_EDITOR
+#include "ToolTargets/StaticMeshComponentToolTarget.h" // for GetActiveEditingLOD
+#endif
 
 #include UE_INLINE_GENERATED_CPP_BY_NAME(ConvertMeshesTool)
 
@@ -25,23 +29,73 @@ using namespace UE::Geometry;
 /*
  * ToolBuilder
  */
-const FToolTargetTypeRequirements& UConvertMeshesToolBuilder::GetTargetRequirements() const
+UInteractiveTool* UConvertMeshesToolBuilder::BuildTool(const FToolBuilderState& SceneState) const
 {
-	static FToolTargetTypeRequirements TypeRequirements({
-		UMaterialProvider::StaticClass(),
-		UMeshDescriptionProvider::StaticClass(),
-		UPrimitiveComponentBackedTarget::StaticClass()
-		});
-	return TypeRequirements;
+	UConvertMeshesTool* Tool = NewObject<UConvertMeshesTool>(SceneState.ToolManager);
+	TArray<TWeakObjectPtr<UPrimitiveComponent>> Inputs;
+	auto TryAddSelected = [&Inputs](UActorComponent* Selected)
+	{
+		if (UPrimitiveComponent* PrimComponent = Cast<UPrimitiveComponent>(Selected))
+		{
+			if (UE::Conversion::CanConvertSceneComponentToDynamicMesh(PrimComponent))
+			{
+				Inputs.Emplace(PrimComponent);
+			}
+		}
+	};
+	if (SceneState.SelectedComponents.Num() > 0)
+	{
+		for (UActorComponent* Selected : SceneState.SelectedComponents)
+		{
+			TryAddSelected(Selected);
+		}
+	}
+	else
+	{
+		for (AActor* SelectedActor : SceneState.SelectedActors)
+		{
+			for (UActorComponent* Selected : SelectedActor->GetComponents())
+			{
+				TryAddSelected(Selected);
+			}
+		}
+	}
+	Tool->InitializeInputs(MoveTemp(Inputs));
+
+#if WITH_EDITOR
+	 
+	if (UStaticMeshComponentToolTargetFactory* StaticMeshComponentTargetFactory = SceneState.TargetManager->FindFirstFactoryByType<UStaticMeshComponentToolTargetFactory>())
+	{
+		EMeshLODIdentifier LOD = StaticMeshComponentTargetFactory->GetActiveEditingLOD();
+		Tool->SetTargetLOD(LOD);
+	}
+#endif
+
+	return Tool;
 }
 
-UMultiSelectionMeshEditingTool* UConvertMeshesToolBuilder::CreateNewTool(const FToolBuilderState& SceneState) const
+bool UConvertMeshesToolBuilder::CanBuildTool(const FToolBuilderState& SceneState) const
 {
-	return NewObject<UConvertMeshesTool>(SceneState.ToolManager);
+	int32 CanConvertCount = 0;
+	auto CanConvert = [](UActorComponent* Comp)->bool { return UE::Conversion::CanConvertSceneComponentToDynamicMesh(Cast<USceneComponent>(Comp)); };
+	if (SceneState.SelectedComponents.Num() > 0)
+	{
+		CanConvertCount = static_cast<int32>(Algo::CountIf(SceneState.SelectedComponents, CanConvert));
+	}
+	else
+	{
+		CanConvertCount =
+			Algo::TransformAccumulate(SceneState.SelectedActors,
+				[&CanConvert](AActor* Actor)
+				{
+					return static_cast<int>(Algo::CountIf(Actor->GetComponents(), CanConvert));
+				},
+				0);
+	}
+	return CanConvertCount > 0;
 }
 
 
-
 /*
  * Tool
  */
@@ -75,10 +129,18 @@ void UConvertMeshesTool::Setup()
 
 bool UConvertMeshesTool::CanAccept() const
 {
-	return Super::CanAccept();
+	// Refuse to accept if we have any invalid inputs
+	for (const TWeakObjectPtr<UPrimitiveComponent>& Input : Inputs)
+	{
+		if (!Input.IsValid())
+		{
+			return false;
+		}
+	}
+	return true;
 }
 
-void UConvertMeshesTool::OnShutdown(EToolShutdownType ShutdownType)
+void UConvertMeshesTool::Shutdown(EToolShutdownType ShutdownType)
 {
 	OutputTypeProperties->SaveProperties(this, TEXT("ConvertMeshesTool"));
 	BasicProperties->SaveProperties(this);
@@ -94,15 +156,47 @@ void UConvertMeshesTool::OnShutdown(EToolShutdownType ShutdownType)
 		// Accumulate info for new mesh objects. Do not immediately create them because then
 		// the new Actors will get a unique-name incremented suffix, because the convert-from
 		// Actors still exist.
-		for (int32 k = 0; k < Targets.Num(); ++k)
+		for (TWeakObjectPtr<UPrimitiveComponent> Input : Inputs)
 		{
-			AActor* TargetActor = UE::ToolTarget::GetTargetActor(Targets[k]);
+			if (!Input.IsValid())
+			{
+				continue;
+			}
+			UPrimitiveComponent* InputComponent = Input.Get();
+			FDynamicMesh3 SourceMesh;
+			FTransform SourceTransform;
+			TArray<UMaterialInterface*> ComponentMaterials;
+			TArray<UMaterialInterface*> AssetMaterials;
+			FText ErrorMessage;
+
+			UE::Conversion::FToMeshOptions Options;
+			Options.LODType = UE::Conversion::EMeshLODType::SourceModel;
+			Options.LODIndex = 0;
+			Options.bUseClosestLOD = true;
+			if ((int32)TargetLOD <= (int32)EMeshLODIdentifier::LOD7)
+			{
+				Options.LODIndex = (int32)TargetLOD;
+			}
+			else if (TargetLOD == EMeshLODIdentifier::MaxQuality)
+			{
+				Options.LODType = UE::Conversion::EMeshLODType::MaxAvailable;
+			}
+			else if (TargetLOD == EMeshLODIdentifier::HiResSource)
+			{
+				Options.LODType = UE::Conversion::EMeshLODType::HiResSourceModel;
+			}
+
+			bool bSuccess = UE::Conversion::SceneComponentToDynamicMesh(InputComponent, Options, false, SourceMesh, SourceTransform, ErrorMessage, &ComponentMaterials, &AssetMaterials);
+			if (!bSuccess)
+			{
+				UE_LOG(LogGeometry, Warning, TEXT("Convert Tool failed to convert %s: %s"), *InputComponent->GetName(), *ErrorMessage.ToString());
+				continue;
+			}
+
+			AActor* TargetActor = InputComponent->GetOwner();
 			check(TargetActor != nullptr);
 			DeleteActors.Add(TargetActor);
 
-			FTransform SourceTransform = (FTransform)UE::ToolTarget::GetLocalToWorldTransform(Targets[k]);
-			FDynamicMesh3 SourceMesh = UE::ToolTarget::GetDynamicMeshCopy(Targets[k], true);
-
 			// if not transferring materials, need to clear out any existing MaterialIDs
 			if (BasicProperties->bTransferMaterials == false && SourceMesh.HasAttributes())
 			{
@@ -116,30 +210,28 @@ void UConvertMeshesTool::OnShutdown(EToolShutdownType ShutdownType)
 			}
 
 			FString AssetName = TargetActor->GetActorNameOrLabel();
-			FComponentMaterialSet Materials = UE::ToolTarget::GetMaterialSet(Targets[k]);
-			const FComponentMaterialSet* TransferMaterials = (BasicProperties->bTransferMaterials) ? &Materials : nullptr;
 
 			FCreateMeshObjectParams NewMeshObjectParams;
-			NewMeshObjectParams.TargetWorld = GetTargetWorld();
+			NewMeshObjectParams.TargetWorld = InputComponent->GetWorld();
 			NewMeshObjectParams.Transform = (FTransform)SourceTransform;
 			NewMeshObjectParams.BaseName = AssetName;
-			if (TransferMaterials != nullptr)
+			if (BasicProperties->bTransferMaterials)
 			{
-				NewMeshObjectParams.Materials = TransferMaterials->Materials;
+				NewMeshObjectParams.Materials = ComponentMaterials;
+				NewMeshObjectParams.AssetMaterials = AssetMaterials;
 			}
 			NewMeshObjectParams.SetMesh(MoveTemp(SourceMesh));
 			
 			if (BasicProperties->bTransferCollision)
 			{
-				UPrimitiveComponent* SourceComponent = UE::ToolTarget::GetTargetComponent(Targets[k]);
-				if (UE::Geometry::ComponentTypeSupportsCollision(SourceComponent, UE::Geometry::EComponentCollisionSupportLevel::ReadOnly))
+				if (UE::Geometry::ComponentTypeSupportsCollision(InputComponent, UE::Geometry::EComponentCollisionSupportLevel::ReadOnly))
 				{
 					NewMeshObjectParams.bEnableCollision = true;
-					FComponentCollisionSettings CollisionSettings = UE::Geometry::GetCollisionSettings(SourceComponent);
+					FComponentCollisionSettings CollisionSettings = UE::Geometry::GetCollisionSettings(InputComponent);
 					NewMeshObjectParams.CollisionMode = (ECollisionTraceFlag)CollisionSettings.CollisionTypeFlag;
 
 					FSimpleShapeSet3d ShapeSet;
-					if (UE::Geometry::GetCollisionShapes(SourceComponent, ShapeSet))
+					if (UE::Geometry::GetCollisionShapes(InputComponent, ShapeSet))
 					{
 						NewMeshObjectParams.CollisionShapeSet = MoveTemp(ShapeSet);
 					}
diff --git a/Engine/Plugins/Experimental/MeshModelingToolsetExp/Source/MeshModelingToolsExp/Public/ConvertMeshesTool.h b/Engine/Plugins/Experimental/MeshModelingToolsetExp/Source/MeshModelingToolsExp/Public/ConvertMeshesTool.h
--- a/Engine/Plugins/Experimental/MeshModelingToolsetExp/Source/MeshModelingToolsExp/Public/ConvertMeshesTool.h
+++ b/Engine/Plugins/Experimental/MeshModelingToolsetExp/Source/MeshModelingToolsExp/Public/ConvertMeshesTool.h
@@ -7,20 +7,20 @@
 #include "InteractiveToolBuilder.h"
 #include "InteractiveToolQueryInterfaces.h" // IInteractiveToolExclusiveToolAPI
 #include "PropertySets/CreateMeshObjectTypeProperties.h"
+#include "TargetInterfaces/MeshTargetInterfaceTypes.h"
+
 #include "ConvertMeshesTool.generated.h"
 
 /**
  *
  */
 UCLASS()
-class MESHMODELINGTOOLSEXP_API UConvertMeshesToolBuilder : public UMultiSelectionMeshEditingToolBuilder
+class MESHMODELINGTOOLSEXP_API UConvertMeshesToolBuilder : public UInteractiveToolBuilder
 {
 	GENERATED_BODY()
-public:
-	virtual UMultiSelectionMeshEditingTool* CreateNewTool(const FToolBuilderState& SceneState) const override;
 
-protected:
-	virtual const FToolTargetTypeRequirements& GetTargetRequirements() const override;
+	virtual bool CanBuildTool(const FToolBuilderState& SceneState) const;
+	virtual UInteractiveTool* BuildTool(const FToolBuilderState& SceneState) const;
 };
 
 /**
@@ -45,7 +45,7 @@ public:
 
 
 UCLASS()
-class MESHMODELINGTOOLSEXP_API UConvertMeshesTool : public UMultiSelectionMeshEditingTool,
+class MESHMODELINGTOOLSEXP_API UConvertMeshesTool : public UInteractiveTool,
 	// Disallow auto-accept switch-away for the tool
 	public IInteractiveToolExclusiveToolAPI
 {
@@ -53,7 +53,7 @@ class MESHMODELINGTOOLSEXP_API UConvertMeshesTool : public UMultiSelectionMeshEd
 
 public:
 	virtual void Setup() override;
-	virtual void OnShutdown(EToolShutdownType ShutdownType) override;
+	virtual void Shutdown(EToolShutdownType ShutdownType) override;
 
 	virtual bool HasCancel() const override { return true; }
 	virtual bool HasAccept() const override { return true; }
@@ -64,4 +64,20 @@ public:
 
 	UPROPERTY()
 	TObjectPtr<UCreateMeshObjectTypeProperties> OutputTypeProperties;
+
+	void InitializeInputs(TArray<TWeakObjectPtr<UPrimitiveComponent>>&& InInputs)
+	{
+		Inputs = MoveTemp(InInputs);
+	}
+
+	void SetTargetLOD(EMeshLODIdentifier LOD)
+	{
+		TargetLOD = LOD;
+	}
+
+private:
+ 	UPROPERTY()
+	TArray<TWeakObjectPtr<UPrimitiveComponent>> Inputs;
+
+	EMeshLODIdentifier TargetLOD = EMeshLODIdentifier::Default;
 };
diff --git a/Engine/Plugins/Runtime/MeshModelingToolset/MeshModelingToolset.uplugin b/Engine/Plugins/Runtime/MeshModelingToolset/MeshModelingToolset.uplugin
--- a/Engine/Plugins/Runtime/MeshModelingToolset/MeshModelingToolset.uplugin
+++ b/Engine/Plugins/Runtime/MeshModelingToolset/MeshModelingToolset.uplugin
@@ -54,6 +54,10 @@
 		{
 			"Name": "ProxyLODPlugin",
 			"Enabled": true
+		},
+		{
+			"Name": "PlanarCut",
+			"Enabled":  true
 		}
 	]
 }
diff --git a/Engine/Plugins/Runtime/MeshModelingToolset/Source/ModelingComponents/ModelingComponents.Build.cs b/Engine/Plugins/Runtime/MeshModelingToolset/Source/ModelingComponents/ModelingComponents.Build.cs
--- a/Engine/Plugins/Runtime/MeshModelingToolset/Source/ModelingComponents/ModelingComponents.Build.cs
+++ b/Engine/Plugins/Runtime/MeshModelingToolset/Source/ModelingComponents/ModelingComponents.Build.cs
@@ -37,6 +37,9 @@ public class ModelingComponents : ModuleRules
 				"SkeletalMeshDescription", // FSkeletalMeshAttributes::DefaultSkinWeightProfileName
 				"SlateCore",
 				"ImageCore",
+				"PlanarCut",
+				"GeometryCollectionEngine",
+				"MeshConversionEngineTypes"
 				// ... add private dependencies that you statically link with here ...
 			}
 		);
diff --git a/Engine/Plugins/Runtime/MeshModelingToolset/Source/ModelingComponents/Private/ConversionUtils/SceneComponentToDynamicMesh.cpp b/Engine/Plugins/Runtime/MeshModelingToolset/Source/ModelingComponents/Private/ConversionUtils/SceneComponentToDynamicMesh.cpp
new file mode 100644
--- /dev/null
+++ b/Engine/Plugins/Runtime/MeshModelingToolset/Source/ModelingComponents/Private/ConversionUtils/SceneComponentToDynamicMesh.cpp
@@ -0,0 +1,588 @@
+// Copyright Epic Games, Inc. All Rights Reserved.
+
+#include "ConversionUtils/SceneComponentToDynamicMesh.h"
+
+#include "Engine/StaticMesh.h"
+#include "Engine/SkeletalMesh.h"
+#include "Engine/SkinnedAssetCommon.h"
+
+#include "Components/SceneComponent.h"
+#include "Components/StaticMeshComponent.h"
+#include "Components/InstancedStaticMeshComponent.h"
+#include "Components/SkeletalMeshComponent.h"
+#include "Components/SkinnedMeshComponent.h"
+#include "Components/BrushComponent.h"
+#include "Components/DynamicMeshComponent.h"
+#include "Components/SplineMeshComponent.h"
+#include "ConversionUtils/VolumeToDynamicMesh.h"
+#include "ConversionUtils/SkinnedMeshToDynamicMesh.h"
+#include "ConversionUtils/SplineComponentDeformDynamicMesh.h"
+#include "DynamicMesh/DynamicMeshAttributeSet.h"
+#include "DynamicMesh/MeshNormals.h"
+#include "DynamicMesh/MeshTransforms.h"
+#include "DynamicMeshEditor.h"
+#include "GeometryCollection/GeometryCollectionComponent.h"
+#include "GeometryCollection/GeometryCollectionObject.h"
+#include "MeshDescription.h"
+#include "MeshDescriptionToDynamicMesh.h"
+#include "Physics/ComponentCollisionUtil.h"
+#include "PlanarCut.h"
+#include "StaticMeshAttributes.h"
+#include "StaticMeshLODResourcesToDynamicMesh.h"
+#include "StaticMeshOperations.h"
+
+#define LOCTEXT_NAMESPACE "ModelingComponents_SceneComponentToDynamicMesh"
+
+namespace UE 
+{
+namespace Conversion
+{
+
+bool CanConvertSceneComponentToDynamicMesh(USceneComponent* Component)
+{
+	if (!Component)
+	{
+		return false;
+	}
+	else if (Cast<USkinnedMeshComponent>(Component))
+	{
+		return true;
+	}
+	else if (Cast<USplineMeshComponent>(Component))
+	{
+		return true;
+	}
+	else if (Cast<UStaticMeshComponent>(Component))
+	{
+		return true;
+	}
+	else if (Cast<UDynamicMeshComponent>(Component))
+	{
+		return true;
+	}
+	else if (Cast<UBrushComponent>(Component))
+	{
+		return true;
+	}
+	else if (Cast<UGeometryCollectionComponent>(Component))
+	{
+		return true;
+	}
+	return false;
+}
+
+// Conversion helpers
+namespace Private::ConversionHelper
+{
+	// Static mesh conversion functions (from geometry script MeshAssetFunctions.cpp)
+	// TODO: these static mesh conversion helpers should be pulled out to their own StaticMeshToDynamicMesh converter method
+
+	struct FStaticMeshConversionOptions
+	{
+		// Whether to apply Build Settings during the mesh copy.
+		bool bApplyBuildSettings = true;
+
+		// Whether to request tangents on the copied mesh. If tangents are not requested, tangent-related build settings will also be ignored.
+		bool bRequestTangents = true;
+
+		// Whether to ignore the 'remove degenerates' option from Build Settings. Note: Only applies if 'Apply Build Settings' is enabled.
+		bool bIgnoreRemoveDegenerates = true;
+
+		// Whether to scale the copied mesh by the Build Setting's 'Build Scale'. Note: This is considered separately from the 'Apply Build Settings' option.
+		bool bUseBuildScale = true;
+	};
+
+	static bool CopyMeshFromStaticMesh_SourceData(
+		UStaticMesh* FromStaticMeshAsset,
+		FStaticMeshConversionOptions AssetOptions,
+		EMeshLODType LODType,
+		int32 LODIndex,
+		FDynamicMesh3& OutMesh,
+		FText& OutErrorMessage
+	)
+	{
+		using namespace ::UE::Geometry;
+
+		bool bSuccess = false;
+		OutMesh.Clear();
+
+		if (!FromStaticMeshAsset)
+		{
+			OutErrorMessage = LOCTEXT("CopyMeshFromStaticMeshSource_NullMesh", "Static Mesh is null");
+			return false;
+		}
+
+		if (LODType != EMeshLODType::MaxAvailable && LODType != EMeshLODType::SourceModel && LODType != EMeshLODType::HiResSourceModel)
+		{
+			OutErrorMessage = LOCTEXT("CopyMeshFromStaticMesh_LODNotAvailable", "Requested LOD Type is not available");
+			return false;
+		}
+
+#if WITH_EDITOR
+		if (LODType == EMeshLODType::HiResSourceModel && FromStaticMeshAsset->IsHiResMeshDescriptionValid() == false)
+		{
+			OutErrorMessage = LOCTEXT("CopyMeshFromStaticMesh_HiResLODNotAvailable", "HiResSourceModel LOD Type is not available");
+			return false;
+		}
+
+		const FMeshDescription* SourceMesh = nullptr;
+		const FMeshBuildSettings* BuildSettings = nullptr;
+
+		if ((LODType == EMeshLODType::HiResSourceModel) ||
+			(LODType == EMeshLODType::MaxAvailable && FromStaticMeshAsset->IsHiResMeshDescriptionValid()))
+		{
+			SourceMesh = FromStaticMeshAsset->GetHiResMeshDescription();
+			const FStaticMeshSourceModel& SourceModel = FromStaticMeshAsset->GetHiResSourceModel();
+			BuildSettings = &SourceModel.BuildSettings;
+		}
+		else
+		{
+			int32 UseLODIndex = FMath::Clamp(LODIndex, 0, FromStaticMeshAsset->GetNumSourceModels() - 1);
+			SourceMesh = FromStaticMeshAsset->GetMeshDescription(UseLODIndex);
+			const FStaticMeshSourceModel& SourceModel = FromStaticMeshAsset->GetSourceModel(UseLODIndex);
+			BuildSettings = &SourceModel.BuildSettings;
+		}
+
+		if (SourceMesh == nullptr)
+		{
+			OutErrorMessage = LOCTEXT("CopyMeshFromStaticMesh_SourceLODIsNull", "Requested SourceModel LOD is null, only RenderData Mesh is available");
+			return false;
+		}
+
+		bool bHasDirtyBuildSettings = BuildSettings->bRecomputeNormals
+			|| (BuildSettings->bRecomputeTangents && AssetOptions.bRequestTangents);
+		bool bNeedsBuildScale = AssetOptions.bUseBuildScale && BuildSettings && !BuildSettings->BuildScale3D.Equals(FVector::OneVector);
+		bool bNeedsOtherBuildSettings = AssetOptions.bApplyBuildSettings && bHasDirtyBuildSettings;
+
+		FMeshDescription LocalSourceMeshCopy;
+		if (bNeedsBuildScale || bNeedsOtherBuildSettings)
+		{
+			LocalSourceMeshCopy = *SourceMesh;
+
+			FStaticMeshAttributes Attributes(LocalSourceMeshCopy);
+
+			if (bNeedsBuildScale)
+			{
+				FTransform BuildScaleTransform = FTransform::Identity;
+				BuildScaleTransform.SetScale3D(BuildSettings->BuildScale3D);
+				FStaticMeshOperations::ApplyTransform(LocalSourceMeshCopy, BuildScaleTransform, true /*use correct normal transforms*/);
+			}
+
+			if (bNeedsOtherBuildSettings)
+			{
+				if (!Attributes.GetTriangleNormals().IsValid() || !Attributes.GetTriangleTangents().IsValid())
+				{
+					// If these attributes don't exist, create them and compute their values for each triangle
+					FStaticMeshOperations::ComputeTriangleTangentsAndNormals(LocalSourceMeshCopy);
+				}
+
+				EComputeNTBsFlags ComputeNTBsOptions = EComputeNTBsFlags::BlendOverlappingNormals;
+				ComputeNTBsOptions |= BuildSettings->bRecomputeNormals ? EComputeNTBsFlags::Normals : EComputeNTBsFlags::None;
+				if (AssetOptions.bRequestTangents)
+				{
+					ComputeNTBsOptions |= BuildSettings->bRecomputeTangents ? EComputeNTBsFlags::Tangents : EComputeNTBsFlags::None;
+					ComputeNTBsOptions |= BuildSettings->bUseMikkTSpace ? EComputeNTBsFlags::UseMikkTSpace : EComputeNTBsFlags::None;
+				}
+				ComputeNTBsOptions |= BuildSettings->bComputeWeightedNormals ? EComputeNTBsFlags::WeightedNTBs : EComputeNTBsFlags::None;
+				if (AssetOptions.bIgnoreRemoveDegenerates == false)
+				{
+					ComputeNTBsOptions |= BuildSettings->bRemoveDegenerates ? EComputeNTBsFlags::IgnoreDegenerateTriangles : EComputeNTBsFlags::None;
+				}
+
+				FStaticMeshOperations::ComputeTangentsAndNormals(LocalSourceMeshCopy, ComputeNTBsOptions);
+			}
+
+			SourceMesh = &LocalSourceMeshCopy;
+		}
+
+		FMeshDescriptionToDynamicMesh Converter;
+		Converter.Convert(SourceMesh, OutMesh, AssetOptions.bRequestTangents);
+
+		bSuccess = true;
+#else
+		OutErrorMessage = LOCTEXT("CopyMeshFromAsset_EditorOnly", "Source Models are not available at Runtime");
+#endif
+
+		return bSuccess;
+	}
+
+
+
+	static bool CopyMeshFromStaticMesh_RenderData(
+		UStaticMesh* FromStaticMeshAsset,
+		FStaticMeshConversionOptions AssetOptions,
+		EMeshLODType LODType,
+		int32 LODIndex,
+		FDynamicMesh3& OutMesh,
+		FText& OutErrorMessage
+	)
+	{
+		using namespace ::UE::Geometry;
+
+		OutMesh.Clear();
+
+		if (LODType != EMeshLODType::MaxAvailable && LODType != EMeshLODType::RenderData)
+		{
+			OutErrorMessage = LOCTEXT("CopyMeshFromStaticMeshRender_LODNotAvailable", "Requested LOD Type is not available");
+			return false;
+		}
+
+#if !WITH_EDITOR
+		if (FromStaticMeshAsset->bAllowCPUAccess == false)
+		{
+			OutErrorMessage = LOCTEXT("CopyMeshFromStaticMesh_CPUAccess", "StaticMesh bAllowCPUAccess must be set to true to read mesh data at Runtime");
+			return false;
+		}
+#endif
+
+		int32 UseLODIndex = FMath::Clamp(LODIndex, 0, FromStaticMeshAsset->GetNumLODs() - 1);
+
+		const FStaticMeshLODResources* LODResources = nullptr;
+		if (FStaticMeshRenderData* RenderData = FromStaticMeshAsset->GetRenderData())
+		{
+			LODResources = &RenderData->LODResources[UseLODIndex];
+		}
+		if (LODResources == nullptr)
+		{
+			OutErrorMessage = LOCTEXT("CopyMeshFromStaticMesh_NoLODResources", "LOD Data is not available");
+			return false;
+		}
+
+		FStaticMeshLODResourcesToDynamicMesh::ConversionOptions ConvertOptions;
+#if WITH_EDITOR
+		if (AssetOptions.bUseBuildScale)
+		{
+			// respect BuildScale build setting
+			const FMeshBuildSettings& LODBuildSettings = FromStaticMeshAsset->GetSourceModel(UseLODIndex).BuildSettings;
+			ConvertOptions.BuildScale = (FVector3d)LODBuildSettings.BuildScale3D;
+		}
+#else
+		if (!AssetOptions.bUseBuildScale)
+		{
+			OutErrorMessage = LOCTEXT("CopyMeshFromStaticMesh_BuildScaleAlreadyBaked", "Requested mesh without BuildScale, but BuildScale is already baked into the RenderData.");
+			return false;
+		}
+#endif
+
+		FStaticMeshLODResourcesToDynamicMesh Converter;
+		Converter.Convert(LODResources, ConvertOptions, OutMesh);
+
+		return true;
+	}
+
+	static bool CopyMeshFromStaticMesh(
+		UStaticMesh* FromStaticMeshAsset,
+		FStaticMeshConversionOptions AssetOptions,
+		EMeshLODType LODType,
+		int32 LODIndex,
+		bool bUseClosestLOD,
+		FDynamicMesh3& OutMesh,
+		FText& OutErrorMessage
+	)
+	{
+		if (!FromStaticMeshAsset)
+		{
+			OutErrorMessage = LOCTEXT("CopyMeshFromStaticMeshRender_NullMesh", "Static Mesh is null");
+			return false;
+		}
+
+		if (bUseClosestLOD)
+		{
+			// attempt to detect if an unavailable LOD was requested, and if so re-map to an available one
+			if (LODType == EMeshLODType::MaxAvailable || LODType == EMeshLODType::HiResSourceModel)
+			{
+				LODIndex = 0;
+			}
+#if WITH_EDITOR
+			if (LODType == EMeshLODType::MaxAvailable)
+			{
+				LODType = EMeshLODType::HiResSourceModel;
+			}
+			if (LODType == EMeshLODType::HiResSourceModel && !FromStaticMeshAsset->IsHiResMeshDescriptionValid())
+			{
+				LODType = EMeshLODType::SourceModel;
+			}
+			if (LODType == EMeshLODType::SourceModel)
+			{
+				LODIndex = FMath::Clamp(LODIndex, 0, FromStaticMeshAsset->GetNumSourceModels() - 1);
+				if (!FromStaticMeshAsset->GetSourceModel(LODIndex).IsSourceModelInitialized())
+				{
+					LODType = EMeshLODType::RenderData;
+				}
+			}
+#else
+			LODType = EMeshLODType::RenderData;
+#endif
+			if (LODType == EMeshLODType::RenderData)
+			{
+				LODIndex = FMath::Clamp(LODIndex, 0, FromStaticMeshAsset->GetNumLODs() - 1);
+			}
+		}
+
+		if (LODType == EMeshLODType::RenderData)
+		{
+			return CopyMeshFromStaticMesh_RenderData(FromStaticMeshAsset, AssetOptions, LODType, LODIndex, OutMesh, OutErrorMessage);
+		}
+		else
+		{
+			return CopyMeshFromStaticMesh_SourceData(FromStaticMeshAsset, AssetOptions, LODType, LODIndex, OutMesh, OutErrorMessage);
+		}
+	}
+}
+
+
+bool SceneComponentToDynamicMesh(USceneComponent* Component, const FToMeshOptions& Options, bool bTransformToWorld, 
+	Geometry::FDynamicMesh3& OutMesh, FTransform& OutLocalToWorld, FText& OutErrorMessage,
+	TArray<UMaterialInterface*>* OutComponentMaterials, TArray<UMaterialInterface*>* OutAssetMaterials)
+{
+	using namespace ::UE::Geometry;
+
+	bool bSuccess = false;
+	OutMesh.Clear();
+
+	OutLocalToWorld = Component->GetComponentTransform();
+
+	auto GetPrimitiveComponentMaterials = [](UPrimitiveComponent* PrimComp, TArray<UMaterialInterface*> Materials)
+	{
+		int32 NumMaterials = PrimComp->GetNumMaterials();
+		Materials.SetNum(NumMaterials);
+		for (int32 k = 0; k < NumMaterials; ++k)
+		{
+			Materials[k] = PrimComp->GetMaterial(k);
+		}
+	};
+
+	// if Component Materials were requested, try to get them generically off the primitive component
+	// Note: Currently all supported types happen to be primitive components as well; will need to update if this changes
+	if (OutComponentMaterials)
+	{
+		OutComponentMaterials->Empty();
+		if (UPrimitiveComponent* PrimComp = Cast<UPrimitiveComponent>(Component))
+		{
+			GetPrimitiveComponentMaterials(PrimComp, *OutComponentMaterials);
+		}
+	}
+
+
+	if (USkinnedMeshComponent* SkinnedMeshComponent = Cast<USkinnedMeshComponent>(Component))
+	{
+		const int32 NumLODs = SkinnedMeshComponent->GetNumLODs();
+		int32 RequestedLOD = Options.LODType == EMeshLODType::MaxAvailable ? 0 : Options.LODIndex;
+		if (Options.bUseClosestLOD)
+		{
+			RequestedLOD = FMath::Clamp(RequestedLOD, 0, NumLODs - 1);
+		}
+		if (RequestedLOD < 0 || RequestedLOD > NumLODs - 1)
+		{
+			OutErrorMessage = LOCTEXT("CopyMeshFromComponent_MissingSkinnedMeshComponentLOD", "SkinnedMeshComponent requested LOD does not exist");
+		}
+		else
+		{
+			USkinnedAsset* SkinnedAsset = SkinnedMeshComponent->GetSkinnedAsset();
+			if (SkinnedAsset)
+			{
+				SkinnedMeshComponentToDynamicMesh(*SkinnedMeshComponent, OutMesh, RequestedLOD, Options.bWantTangents);
+				OutMesh.DiscardTriangleGroups();
+
+				if (OutAssetMaterials)
+				{
+					const TArray<FSkeletalMaterial>& Materials = SkinnedAsset->GetMaterials();
+					OutAssetMaterials->SetNum(Materials.Num());
+					for (int32 k = 0; k < Materials.Num(); ++k)
+					{
+						(*OutAssetMaterials)[k] = Materials[k].MaterialInterface;
+					}
+				}
+				bSuccess = true;
+			}
+			else
+			{
+				OutErrorMessage = LOCTEXT("CopyMeshFromComponent_MissingSkinnedAsset", "SkinnedMeshComponent has a null SkinnedAsset");
+			}
+		}
+
+	}
+	else if (USplineMeshComponent* SplineMeshComponent = Cast<USplineMeshComponent>(Component))
+	{
+		UStaticMesh* StaticMesh = SplineMeshComponent->GetStaticMesh();
+		if (StaticMesh)
+		{
+			Private::ConversionHelper::FStaticMeshConversionOptions AssetOptions;
+			AssetOptions.bApplyBuildSettings = (Options.bWantNormals || Options.bWantTangents);
+			AssetOptions.bRequestTangents = Options.bWantTangents;
+			bSuccess = Private::ConversionHelper::CopyMeshFromStaticMesh(
+				StaticMesh, AssetOptions, Options.LODType, Options.LODIndex, Options.bUseClosestLOD, OutMesh, OutErrorMessage);
+
+			// deform the dynamic mesh and its tangent space with the spline
+			if (bSuccess)
+			{
+				const bool bUpdateTangentSpace = Options.bWantTangents;
+				SplineDeformDynamicMesh(*SplineMeshComponent, OutMesh, bUpdateTangentSpace);
+
+				if (OutAssetMaterials)
+				{
+					int32 NumMaterials = StaticMesh->GetStaticMaterials().Num();
+					OutAssetMaterials->SetNum(NumMaterials);
+					for (int32 k = 0; k < NumMaterials; ++k)
+					{
+						(*OutAssetMaterials)[k] = StaticMesh->GetMaterial(k);
+					}
+				}
+			}
+		}
+		else
+		{
+			OutErrorMessage = LOCTEXT("CopyMeshFromSplineMeshComponent_MissingStaticMesh", "SplineMeshComponent has a null StaticMesh");
+		}
+	}
+	else if (UStaticMeshComponent* StaticMeshComponent = Cast<UStaticMeshComponent>(Component))
+	{
+		UStaticMesh* StaticMesh = StaticMeshComponent->GetStaticMesh();
+		if (StaticMesh)
+		{
+			Private::ConversionHelper::FStaticMeshConversionOptions AssetOptions;
+			AssetOptions.bApplyBuildSettings = (Options.bWantNormals || Options.bWantTangents);
+			AssetOptions.bRequestTangents = Options.bWantTangents;
+			bSuccess = Private::ConversionHelper::CopyMeshFromStaticMesh(
+				StaticMesh, AssetOptions, Options.LODType, Options.LODIndex, Options.bUseClosestLOD, OutMesh, OutErrorMessage);
+
+			// if we have an ISMC, append instances
+			if (UInstancedStaticMeshComponent* ISMComponent = Cast<UInstancedStaticMeshComponent>(StaticMeshComponent))
+			{
+				FDynamicMesh3 InstancedMesh = MoveTemp(OutMesh);
+				OutMesh.Clear();
+
+				FDynamicMesh3 AccumMesh;
+				AccumMesh.EnableMatchingAttributes(InstancedMesh);
+				FDynamicMeshEditor Editor(&AccumMesh);
+				FMeshIndexMappings Mappings;
+
+				int32 NumInstances = ISMComponent->GetInstanceCount();
+				for (int32 InstanceIdx = 0; InstanceIdx < NumInstances; ++InstanceIdx)
+				{
+					if (ISMComponent->IsValidInstance(InstanceIdx))
+					{
+						FTransform InstanceTransform;
+						ISMComponent->GetInstanceTransform(InstanceIdx, InstanceTransform, /*bWorldSpace=*/false);
+						FTransformSRT3d XForm(InstanceTransform);
+
+						Mappings.Reset();
+						Editor.AppendMesh(&InstancedMesh, Mappings,
+							[&](int, const FVector3d& Position) { return XForm.TransformPosition(Position); },
+							[&](int, const FVector3d& Normal) { return XForm.TransformNormal(Normal); });
+					}
+				}
+
+				OutMesh = MoveTemp(AccumMesh);
+			}
+
+			if (OutAssetMaterials)
+			{
+				int32 NumMaterials = StaticMesh->GetStaticMaterials().Num();
+				OutAssetMaterials->SetNum(NumMaterials);
+				for (int32 k = 0; k < NumMaterials; ++k)
+				{
+					(*OutAssetMaterials)[k] = StaticMesh->GetMaterial(k);
+				}
+			}
+		}
+		else
+		{
+			OutErrorMessage = LOCTEXT("CopyMeshFromComponent_MissingStaticMesh", "StaticMeshComponent has a null StaticMesh");
+		}
+	}
+	else if (UDynamicMeshComponent* DynamicMeshComponent = Cast<UDynamicMeshComponent>(Component))
+	{
+		UDynamicMesh* CopyDynamicMesh = DynamicMeshComponent->GetDynamicMesh();
+		if (CopyDynamicMesh)
+		{
+			CopyDynamicMesh->ProcessMesh([&](const FDynamicMesh3& Mesh)
+			{
+				OutMesh = Mesh;
+			});
+			bSuccess = true;
+		}
+		else
+		{
+			OutErrorMessage = LOCTEXT("CopyMeshFromComponent_MissingDynamicMesh", "DynamicMeshComponent has a null DynamicMesh");
+		}
+	}
+	else if (UBrushComponent* BrushComponent = Cast<UBrushComponent>(Component))
+	{
+		FVolumeToMeshOptions VolOptions;
+		VolOptions.bMergeVertices = true;
+		VolOptions.bAutoRepairMesh = true;
+		VolOptions.bOptimizeMesh = true;
+		VolOptions.bSetGroups = true;
+
+		OutMesh.EnableTriangleGroups();
+		BrushComponentToDynamicMesh(BrushComponent, OutMesh, VolOptions);
+
+		// compute normals for current polygroup topology
+		OutMesh.EnableAttributes();
+		if (Options.bWantNormals)
+		{
+			FDynamicMeshNormalOverlay* Normals = OutMesh.Attributes()->PrimaryNormals();
+			FMeshNormals::InitializeOverlayTopologyFromFaceGroups(&OutMesh, Normals);
+			FMeshNormals::QuickRecomputeOverlayNormals(OutMesh);
+		}
+
+		if (OutMesh.TriangleCount() > 0)
+		{
+			bSuccess = true;
+		}
+		else
+		{
+			OutErrorMessage = LOCTEXT("CopyMeshFromComponent_InvalidBrushConversion", "BrushComponent conversion produced 0 triangles");
+		}
+	}
+	else if (UGeometryCollectionComponent* GeometryCollectionComponent = Cast<UGeometryCollectionComponent>(Component))
+	{
+		if (const UGeometryCollection* RestCollection = GeometryCollectionComponent->GetRestCollection())
+		{
+			if (const FGeometryCollection* Collection = RestCollection->GetGeometryCollection().Get())
+			{
+				FTransform UnusedTransform;
+				const TArray<FTransform3f>& DynamicTransforms = GeometryCollectionComponent->GetComponentSpaceTransforms3f();
+				if (!DynamicTransforms.IsEmpty())
+				{
+					ConvertGeometryCollectionToDynamicMesh(OutMesh, UnusedTransform, false, *Collection, true, DynamicTransforms, false, Collection->TransformIndex.GetConstArray());
+				}
+				else
+				{
+					ConvertGeometryCollectionToDynamicMesh(OutMesh, UnusedTransform, false, *Collection, true, TArrayView<const FTransform3f>(Collection->Transform.GetConstArray()), true, Collection->TransformIndex.GetConstArray());
+				}
+				bSuccess = true;
+
+				if (OutAssetMaterials)
+				{
+					//const TArray<TObjectPtr<UMaterialInterface>>& AssetMaterials = RestCollection->Materials;
+					*OutAssetMaterials = RestCollection->Materials;
+				}
+			}
+			else
+			{
+				OutErrorMessage = LOCTEXT("CopyMeshFromComponent_MissingCollectionData", "GeometryCollectionComponent has null Geometry Collection data");
+			}
+		}
+		else
+		{
+			OutErrorMessage = LOCTEXT("CopyMeshFromComponent_MissingRestCollection", "GeometryCollectionComponent has null Rest Collection object");
+		}
+	}
+
+	// transform mesh to world
+	if (bSuccess && bTransformToWorld)
+	{
+		MeshTransforms::ApplyTransform(OutMesh, (FTransformSRT3d)OutLocalToWorld, true);
+	}
+
+	return bSuccess;
+}
+
+} // end namespace Conversion
+} // end namespace UE
+
+
+#undef LOCTEXT_NAMESPACE
+
diff --git a/Engine/Plugins/Runtime/MeshModelingToolset/Source/ModelingComponents/Public/ConversionUtils/SceneComponentToDynamicMesh.h b/Engine/Plugins/Runtime/MeshModelingToolset/Source/ModelingComponents/Public/ConversionUtils/SceneComponentToDynamicMesh.h
new file mode 100644
--- /dev/null
+++ b/Engine/Plugins/Runtime/MeshModelingToolset/Source/ModelingComponents/Public/ConversionUtils/SceneComponentToDynamicMesh.h
@@ -0,0 +1,86 @@
+// Copyright Epic Games, Inc. All Rights Reserved.
+
+#pragma once
+
+#include "Materials/MaterialInterface.h"
+#include "Math/Transform.h"
+#include "Internationalization/Text.h"
+
+class USceneComponent;
+
+namespace UE {
+namespace Geometry{ class FDynamicMesh3; };
+
+namespace Conversion {
+
+	/**
+	 * The Type of LOD in a Mesh Asset. Note some options are only applicable to some asset types (e.g., only Static Meshes)
+	 */
+	enum class EMeshLODType : uint8
+	{
+		/** The Maximum-quality available SourceModel LOD (HiResSourceModel if it is available, otherwise SourceModel LOD0) */
+		MaxAvailable,
+		/** The HiRes SourceModel. LOD Index is ignored. HiResSourceModel is not available at Runtime. */
+		HiResSourceModel,
+		/**
+		 * The SourceModel mesh at a given LOD Index. Note that a StaticMesh Asset with Auto-Generated LODs may not have a valid SourceModel for every LOD Index
+		 * SourceModel meshes are not available at Runtime.
+		 */
+		SourceModel,
+		/**
+		 * The Render mesh at at given LOD Index.
+		 * A StaticMesh Asset derives its RenderData LODs from it's SourceModel LODs. RenderData LODs always exist for every valid LOD Index.
+		 * However the RenderData LODs are not identical to SourceModel LODs, in particular they will be split at UV seams, Hard Normal creases, etc.
+		 * RenderData LODs in a StaticMesh Asset are only available at Runtime if the bAllowCPUAccess flag was enabled on the Asset at Cook time.
+		 */
+		RenderData
+	};
+
+
+	// General options for converting/extracting a mesh from an asset or scene component
+	// Options may be ignored where not applicable (for example if the source type does not have LODs)
+	struct FToMeshOptions
+	{
+		// default constructor -- requests max available LOD with normals and tangents
+		FToMeshOptions() = default;
+		// LOD index constructor -- requests specified source model LOD index
+		FToMeshOptions(int32 LODIndex, bool bWantNormals = true, bool bWantTangents = true)
+			: LODType(EMeshLODType::SourceModel), LODIndex(LODIndex), bWantNormals(bWantNormals), bWantTangents(bWantTangents)
+		{}
+
+		EMeshLODType LODType = EMeshLODType::MaxAvailable;
+		// Which LOD to use. Ignored if the LODType is MaxAvailable or HiResSourceModel
+		int32 LODIndex = 0;
+		// Whether to fall back to the most similar LOD if the requested LOD is not available.
+		bool bUseClosestLOD = true;
+
+		bool bWantNormals = true;
+		bool bWantTangents = true;
+	};
+
+	/**
+	* Converts geometry from various types of Scene Component to Dynamic Mesh, or reports that it cannot.
+	* 
+	* @param SceneComponent			Component to attempt to convert to a mesh
+	* @param Options				Options controlling the conversion
+	* @param OutMesh				Stores the converted mesh
+	* @param OutTransform			Stores the world transform of the mesh
+	* @param OutErrorMessage		Stores any error message, if the conversion failed
+	* @param OutComponentMaterials	Stores materials used on component, if non-null and applicable
+	* @param OutAssetMaterials		Stores materials used on asset, if non-null and applicable
+	* @return true if the component could be converted, false otherwise
+	*/
+	bool MODELINGCOMPONENTS_API SceneComponentToDynamicMesh(
+		USceneComponent* SceneComponent, const FToMeshOptions& ConversionOptions, bool bTransformToWorld, 
+		Geometry::FDynamicMesh3& OutMesh, FTransform& OutTransform, FText& OutErrorMessage, 
+		TArray<UMaterialInterface*>* OutComponentMaterials = nullptr, TArray<UMaterialInterface*>* OutAssetMaterials = nullptr);
+	
+	/**
+	 * Helper to quickly test if we expect to be able to convert SceneComponent to a dynamic mesh
+	 * 
+	 * @return true if we expect SceneComponentToDynamicMesh to be able to convert this component to a dynamic mesh
+	 */
+	bool MODELINGCOMPONENTS_API CanConvertSceneComponentToDynamicMesh(USceneComponent* SceneComponent);
+
+} // end namespace Geometry
+} // end namespace UE
\ No newline at end of file
