fix static analysis complaint about always-true if statement in the non-editor build of SceneComponentToDynamicMesh.cpp

diff --git a/Engine/Plugins/Runtime/MeshModelingToolset/Source/ModelingComponents/Private/ConversionUtils/SceneComponentToDynamicMesh.cpp b/Engine/Plugins/Runtime/MeshModelingToolset/Source/ModelingComponents/Private/ConversionUtils/SceneComponentToDynamicMesh.cpp
--- a/Engine/Plugins/Runtime/MeshModelingToolset/Source/ModelingComponents/Private/ConversionUtils/SceneComponentToDynamicMesh.cpp
+++ b/Engine/Plugins/Runtime/MeshModelingToolset/Source/ModelingComponents/Private/ConversionUtils/SceneComponentToDynamicMesh.cpp
@@ -310,13 +310,14 @@ namespace Private::ConversionHelper
 					LODType = EMeshLODType::RenderData;
 				}
 			}
-#else
-			LODType = EMeshLODType::RenderData;
-#endif
 			if (LODType == EMeshLODType::RenderData)
 			{
 				LODIndex = FMath::Clamp(LODIndex, 0, FromStaticMeshAsset->GetNumLODs() - 1);
 			}
+#else
+			LODType = EMeshLODType::RenderData;
+			LODIndex = FMath::Clamp(LODIndex, 0, FromStaticMeshAsset->GetNumLODs() - 1);
+#endif
 		}
 
 		if (LODType == EMeshLODType::RenderData)
