ModelingTools: Remove unreachable ensure(false) that static analysis didn't like, and fix accidental duplicate LOCTEXT key.

diff --git a/Engine/Plugins/Editor/ModelingToolsEditorMode/Source/ModelingToolsEditorMode/Private/ModelingToolsHostCustomizationAPI.cpp b/Engine/Plugins/Editor/ModelingToolsEditorMode/Source/ModelingToolsEditorMode/Private/ModelingToolsHostCustomizationAPI.cpp
--- a/Engine/Plugins/Editor/ModelingToolsEditorMode/Source/ModelingToolsEditorMode/Private/ModelingToolsHostCustomizationAPI.cpp
+++ b/Engine/Plugins/Editor/ModelingToolsEditorMode/Source/ModelingToolsEditorMode/Private/ModelingToolsHostCustomizationAPI.cpp
@@ -58,9 +58,6 @@ UModelingToolsHostCustomizationAPI* UModelingToolsHostCustomizationAPI::Register
 		ensure(ToolsContext->ContextObjectStore->AddContextObject(Instance));
 		return Instance;
 	}
-
-	ensure(false); // Shouldn't get here
-	return nullptr;
 }
 
 bool UModelingToolsHostCustomizationAPI::Deregister(UInteractiveToolsContext* ToolsContext)
diff --git a/Engine/Plugins/Runtime/MeshModelingToolset/Source/MeshModelingTools/Private/EditMeshPolygonsTool.cpp b/Engine/Plugins/Runtime/MeshModelingToolset/Source/MeshModelingTools/Private/EditMeshPolygonsTool.cpp
--- a/Engine/Plugins/Runtime/MeshModelingToolset/Source/MeshModelingTools/Private/EditMeshPolygonsTool.cpp
+++ b/Engine/Plugins/Runtime/MeshModelingToolset/Source/MeshModelingTools/Private/EditMeshPolygonsTool.cpp
@@ -1373,7 +1373,7 @@ void UEditMeshPolygonsTool::StartActivity(TObjectPtr<UInteractiveToolActivity> A
 					Params.OverrideAcceptButtonText = LOCTEXT("AcceptSubActionButton", "Accept Action");
 					Params.OverrideAcceptButtonTooltip = LOCTEXT("AcceptSubActionTooltip", "Accept the action currently being performed.");
 					Params.OverrideCancelButtonText = LOCTEXT("CancelSubActionButton", "Cancel Action");
-					Params.OverrideCancelButtonTooltip = LOCTEXT("AcceptSubActionTooltip", "Cancel the action currently being performed.");
+					Params.OverrideCancelButtonTooltip = LOCTEXT("CancelSubActionTooltip", "Cancel the action currently being performed.");
 					Params.CanAccept = [this]() { return CurrentActivity->CanAccept(); };
 					Params.OnAcceptCancelTriggered = [this](bool bAccept) 
 					{ 
