PolyEdit: Add support for deleting polygroup edges

diff --git a/Engine/Plugins/Runtime/MeshModelingToolset/Source/MeshModelingTools/Private/EditMeshPolygonsTool.cpp b/Engine/Plugins/Runtime/MeshModelingToolset/Source/MeshModelingTools/Private/EditMeshPolygonsTool.cpp
--- a/Engine/Plugins/Runtime/MeshModelingToolset/Source/MeshModelingTools/Private/EditMeshPolygonsTool.cpp
+++ b/Engine/Plugins/Runtime/MeshModelingToolset/Source/MeshModelingTools/Private/EditMeshPolygonsTool.cpp
@@ -850,7 +850,9 @@ void UEditMeshPolygonsTool::RegisterActions(FInteractiveToolActionSet& ActionSet
 	auto OnDeletionKeyPress = [this]() 
 	{
 		if ((EditActions && EditActions->IsPropertySetEnabled())
-			|| (EditActions_Triangles && EditActions_Triangles->IsPropertySetEnabled()))
+			|| (EditActions_Triangles && EditActions_Triangles->IsPropertySetEnabled())
+			|| (EditEdgeActions && EditEdgeActions->IsPropertySetEnabled())
+			)
 		{
 			RequestAction(EEditMeshPolygonsToolActions::Delete);
 		}
@@ -1579,6 +1581,26 @@ void UEditMeshPolygonsTool::UpdateFromCurrentMesh(bool bUpdateTopology)
 
 
 
+void UEditMeshPolygonsTool::ApplyDelete()
+{
+	if (BeginMeshFaceEditChange())
+	{
+		ApplyDeleteFaces();
+	}
+	else if (BeginMeshEdgeEditChange())
+	{
+		ApplyDeleteEdges();
+	}
+	else
+	{
+		GetToolManager()->DisplayMessage(
+			LOCTEXT("OnDeleteFailedMessage", "Cannot Delete Current Selection"),
+			EToolMessageLevel::UserWarning);
+	}
+}
+
+
+
 void UEditMeshPolygonsTool::ApplyMerge()
 {
 	if (BeginMeshFaceEditChange() == false)
@@ -1612,16 +1634,8 @@ void UEditMeshPolygonsTool::ApplyMerge()
 
 
 
-void UEditMeshPolygonsTool::ApplyDelete()
+void UEditMeshPolygonsTool::ApplyDeleteFaces()
 {
-	if (BeginMeshFaceEditChange() == false)
-	{
-		GetToolManager()->DisplayMessage(
-			LOCTEXT("OnDeleteFailedMessage", "Cannot Delete Current Selection"),
-			EToolMessageLevel::UserWarning);
-		return;
-	}
-
 	FDynamicMesh3* Mesh = CurrentMesh.Get();
 
 	// prevent deleting all triangles
@@ -1640,7 +1654,7 @@ void UEditMeshPolygonsTool::ApplyDelete()
 	Editor.RemoveTriangles(ActiveTriangleSelection, true);
 
 	FGroupTopologySelection NewSelection;
-	EmitCurrentMeshChangeAndUpdate(LOCTEXT("PolyMeshDeleteChange", "Delete"),
+	EmitCurrentMeshChangeAndUpdate(LOCTEXT("PolyMeshDeleteFacesChange", "Delete Faces"),
 		ChangeTracker.EndChange(), NewSelection);
 }
 
@@ -2130,6 +2144,48 @@ void UEditMeshPolygonsTool::ApplyStraightenEdges()
 		ChangeTracker.EndChange(), NewSelection);
 }
 
+void UEditMeshPolygonsTool::ApplyDeleteEdges()
+{
+	FDynamicMesh3* Mesh = CurrentMesh.Get();
+	FDynamicMeshChangeTracker ChangeTracker(Mesh);
+	FGroupTopologySelection NewSelection;
+	FMeshConnectedComponents Components(Mesh);
+
+	// Using sets here because we only want unique triangles/edges
+	TSet<int32> EdgeIDs;
+	TSet<int32> SeedTriangleIDs;
+	for (FSelectedEdge& Edge : ActiveEdgeSelection)
+	{
+		for (int32 Eid : Edge.EdgeIDs)
+		{
+			FIndex2i AdjacentTriangles = Mesh->GetEdgeT(Eid);
+			EdgeIDs.Add(Eid);
+			SeedTriangleIDs.Add(AdjacentTriangles.A);
+			if (AdjacentTriangles.B != FDynamicMesh3::InvalidID)
+			{
+				SeedTriangleIDs.Add(AdjacentTriangles.B);
+			}
+		}
+	}
+	
+	Components.FindTrianglesConnectedToSeeds(SeedTriangleIDs.Array(), [&Mesh, &EdgeIDs](int32 t0, int32 t1)
+	{
+		return Mesh->GetTriangleGroup(t0) == Mesh->GetTriangleGroup(t1) || EdgeIDs.Contains(Mesh->FindEdgeFromTriPair(t0,t1));
+	});
+	
+	ChangeTracker.BeginChange();
+	
+	for (FMeshConnectedComponents::FComponent& Component : Components.Components)
+	{
+		ChangeTracker.SaveTriangles(Component.Indices, true);
+		int32 NewGroupID = Mesh->GetTriangleGroup(Component.Indices[0]);
+		FaceGroupUtil::SetGroupID(*Mesh, Component.Indices, NewGroupID);
+		NewSelection.SelectedGroupIDs.Add(NewGroupID);
+	}
+	
+	EmitCurrentMeshChangeAndUpdate(LOCTEXT("PolyMeshDeleteEdgesChange", "Delete Edges"),
+		ChangeTracker.EndChange(), NewSelection);
+}
 
 void UEditMeshPolygonsTool::ApplySimplifyAlongEdges()
 {
diff --git a/Engine/Plugins/Runtime/MeshModelingToolset/Source/MeshModelingTools/Public/EditMeshPolygonsTool.h b/Engine/Plugins/Runtime/MeshModelingToolset/Source/MeshModelingTools/Public/EditMeshPolygonsTool.h
--- a/Engine/Plugins/Runtime/MeshModelingToolset/Source/MeshModelingTools/Public/EditMeshPolygonsTool.h
+++ b/Engine/Plugins/Runtime/MeshModelingToolset/Source/MeshModelingTools/Public/EditMeshPolygonsTool.h
@@ -271,7 +271,7 @@ public:
 	void Merge() { PostAction(EEditMeshPolygonsToolActions::Merge);	}
 
 	/** Delete the current set of selected faces */
-	UFUNCTION(CallInEditor, Category = FaceEdits, meta = (DisplayName = "Delete", DisplayPriority = 4))
+	UFUNCTION(CallInEditor, Category = FaceEdits, meta = (DisplayName = "Delete Faces", DisplayPriority = 4))
 	void Delete() { PostAction(EEditMeshPolygonsToolActions::Delete); }
 
 	/** Cut the current set of selected faces. Click twice in viewport to set cut line. */
@@ -344,7 +344,7 @@ public:
 	void Outset() { PostAction(EEditMeshPolygonsToolActions::Outset);	}
 
 	/** Delete the current set of selected faces */
-	UFUNCTION(CallInEditor, Category = TriangleEdits, meta = (DisplayName = "Delete", DisplayPriority = 4))
+	UFUNCTION(CallInEditor, Category = TriangleEdits, meta = (DisplayName = "Delete Faces", DisplayPriority = 4))
 	void Delete() { PostAction(EEditMeshPolygonsToolActions::Delete); }
 
 	/** Cut the current set of selected faces. Click twice in viewport to set cut line. */
@@ -423,6 +423,10 @@ public:
 	/** Simplify the underlying triangulation along the selected edges, when doing so won't change the shape or UVs, or make low-quality triangles */
 	UFUNCTION(CallInEditor, Category = EdgeEdits, meta = (DisplayPriority = 6))
 	void Simplify() { PostAction(EEditMeshPolygonsToolActions::SimplifyAlongEdges); }
+	
+	/** Delete selected edge, implicitly merging any connected faces */
+	UFUNCTION(CallInEditor, Category = EdgeEdits, meta = (DisplayName = "Delete Edges", DisplayPriority = 7))
+	void DeleteEdge() { PostAction(EEditMeshPolygonsToolActions::Delete); }
 };
 
 
@@ -450,7 +454,6 @@ public:
 	/** Split the selected edges, inserting a new vertex at each edge midpoint */
 	UFUNCTION(CallInEditor, Category = EdgeEdits, meta = (DisplayName = "Split", DisplayPriority = 1))
 	void Split() { PostAction(EEditMeshPolygonsToolActions::SplitSingleEdge); }
-
 };
 
 
@@ -644,8 +647,10 @@ protected:
 
 	void UpdateGizmoVisibility();
 
-	void ApplyMerge();
 	void ApplyDelete();
+	
+	void ApplyMerge();
+	void ApplyDeleteFaces();
 	void ApplyRecalcNormals();
 	void ApplyFlipNormals();
 	void ApplyRetriangulate();
@@ -657,6 +662,7 @@ protected:
 	void ApplyCollapseEdge();
 	void ApplyWeldEdges();
 	void ApplyStraightenEdges();
+	void ApplyDeleteEdges();
 	void ApplyFillHole();
 	void ApplyBridgeEdges();
 	void ApplySimplifyAlongEdges();
