Modeling Tools: Adding AtOrigin functionality to AddPrimitiveTool placement

diff --git a/Engine/Plugins/Runtime/MeshModelingToolset/Source/MeshModelingTools/Private/AddPrimitiveTool.cpp b/Engine/Plugins/Runtime/MeshModelingToolset/Source/MeshModelingTools/Private/AddPrimitiveTool.cpp
--- a/Engine/Plugins/Runtime/MeshModelingToolset/Source/MeshModelingTools/Private/AddPrimitiveTool.cpp
+++ b/Engine/Plugins/Runtime/MeshModelingToolset/Source/MeshModelingTools/Private/AddPrimitiveTool.cpp
@@ -116,6 +116,8 @@ void UAddPrimitiveTool::Setup()
 	AddToolPropertySource(OutputTypeProperties);
 
 	AddToolPropertySource(ShapeSettings);
+	
+	ShapeSettings->WatchProperty(ShapeSettings->TargetSurface, [this](EMakeMeshPlacementType){UpdateTargetSurface();});
 	ShapeSettings->RestoreProperties(this);
 
 	MaterialProperties = NewObject<UNewMeshMaterialProperties>(this);
@@ -147,7 +149,6 @@ void UAddPrimitiveTool::Setup()
 	DragAlignmentMechanic->AddToGizmo(Gizmo);
 
 	UpdatePreviewMesh();
-
 	SetState(EState::PlacingPrimitive);
 }
 
@@ -277,7 +278,7 @@ void UAddPrimitiveTool::UpdatePreviewPosition(const FInputDeviceRay& DeviceClick
 		bHit = true;
 		ShapeFrame = FFrame3d(DrawPlanePos);
 	}
-	else
+	else if (ShapeSettings->TargetSurface == EMakeMeshPlacementType::OnScene)
 	{
 		// cast ray into scene
 		FHitResult Result;
@@ -300,6 +301,11 @@ void UAddPrimitiveTool::UpdatePreviewPosition(const FInputDeviceRay& DeviceClick
 			ShapeFrame = FFrame3d(DrawPlanePos);
 		}
 	}
+	else
+	{
+		bHit = true;
+		ShapeFrame = FFrame3d();
+	}
 
 	// Snap to grid
 	USceneSnappingManager* SnapManager = USceneSnappingManager::Find(GetToolManager());
@@ -374,6 +380,24 @@ void UAddPrimitiveTool::UpdatePreviewMesh() const
 	checkSlow(CalculateTangentsSuccessful);
 }
 
+void UAddPrimitiveTool::UpdateTargetSurface()
+{
+	if (ShapeSettings->TargetSurface == EMakeMeshPlacementType::AtOrigin)
+	{
+		// default ray is used as coordinates will not be needed to set position at origin
+		const FInputDeviceRay DefaultRay = FInputDeviceRay();
+		UpdatePreviewPosition(DefaultRay);
+
+		SetState(EState::AdjustingSettings);
+		GetToolManager()->EmitObjectChange(this, MakeUnique<FStateChange>(PreviewMesh->GetTransform()),
+				LOCTEXT("PlaceMeshTransaction", "Place Mesh"));
+	} else
+	{
+		SetState(EState::PlacingPrimitive);
+	}
+}
+
+
 
 bool UAddPrimitiveTool::SupportsWorldSpaceFocusBox()
 {
diff --git a/Engine/Plugins/Runtime/MeshModelingToolset/Source/MeshModelingTools/Public/AddPrimitiveTool.h b/Engine/Plugins/Runtime/MeshModelingToolset/Source/MeshModelingTools/Public/AddPrimitiveTool.h
--- a/Engine/Plugins/Runtime/MeshModelingToolset/Source/MeshModelingTools/Public/AddPrimitiveTool.h
+++ b/Engine/Plugins/Runtime/MeshModelingToolset/Source/MeshModelingTools/Public/AddPrimitiveTool.h
@@ -50,7 +50,8 @@ UENUM()
 enum class EMakeMeshPlacementType : uint8
 {
 	GroundPlane = 0,
-	OnScene     = 1
+	OnScene     = 1,
+	AtOrigin	= 2
 };
 
 /** Placement Pivot Location */
@@ -86,7 +87,7 @@ public:
 	EMakeMeshPolygroupMode PolygroupMode = EMakeMeshPolygroupMode::PerFace;
 
 	/** How the shape is placed in the scene. */
-	UPROPERTY(EditAnywhere, Category = Positioning)
+	UPROPERTY(EditAnywhere, Category = Positioning, meta = (DisplayName = "Target Position"))
 	EMakeMeshPlacementType TargetSurface = EMakeMeshPlacementType::OnScene;
 
 	/** Location of pivot within the shape */
@@ -517,6 +518,8 @@ protected:
 
 	void UpdatePreviewMesh() const;
 
+	void UpdateTargetSurface();
+
 	// Used to make the initial placement of the mesh undoable
 	class FStateChange : public FToolCommandChange
 	{
