Modeling Tools: Allows 'Q' hotkey to return to 'select' in polygroup edit by changing hotkey for 'Lock Rotation' from 'Q' to 'Ctrl+R'

diff --git a/Engine/Plugins/Runtime/MeshModelingToolset/Source/MeshModelingTools/Private/EditMeshPolygonsTool.cpp b/Engine/Plugins/Runtime/MeshModelingToolset/Source/MeshModelingTools/Private/EditMeshPolygonsTool.cpp
--- a/Engine/Plugins/Runtime/MeshModelingToolset/Source/MeshModelingTools/Private/EditMeshPolygonsTool.cpp
+++ b/Engine/Plugins/Runtime/MeshModelingToolset/Source/MeshModelingTools/Private/EditMeshPolygonsTool.cpp
@@ -66,11 +66,11 @@ namespace EditMeshPolygonsToolLocals
 {
 	FText PolyEditDefaultMessage = LOCTEXT("OnStartEditMeshPolygonsTool_TriangleMode", "Select triangles to edit mesh. Use middle mouse on gizmo to "
 		"reposition it. Hold Ctrl while translating or (in local mode) rotating to align to scene. Shift and Ctrl "
-		"change marquee select behavior. Q toggles Gizmo Orientation Lock.");
+		"change marquee select behavior. Ctrl+R toggles Gizmo Orientation Lock.");
 
 	FText TriEditDefaultMessage = LOCTEXT("OnStartEditMeshPolygonsTool", "Select PolyGroups to edit mesh. Use middle mouse on gizmo to reposition it. "
 		"Hold Ctrl while translating or (in local mode) rotating to align to scene. Shift and Ctrl change marquee select "
-		"behavior. Q toggles Gizmo Orientation Lock.");
+		"behavior. Ctrl+R toggles Gizmo Orientation Lock.");
 
 	FString GetPropertyCacheIdentifier(bool bTriangleMode)
 	{
@@ -843,7 +843,7 @@ void UEditMeshPolygonsTool::RegisterActions(FInteractiveToolActionSet& ActionSet
 		TEXT("ToggleLockRotation"),
 		LOCTEXT("ToggleLockRotationUIName", "Lock Rotation"),
 		LOCTEXT("ToggleLockRotationTooltip", "Toggle Frame Rotation Lock on and off"),
-		EModifierKey::None, EKeys::Q,
+		EModifierKey::Control, EKeys::R,
 		[this]() { CommonProps->bLockRotation = !CommonProps->bLockRotation; });
 	
 	// Backspace and delete both trigger deletion (as long as the delete button is also enabled)
