MeshElementSelection: Make selectable mesh elements visible.

When we have a valid topology type (Triangle or Poly, not None), we render selectable mesh elements.

Added AccumulateElementsFromPredicate to the GeometrySelector interface which allows a predicate function to filter which geometry elements ought to be included in the provided FGeometrySelectionElements. The predicate is aware of the type of element (vert, edge, tri) and the selected element ID (which is a single value encoding the topology ID and the element ID). The predicate returns true for elements to include and false for elements to exclude. Internally, this function builds transient vert, edge, and tri selections with this predicate and then accumulates all selection elements into the provided FGeometrySelectionElements.

diff --git a/Engine/Plugins/Editor/ModelingToolsEditorMode/Source/ModelingToolsEditorMode/Private/ModelingToolsEditorMode.cpp b/Engine/Plugins/Editor/ModelingToolsEditorMode/Source/ModelingToolsEditorMode/Private/ModelingToolsEditorMode.cpp
--- a/Engine/Plugins/Editor/ModelingToolsEditorMode/Source/ModelingToolsEditorMode/Private/ModelingToolsEditorMode.cpp
+++ b/Engine/Plugins/Editor/ModelingToolsEditorMode/Source/ModelingToolsEditorMode/Private/ModelingToolsEditorMode.cpp
@@ -260,7 +260,7 @@ bool UModelingToolsEditorMode::ShouldDrawWidget() const
 	}
 
 	// hide standard xform gizmo if we have an active selection
-	if (GetSelectionManager() && GetSelectionManager()->HasSelection())
+	if (GetSelectionManager() && (GetSelectionManager()->HasSelection() || GetSelectionManager()->GetMeshTopologyMode() != UGeometrySelectionManager::EMeshTopologyMode::None))
 	{
 		return false;
 	}
diff --git a/Engine/Plugins/Runtime/MeshModelingToolset/Source/ModelingComponents/Private/Selection/DynamicMeshSelector.cpp b/Engine/Plugins/Runtime/MeshModelingToolset/Source/ModelingComponents/Private/Selection/DynamicMeshSelector.cpp
--- a/Engine/Plugins/Runtime/MeshModelingToolset/Source/ModelingComponents/Private/Selection/DynamicMeshSelector.cpp
+++ b/Engine/Plugins/Runtime/MeshModelingToolset/Source/ModelingComponents/Private/Selection/DynamicMeshSelector.cpp
@@ -628,6 +628,22 @@ void FBaseDynamicMeshSelector::AccumulateSelectionElements(const FGeometrySelect
 		bMapFacesToEdges);
 }
 
+void FBaseDynamicMeshSelector::AccumulateElementsFromPredicate(FGeometrySelectionElements& Elements, bool bTransformToWorld, bool bIsForPreview, bool bUseGroupTopology, TFunctionRef<bool(EGeometryElementType, FGeoSelectionID)> Predicate)
+{
+	auto AccumulateElementsOfTypeFromPredicate = [this, &Elements, &Predicate, bTransformToWorld, bIsForPreview, bUseGroupTopology](EGeometryElementType ElementType)
+	{
+		FGeometrySelection Selection;
+		Selection.ElementType = ElementType;
+		Selection.TopologyType = bUseGroupTopology ? EGeometryTopologyType::Polygroup : EGeometryTopologyType::Triangle;
+		InitializeSelectionFromPredicate(Selection, [&Predicate, ElementType](FGeoSelectionID InID){ return Predicate(ElementType, InID); });
+		AccumulateSelectionElements(Selection, Elements, bTransformToWorld, bIsForPreview);
+	};
+
+	AccumulateElementsOfTypeFromPredicate(EGeometryElementType::Vertex);
+	AccumulateElementsOfTypeFromPredicate(EGeometryElementType::Edge);
+	AccumulateElementsOfTypeFromPredicate(EGeometryElementType::Face);
+}
+
 
 
 
diff --git a/Engine/Plugins/Runtime/MeshModelingToolset/Source/ModelingComponents/Private/Selection/GeometrySelectionManager.cpp b/Engine/Plugins/Runtime/MeshModelingToolset/Source/ModelingComponents/Private/Selection/GeometrySelectionManager.cpp
--- a/Engine/Plugins/Runtime/MeshModelingToolset/Source/ModelingComponents/Private/Selection/GeometrySelectionManager.cpp
+++ b/Engine/Plugins/Runtime/MeshModelingToolset/Source/ModelingComponents/Private/Selection/GeometrySelectionManager.cpp
@@ -11,6 +11,7 @@
 #include "ToolContextInterfaces.h"
 #include "ToolDataVisualizer.h"
 #include "Selections/GeometrySelectionUtil.h"
+#include "Util/ColorConstants.h"
 
 #include UE_INLINE_GENERATED_CPP_BY_NAME(GeometrySelectionManager)
 
@@ -123,6 +124,9 @@ void UGeometrySelectionManager::SetSelectionElementTypeInternal(EGeometryElement
 			bool bEnableTopologyFilter = (Target->Selection.TopologyType == EGeometryTopologyType::Polygroup && Target->Selection.ElementType != EGeometryElementType::Vertex);
 			Target->SelectionEditor->UpdateQueryConfig(GetCurrentSelectionQueryConfig(), bEnableTopologyFilter);
 		}
+
+		MarkRenderCachesDirty();
+		ClearActivePreview();
 	}
 }
 
@@ -168,6 +172,9 @@ void UGeometrySelectionManager::SetMeshTopologyModeInternal(EMeshTopologyMode Ne
 			bool bEnableTopologyFilter = (Target->Selection.TopologyType == EGeometryTopologyType::Polygroup && Target->Selection.ElementType != EGeometryElementType::Vertex);
 			Target->SelectionEditor->UpdateQueryConfig(GetCurrentSelectionQueryConfig(), bEnableTopologyFilter);
 		}
+
+		MarkRenderCachesDirty();
+		ClearActivePreview();
 	}
 }
 
@@ -359,7 +366,7 @@ bool UGeometrySelectionManager::AddActiveTarget(FGeometryIdentifier TargetIdenti
 
 	ActiveTargetMap.Add(TargetIdentifier, SelectionTarget);
 	ActiveTargetReferences.Add(SelectionTarget);
-
+	
 	SelectionTarget->OnGeometryModifiedHandle = 
 		SelectionTarget->Selector->GetOnGeometryModifed().AddUObject(this, &UGeometrySelectionManager::OnTargetGeometryModified);
 
@@ -374,7 +381,7 @@ void UGeometrySelectionManager::SynchronizeActiveTargets(
 	TFunctionRef<void()> WillChangeActiveTargetsCallback)
 {
 	TArray<FGeometryIdentifier> Before = GetCurrentTargetIdentifiers();
-
+	
 	// currently only support single selection
 	if (DesiredActiveSet.Num() == 1)
 	{
@@ -570,7 +577,7 @@ TSharedPtr<UGeometrySelectionManager::FGeometrySelectionTarget> UGeometrySelecti
 	{
 		TargetCache.Add(TargetIdentifier, SelectionTarget);
 	}
-
+	
 	return SelectionTarget;
 }
 
@@ -649,7 +656,7 @@ void UGeometrySelectionManager::ClearSelection(bool bSaveSelectionBeforeClear)
 
 	GetTransactionsAPI()->EndUndoTransaction();
 
-	bSelectionRenderCachesDirty = true;
+	MarkRenderCachesDirty();
 	OnSelectionModified.Broadcast();
 }
 
@@ -689,7 +696,7 @@ void UGeometrySelectionManager::UpdateSelectionViaRaycast(
 		GetTransactionsAPI()->AppendChange(this, MoveTemp(DeltaChange), LOCTEXT("UpdateSelectionViaRaycast", "Change Selection"));
 		GetTransactionsAPI()->EndUndoTransaction();
 
-		bSelectionRenderCachesDirty = true;
+		MarkRenderCachesDirty();
 		OnSelectionModified.Broadcast();
 	}
 	else if (ResultOut.bSelectionMissed && UpdateConfig.ChangeType == EGeometrySelectionChangeType::Replace)
@@ -732,7 +739,7 @@ void UGeometrySelectionManager::UpdateSelectionViaConvex(
 		GetTransactionsAPI()->AppendChange(this, MoveTemp(DeltaChange), LOCTEXT("UpdateSelectionViaConvex", "Change Selection"));
 		GetTransactionsAPI()->EndUndoTransaction();
 
-		bSelectionRenderCachesDirty = true;
+		MarkRenderCachesDirty();
 		OnSelectionModified.Broadcast();
 	}
 	else if (ResultOut.bSelectionMissed && UpdateConfig.ChangeType == EGeometrySelectionChangeType::Replace)
@@ -777,7 +784,7 @@ bool UGeometrySelectionManager::BeginTrackedSelectionChange(FGeometrySelectionUp
 
 	if (bClearOnBegin && InitialTrackedDelta.IsEmpty() == false)
 	{
-		bSelectionRenderCachesDirty = true;
+		MarkRenderCachesDirty();
 		OnSelectionModified.Broadcast();
 	}
 
@@ -807,7 +814,7 @@ void UGeometrySelectionManager::AccumulateSelectionUpdate_Raycast(
 		ActiveTrackedDelta.Added.Append( ResultOut.SelectionDelta.Added );
 		ActiveTrackedDelta.Removed.Append( ResultOut.SelectionDelta.Removed );
 
-		bSelectionRenderCachesDirty = true;
+		MarkRenderCachesDirty();
 		OnSelectionModified.Broadcast();
 	}
 }
@@ -862,7 +869,7 @@ bool UGeometrySelectionManager::SetSelectionForComponent(UPrimitiveComponent* Co
 				NewSelectionChange->Before = InitialSelection;
 				GetTransactionsAPI()->AppendChange(this, MoveTemp(NewSelectionChange), LOCTEXT("NewSelection", "New Selection"));
 
-				bSelectionRenderCachesDirty = true;
+				MarkRenderCachesDirty();
 				OnSelectionModified.Broadcast();
 			}
 			return true;
@@ -922,7 +929,7 @@ bool UGeometrySelectionManager::RestoreSavedSelection()
 					NewSelectionChange->Before = InitialSelection;
 					GetTransactionsAPI()->AppendChange(this, MoveTemp(NewSelectionChange), LOCTEXT("RestoreSelection", "Restore Selection"));
 
-					bSelectionRenderCachesDirty = true;
+					MarkRenderCachesDirty();
 					OnSelectionModified.Broadcast();
 				}
 				bFound = true;
@@ -960,8 +967,7 @@ bool UGeometrySelectionManager::UpdateSelectionPreviewViaRaycast(
 	{
 		return false;
 	}
-	bool bUseSimplifiedPreviewHighlight = (CVarGeometrySelectionManager_FullSelectionHoverHighlights.GetValueOnGameThread() == 0);
-
+	
 	// currently only going to support one object, not sure how to support more yet...
 
 	FGeometrySelectionTarget* Target = ActiveTargetReferences[0].Get();
@@ -977,10 +983,7 @@ bool UGeometrySelectionManager::UpdateSelectionPreviewViaRaycast(
 	{
 		ActivePreviewSelection = MoveTemp(NewPreview.PreviewSelection);
 		CachedPreviewRenderElements.Reset();
-		if (ActivePreviewSelection.IsEmpty() == false)
-		{
-			Target->Selector->AccumulateSelectionElements(ActivePreviewSelection, CachedPreviewRenderElements, true, bUseSimplifiedPreviewHighlight);
-		}
+		MarkRenderCachesDirty();
 	}
 
 	return (ActivePreviewSelection.IsEmpty() == false);
@@ -988,12 +991,9 @@ bool UGeometrySelectionManager::UpdateSelectionPreviewViaRaycast(
 
 void UGeometrySelectionManager::ClearSelectionPreview()
 {
-	ActivePreviewSelection.Selection.Reset();
-	CachedPreviewRenderElements.Reset();
+	ClearActivePreview();
 }
 
-
-
 bool UGeometrySelectionManager::GetSelectionBounds(FGeometrySelectionBounds& BoundsOut) const
 {
 	BoundsOut = FGeometrySelectionBounds();
@@ -1216,7 +1216,7 @@ void UGeometrySelectionManager::ExecuteSelectionCommand(UGeometrySelectionEditCo
 	GetTransactionsAPI()->EndUndoTransaction();
 
 	// assume this is true for now
-	bSelectionRenderCachesDirty = true;
+	MarkRenderCachesDirty();
 	OnSelectionModified.Broadcast();
 }
 
@@ -1263,7 +1263,7 @@ void UGeometrySelectionManager::ApplyChange(IGeometrySelectionChange* Change)
 
 			if (ApplyDelta.IsEmpty() == false)
 			{
-				bSelectionRenderCachesDirty = true;
+				MarkRenderCachesDirty();
 				OnSelectionModified.Broadcast();
 			}
 
@@ -1292,7 +1292,7 @@ void UGeometrySelectionManager::RevertChange(IGeometrySelectionChange* Change)
 
 			if (RevertDelta.IsEmpty() == false)
 			{
-				bSelectionRenderCachesDirty = true;
+				MarkRenderCachesDirty();
 				OnSelectionModified.Broadcast();
 			}
 
@@ -1304,8 +1304,11 @@ void UGeometrySelectionManager::RevertChange(IGeometrySelectionChange* Change)
 
 void UGeometrySelectionManager::OnTargetGeometryModified(IGeometrySelector* Selector)
 {
-	bSelectionRenderCachesDirty = true;
-
+	CachedSelectableRenderElements.Reset();
+	
+	CachedSelectableRenderElements.SetNum(ActiveTargetReferences.Num());
+	
+	MarkRenderCachesDirty();
 	ClearActivePreview();
 }
 
@@ -1313,13 +1316,38 @@ void UGeometrySelectionManager::OnTargetGeometryModified(IGeometrySelector* Sele
 void UGeometrySelectionManager::UpdateSelectionRenderCacheOnTargetChange()
 {
 	CachedSelectionRenderElements.Reset();
+	CachedSelectableRenderElements.Reset();
+	
 	CachedSelectionRenderElements.SetNum(ActiveTargetReferences.Num());
-	bSelectionRenderCachesDirty = true;
-
+	CachedSelectableRenderElements.SetNum(ActiveTargetReferences.Num());
+	
+	MarkRenderCachesDirty();
 	ClearActivePreview();
 }
 
+void UGeometrySelectionManager::MarkRenderCachesDirty()
+{
+	bSelectionRenderCachesDirty = true;
+	bSelectableRenderCachesDirty = true;
+	bPreviewRenderCachesDirty = true;
+}
+
 void UGeometrySelectionManager::RebuildSelectionRenderCaches()
+{
+	RebuildSelectionRenderCache();
+	RebuildSelectableRenderCache();
+	RebuildPreviewRenderCache();
+}
+
+void UGeometrySelectionManager::ClearActivePreview()
+{
+	ActivePreviewSelection.Reset();
+	CachedPreviewRenderElements.Reset();
+	bSelectableRenderCachesDirty = true;
+	bPreviewRenderCachesDirty = true;
+}
+
+void UGeometrySelectionManager::RebuildSelectionRenderCache()
 {
 	if (bSelectionRenderCachesDirty == false)
 	{
@@ -1330,21 +1358,75 @@ void UGeometrySelectionManager::RebuildSelectionRenderCaches()
 	for (int32 k = 0; k < ActiveTargetReferences.Num(); ++k)
 	{
 		TSharedPtr<FGeometrySelectionTarget> Target = ActiveTargetReferences[k];
-
-		FGeometrySelectionElements& Elements = CachedSelectionRenderElements[k];
-		Elements.Reset();
-
-		Target->Selector->AccumulateSelectionElements(Target->Selection, Elements, true, false);
+		
+		FGeometrySelectionElements& SelectionElements = CachedSelectionRenderElements[k];
+		SelectionElements.Reset();
+		Target->Selector->AccumulateSelectionElements(Target->Selection, SelectionElements, true, false);
 	}
 
 	bSelectionRenderCachesDirty = false;
 }
 
+void UGeometrySelectionManager::RebuildSelectableRenderCache()
+{
+	if (bSelectableRenderCachesDirty == false || MeshTopologyMode == EMeshTopologyMode::None)
+	{
+		return;
+	}
 
-void UGeometrySelectionManager::ClearActivePreview()
+	check(ActiveTargetReferences.Num() == CachedSelectionRenderElements.Num());
+	for (int32 k = 0; k < ActiveTargetReferences.Num(); ++k)
+	{
+		TSharedPtr<FGeometrySelectionTarget> Target = ActiveTargetReferences[k];
+		FGeometrySelectionElements& AllElements = CachedSelectableRenderElements[k];
+		
+		AllElements.Reset();
+		
+		Target->Selector->AccumulateElementsFromPredicate(AllElements, true, false, MeshTopologyMode == EMeshTopologyMode::Polygroup, [Target, this](EGeometryElementType Type, FGeoSelectionID ID)
+		{
+			uint64 EncodedID = ID.Encoded();
+
+			// Selectable faces are not displayed directly, just implicitly via displayed edges.
+			if (Type == EGeometryElementType::Face)
+			{
+				return false;
+			}
+			
+			// Exclude selected elements from selectable elements
+			if (Target->Selection.ElementType == Type && Target->Selection.Selection.Contains(EncodedID))
+			{
+				return false;
+			}
+
+			// Exclude hovered elements from selectable elements
+			if (ActivePreviewSelection.ElementType == Type && ActivePreviewSelection.Selection.Contains(EncodedID))
+			{
+				return false;
+			}
+			
+			return true;
+		});
+	}
+
+	bSelectableRenderCachesDirty = false;
+}
+
+void UGeometrySelectionManager::RebuildPreviewRenderCache()
 {
-	ActivePreviewSelection.Reset();
-	CachedPreviewRenderElements.Reset();
+	if (bPreviewRenderCachesDirty == false || ActiveTargetReferences.Num() == 0)
+	{
+		return;
+	}
+
+	const bool bUseSimplifiedPreviewHighlight = (CVarGeometrySelectionManager_FullSelectionHoverHighlights.GetValueOnGameThread() == 0);
+	const TSharedPtr<FGeometrySelectionTarget> Target = ActiveTargetReferences[0];
+	
+	if (ActivePreviewSelection.IsEmpty() == false)
+	{
+		Target->Selector->AccumulateSelectionElements(ActivePreviewSelection, CachedPreviewRenderElements, true, bUseSimplifiedPreviewHighlight);
+	}
+	
+	bPreviewRenderCachesDirty = false;
 }
 
 
@@ -1383,12 +1465,25 @@ void UGeometrySelectionManager::DebugRender(IToolsContextRenderAPI* RenderAPI)
 
 	//const UMaterialInterface* TriangleMaterial = ToolSetupUtil::GetSelectionMaterial(FLinearColor(1.0f, 0, 0, 0.5f), nullptr, 0.5f);
 
-	RebuildSelectionRenderCaches();
-	for ( const FGeometrySelectionElements& Elements : CachedSelectionRenderElements )
+	if (!this->ToolsContext->ToolManager->HasAnyActiveTool())
 	{
-		ToolSelectionUtil::DebugRenderGeometrySelectionElements(RenderAPI, Elements, false);
+		RebuildSelectionRenderCaches();
+		
+		for ( const FGeometrySelectionElements& Elements : CachedSelectionRenderElements )
+		{
+			ToolSelectionUtil::DebugRender(RenderAPI, Elements, 4.f, LinearColors::Gold3f(), 10.f, LinearColors::Gold3f(), 6.f);
+		}
+
+		if (MeshTopologyMode != EMeshTopologyMode::None)
+		{
+			for ( const FGeometrySelectionElements& Elements : CachedSelectableRenderElements )
+			{
+				ToolSelectionUtil::DebugRender(RenderAPI, Elements, 2.f, LinearColors::Red3f(), 8.f, LinearColors::Red3f(), 5.f);
+			}
+		}
+		
+		ToolSelectionUtil::DebugRender(RenderAPI, CachedPreviewRenderElements, 4.f, LinearColors::Green3f(), 10.f, LinearColors::Green3f(), 7.f);
 	}
-	ToolSelectionUtil::DebugRenderGeometrySelectionElements(RenderAPI, CachedPreviewRenderElements, true);
 }
 
 
diff --git a/Engine/Plugins/Runtime/MeshModelingToolset/Source/ModelingComponents/Private/Selection/ToolSelectionUtil.cpp b/Engine/Plugins/Runtime/MeshModelingToolset/Source/ModelingComponents/Private/Selection/ToolSelectionUtil.cpp
--- a/Engine/Plugins/Runtime/MeshModelingToolset/Source/ModelingComponents/Private/Selection/ToolSelectionUtil.cpp
+++ b/Engine/Plugins/Runtime/MeshModelingToolset/Source/ModelingComponents/Private/Selection/ToolSelectionUtil.cpp
@@ -75,7 +75,30 @@ void ToolSelectionUtil::DebugRenderGeometrySelectionElements(
 	IToolsContextRenderAPI* RenderAPI,
 	const FGeometrySelectionElements& Elements,
 	bool bIsPreview)
+{	
+	DebugRender(
+		RenderAPI,
+		Elements,
+		bIsPreview ? 1.f : 3.f,
+		bIsPreview ? FLinearColor(1, 1, 0, 1) : FLinearColor(0, 0.3f, 0.95f, 1),
+		bIsPreview ? 5.f : 10.f,
+		bIsPreview ? FLinearColor(1, 1, 0, 1) : FLinearColor(0, 0.3f, 0.95f, 1));
+}
+
+void ToolSelectionUtil::DebugRender(
+	IToolsContextRenderAPI* RenderAPI,
+	const UE::Geometry::FGeometrySelectionElements& Elements,
+	float LineThickness,
+	FLinearColor LineColor,
+	float PointSize,
+	FLinearColor PointColor,
+	float DepthBias)
 {
+	if (RenderAPI == nullptr)
+	{
+		return;
+	}
+	
 	FPrimitiveDrawInterface* CurrentPDI = RenderAPI->GetPrimitiveDrawInterface();
 
 	// batch render all the triangles, vastly more efficient than drawing one by one!
@@ -95,22 +118,13 @@ void ToolSelectionUtil::DebugRenderGeometrySelectionElements(
 	FMaterialRenderProxy* MaterialRenderProxy = GEngine->ConstraintLimitMaterialX->GetRenderProxy();
 	MeshBuilder.Draw(CurrentPDI, FMatrix::Identity, MaterialRenderProxy, DepthPriority, false, false);
 
-
 	FToolDataVisualizer Visualizer;
-	Visualizer.bDepthTested = false;
+	Visualizer.DepthBias = DepthBias;
 	Visualizer.BeginFrame(RenderAPI);
 
-	if (bIsPreview)
-	{
-		Visualizer.SetLineParameters(FLinearColor(1, 1, 0, 1), 1.0f);
-		Visualizer.SetPointParameters(FLinearColor(1, 1, 0, 1), 5.0f);
-	}
-	else
-	{
-		Visualizer.SetLineParameters(FLinearColor(0,0.3f,0.95f,1), 3.0f);
-		Visualizer.SetPointParameters(FLinearColor(0,0.3f,0.95f,1), 10.0f);
-	}
-
+	Visualizer.SetLineParameters(LineColor, LineThickness);
+	Visualizer.SetPointParameters(PointColor, PointSize);
+	
 	for (const FSegment3d& Segment : Elements.Segments)
 	{
 		Visualizer.DrawLine(Segment.StartPoint(), Segment.EndPoint());
diff --git a/Engine/Plugins/Runtime/MeshModelingToolset/Source/ModelingComponents/Public/Selection/DynamicMeshSelector.h b/Engine/Plugins/Runtime/MeshModelingToolset/Source/ModelingComponents/Public/Selection/DynamicMeshSelector.h
--- a/Engine/Plugins/Runtime/MeshModelingToolset/Source/ModelingComponents/Public/Selection/DynamicMeshSelector.h
+++ b/Engine/Plugins/Runtime/MeshModelingToolset/Source/ModelingComponents/Public/Selection/DynamicMeshSelector.h
@@ -118,6 +118,14 @@ public:
 	virtual void GetSelectionFrame(const FGeometrySelection& Selection, UE::Geometry::FFrame3d& SelectionFrame, bool bTransformToWorld) override;
 	virtual void AccumulateSelectionBounds(const FGeometrySelection& Selection, FGeometrySelectionBounds& BoundsInOut, bool bTransformToWorld) override;
 	virtual void AccumulateSelectionElements(const FGeometrySelection& Selection, FGeometrySelectionElements& Elements, bool bTransformToWorld, bool bIsForPreview) override;
+	
+	virtual void AccumulateElementsFromPredicate(
+		FGeometrySelectionElements& Elements,
+		bool bTransformToWorld,
+		bool bIsForPreview,
+		bool bUseGroupTopology,
+		TFunctionRef<bool(UE::Geometry::EGeometryElementType, UE::Geometry::FGeoSelectionID)> Predicate
+	) override;
 
 	/**
 	 * UpdateAfterGeometryEdit should be called after editing the UDynamicMesh owned by the Selector (TargetMesh). 
diff --git a/Engine/Plugins/Runtime/MeshModelingToolset/Source/ModelingComponents/Public/Selection/GeometrySelectionManager.h b/Engine/Plugins/Runtime/MeshModelingToolset/Source/ModelingComponents/Public/Selection/GeometrySelectionManager.h
--- a/Engine/Plugins/Runtime/MeshModelingToolset/Source/ModelingComponents/Public/Selection/GeometrySelectionManager.h
+++ b/Engine/Plugins/Runtime/MeshModelingToolset/Source/ModelingComponents/Public/Selection/GeometrySelectionManager.h
@@ -152,7 +152,11 @@ public:
 		FInputRayHit& HitResultOut
 	);
 
-
+	/**
+	 * Invalidates all cached selection elements.
+	 */
+	void MarkRenderCachesDirty();
+	
 	//
 	// Selection Updates
 	//
@@ -278,9 +282,12 @@ public:
 		const FRay3d& WorldRay
 	);
 
+	/**
+	 * Resets the active preview selection and invalidates its associated cached render elements.
+	 */
 	virtual void ClearSelectionPreview();
 
-
+	
 	//
 	// Selection queries
 	//
@@ -448,25 +455,20 @@ protected:
 
 	void OnTargetGeometryModified(IGeometrySelector* Selector);
 
-	//
-	// 3D geometry for element selections of each ActiveTarget is cached
-	// to improve rendering performance
-	//
-	TArray<FGeometrySelectionElements> CachedSelectionRenderElements;
+	
+	// todo [nickolas.drake]: cane we move CachedSelectionRenderElements, CachedPreviewRenderElements, and bSelectionRenderCachesDirty to private?
+	
+	TArray<FGeometrySelectionElements> CachedSelectionRenderElements;				// Cached 3D geometry for current selection
 	bool bSelectionRenderCachesDirty = false;
 	void UpdateSelectionRenderCacheOnTargetChange();
 	void RebuildSelectionRenderCaches();
 
+	UE_DEPRECATED(5.6, "Direct acces to ActivePreviewSelection is deprecated!")
+    FGeometrySelection ActivePreviewSelection;										// Selection representing the active preview
+	FGeometrySelectionElements CachedPreviewRenderElements;							// Cached 3D geometry for active preview elements
+    void ClearActivePreview();
 
-	//
-	// 3D geometry for active hover/preview highlight
-	// note: currently only supporting single target here, will need to be refactored to handle multiple targets...
-	//
-	FGeometrySelection ActivePreviewSelection;
-	FGeometrySelectionElements CachedPreviewRenderElements;
-	void ClearActivePreview();
-
-
+	
 	// various change types need internal access
 
 	friend class FGeometrySelectionManager_SelectionTypeChange;
@@ -486,6 +488,23 @@ protected:
 
 private:
 
+	//
+	// 3D geometry for element selections of each ActiveTarget is cached
+	// to improve rendering performance
+	//
+	
+	// todo [nickolas.drake]: refactor Rebuild_X_RenderCache() functions below to populate a UPreviewGeometry with the accumulated elements and remove PDI rendering
+	
+	void RebuildSelectionRenderCache();
+	
+	TArray<FGeometrySelectionElements> CachedSelectableRenderElements;				// Cached 3D geometry for all selectable elements
+	void RebuildSelectableRenderCache();
+	bool bSelectableRenderCachesDirty = false;
+	
+	void RebuildPreviewRenderCache();
+	bool bPreviewRenderCachesDirty = false;
+
+	
 	// Tracks saved selection state. Useful when the selection is temporarily cleared (e.g., for a tool)
 	struct FSavedSelection
 	{
diff --git a/Engine/Plugins/Runtime/MeshModelingToolset/Source/ModelingComponents/Public/Selection/GeometrySelector.h b/Engine/Plugins/Runtime/MeshModelingToolset/Source/ModelingComponents/Public/Selection/GeometrySelector.h
--- a/Engine/Plugins/Runtime/MeshModelingToolset/Source/ModelingComponents/Public/Selection/GeometrySelector.h
+++ b/Engine/Plugins/Runtime/MeshModelingToolset/Source/ModelingComponents/Public/Selection/GeometrySelector.h
@@ -411,7 +411,23 @@ public:
 	 */
 	virtual void AccumulateSelectionElements(const FGeometrySelection& Selection, FGeometrySelectionElements& ElementsInOut, bool bTransformToWorld, bool bIsForPreview) = 0;
 
-
+	/**
+	 * Accumulate all geometric elements (currently 3D triangles, line segments, and points) in the provided ElementsInOut
+	 * by initializing vertex, edge, and face selections using the provided predicate and then accumulating elements for those selections.
+	 * ElementsInOut is not cleared.
+	 * @param bTransformToWorld if true each geometric element will be transformed to World space, based on GetLocalToWorldTransform()
+	 * @param bIsForPreview if true, geometry is being collected for a preview of selection. Selector may return simplified geometry in this case
+	 * @param bUseGroupTopology if true, attempt to use group topology when populating ElementsInOut
+	 * @param Predicate if true for a given element type and element ID, this element will be collected, otherwise this element will be skipped.
+	 */
+	virtual void AccumulateElementsFromPredicate(
+		FGeometrySelectionElements& ElementsInOut,
+		bool bTransformToWorld,
+		bool bIsForPreview,
+		bool bUseGroupTopology,
+		TFunctionRef<bool(UE::Geometry::EGeometryElementType, UE::Geometry::FGeoSelectionID)> Predicate
+	) = 0;
+	
 	/**
 	 * Create and initialize a IGeometrySelectionTransformer for the provided Selection.
 	 * IGeometrySelector retains ownership of the active Transformer, it should not be
diff --git a/Engine/Plugins/Runtime/MeshModelingToolset/Source/ModelingComponents/Public/Selection/ToolSelectionUtil.h b/Engine/Plugins/Runtime/MeshModelingToolset/Source/ModelingComponents/Public/Selection/ToolSelectionUtil.h
--- a/Engine/Plugins/Runtime/MeshModelingToolset/Source/ModelingComponents/Public/Selection/ToolSelectionUtil.h
+++ b/Engine/Plugins/Runtime/MeshModelingToolset/Source/ModelingComponents/Public/Selection/ToolSelectionUtil.h
@@ -48,7 +48,7 @@ namespace ToolSelectionUtil
 			const UE::Geometry::FGroupTopology* Topology = nullptr,
 			const FTransform* ApplyTransform = nullptr,
 			bool bMapFacesToEdges = false);
-
+	
 	/**
 	 * Render the given Elements using FPrimitiveDrawInterface
 	 */
@@ -56,6 +56,16 @@ namespace ToolSelectionUtil
 		IToolsContextRenderAPI* RenderAPI,
 		const UE::Geometry::FGeometrySelectionElements& Elements,
 		bool bIsPreview = false);
+
+	// todo [nickolas.drake]: remove this function when no longer used by GeometrySelectionManager
+	void DebugRender(
+		IToolsContextRenderAPI* RenderAPI,
+		const UE::Geometry::FGeometrySelectionElements& Elements,
+		float LineThickness,
+		FLinearColor LineColor,
+		float PointSize,
+		FLinearColor PointColor,
+		float DepthBias = 0.f);
 }
 
 namespace UE::Geometry
