MeshElementSelection: Remove deprecation markup from UGeometrySelectionManager::ActivePreviewSelection

This was mistakenly included as I thought the SubmitTool had finished its work and I started messing around and it picked up my change

diff --git a/Engine/Plugins/Runtime/MeshModelingToolset/Source/ModelingComponents/Public/Selection/GeometrySelectionManager.h b/Engine/Plugins/Runtime/MeshModelingToolset/Source/ModelingComponents/Public/Selection/GeometrySelectionManager.h
--- a/Engine/Plugins/Runtime/MeshModelingToolset/Source/ModelingComponents/Public/Selection/GeometrySelectionManager.h
+++ b/Engine/Plugins/Runtime/MeshModelingToolset/Source/ModelingComponents/Public/Selection/GeometrySelectionManager.h
@@ -463,7 +463,6 @@ protected:
 	void UpdateSelectionRenderCacheOnTargetChange();
 	void RebuildSelectionRenderCaches();
 
-	UE_DEPRECATED(5.6, "Direct acces to ActivePreviewSelection is deprecated!")
     FGeometrySelection ActivePreviewSelection;										// Selection representing the active preview
 	FGeometrySelectionElements CachedPreviewRenderElements;							// Cached 3D geometry for active preview elements
     void ClearActivePreview();
