MeshModelingTools: Add extrude edge activity to PolyEd and mesh selection. Kept implementating classes in private folder for this release.

diff --git a/Engine/Plugins/Editor/ModelingToolsEditorMode/Source/ModelingToolsEditorMode/Private/ModelingToolsEditorMode.cpp b/Engine/Plugins/Editor/ModelingToolsEditorMode/Source/ModelingToolsEditorMode/Private/ModelingToolsEditorMode.cpp
--- a/Engine/Plugins/Editor/ModelingToolsEditorMode/Source/ModelingToolsEditorMode/Private/ModelingToolsEditorMode.cpp
+++ b/Engine/Plugins/Editor/ModelingToolsEditorMode/Source/ModelingToolsEditorMode/Private/ModelingToolsEditorMode.cpp
@@ -855,6 +855,7 @@ void UModelingToolsEditorMode::Enter()
 	RegisterPolyModelActionTool(EEditMeshPolygonsToolActions::Outset, ToolManagerCommands.BeginPolyModelTool_Outset, TEXT("PolyEdit_Outset"), true);
 	RegisterPolyModelActionTool(EEditMeshPolygonsToolActions::CutFaces, ToolManagerCommands.BeginPolyModelTool_CutFaces, TEXT("PolyEdit_CutFaces"), true);
 	RegisterPolyModelActionTool(EEditMeshPolygonsToolActions::InsertEdgeLoop, ToolManagerCommands.BeginPolyModelTool_InsertEdgeLoop, TEXT("PolyEdit_InsertEdgeLoop"), false);
+	RegisterPolyModelActionTool(EEditMeshPolygonsToolActions::ExtrudeEdges, ToolManagerCommands.BeginPolyModelTool_ExtrudeEdges, TEXT("PolyEdit_ExtrudeEdges"), true);
 	RegisterPolyModelActionTool(EEditMeshPolygonsToolActions::PushPull, ToolManagerCommands.BeginPolyModelTool_PushPull, TEXT("PolyEdit_PushPull"), true);
 	RegisterPolyModelActionTool(EEditMeshPolygonsToolActions::BevelAuto, ToolManagerCommands.BeginPolyModelTool_Bevel, TEXT("PolyEdit_Bevel"), true);
 	
diff --git a/Engine/Plugins/Editor/ModelingToolsEditorMode/Source/ModelingToolsEditorMode/Private/ModelingToolsEditorModeToolkit.cpp b/Engine/Plugins/Editor/ModelingToolsEditorMode/Source/ModelingToolsEditorMode/Private/ModelingToolsEditorModeToolkit.cpp
--- a/Engine/Plugins/Editor/ModelingToolsEditorMode/Source/ModelingToolsEditorMode/Private/ModelingToolsEditorModeToolkit.cpp
+++ b/Engine/Plugins/Editor/ModelingToolsEditorMode/Source/ModelingToolsEditorMode/Private/ModelingToolsEditorModeToolkit.cpp
@@ -275,6 +275,7 @@ void FModelingToolsEditorModeToolkit::RegisterPalettes()
 			Commands.BeginSelectionAction_Delete,
 			Commands.BeginSelectionAction_Extrude,
 			Commands.BeginSelectionAction_Offset,
+			Commands.BeginPolyModelTool_ExtrudeEdges,
 			Commands.BeginPolyModelTool_PushPull,
 
 			Commands.BeginPolyModelTool_Inset,
@@ -1832,6 +1833,7 @@ void FModelingToolsEditorModeToolkit::BuildToolPalette(FName PaletteIndex, class
 		ToolbarBuilder.AddToolBarButton(Commands.BeginSelectionAction_Extrude);
 		ToolbarBuilder.AddToolBarButton(Commands.BeginSelectionAction_Offset);
 
+		ToolbarBuilder.AddToolBarButton(Commands.BeginPolyModelTool_ExtrudeEdges);
 		ToolbarBuilder.AddToolBarButton(Commands.BeginPolyModelTool_PushPull);
 		ToolbarBuilder.AddToolBarButton(Commands.BeginPolyModelTool_Inset);
 		ToolbarBuilder.AddToolBarButton(Commands.BeginPolyModelTool_Outset);
diff --git a/Engine/Plugins/Editor/ModelingToolsEditorMode/Source/ModelingToolsEditorMode/Private/ModelingToolsManagerActions.cpp b/Engine/Plugins/Editor/ModelingToolsEditorMode/Source/ModelingToolsEditorMode/Private/ModelingToolsManagerActions.cpp
--- a/Engine/Plugins/Editor/ModelingToolsEditorMode/Source/ModelingToolsEditorMode/Private/ModelingToolsManagerActions.cpp
+++ b/Engine/Plugins/Editor/ModelingToolsEditorMode/Source/ModelingToolsEditorMode/Private/ModelingToolsManagerActions.cpp
@@ -336,6 +336,14 @@ void FModelingToolsManagerCommands::RegisterCommands()
 	REGISTER_MODELING_TOOL_COMMAND(BeginPolyModelTool_Inset, "Inset", "Inset the current Selected Faces");
 	REGISTER_MODELING_TOOL_COMMAND(BeginPolyModelTool_Outset, "Outset", "Outset the current Selected Faces");
 	REGISTER_MODELING_TOOL_COMMAND(BeginPolyModelTool_CutFaces, "Cut", "Cut the current Selected Faces");
+	if (bUseShortNames)
+	{
+		REGISTER_MODELING_TOOL_COMMAND(BeginPolyModelTool_ExtrudeEdges, "ExtEdg", "Extrude selected boundary edges.");
+	}
+	else
+	{
+		REGISTER_MODELING_TOOL_COMMAND(BeginPolyModelTool_ExtrudeEdges, "Extrude Edges", "Extrude selected boundary edges.");
+	}
 	REGISTER_MODELING_TOOL_COMMAND(BeginPolyModelTool_PushPull, "PushPull", "Push/Pull the current Selected Faces");
 	REGISTER_MODELING_TOOL_COMMAND(BeginPolyModelTool_Bevel, "Bevel", "Bevel the current Mesh Selection (Edges or Faces)");
 	if ( bUseShortNames )
diff --git a/Engine/Plugins/Editor/ModelingToolsEditorMode/Source/ModelingToolsEditorMode/Public/ModelingToolsManagerActions.h b/Engine/Plugins/Editor/ModelingToolsEditorMode/Source/ModelingToolsEditorMode/Public/ModelingToolsManagerActions.h
--- a/Engine/Plugins/Editor/ModelingToolsEditorMode/Source/ModelingToolsEditorMode/Public/ModelingToolsManagerActions.h
+++ b/Engine/Plugins/Editor/ModelingToolsEditorMode/Source/ModelingToolsEditorMode/Public/ModelingToolsManagerActions.h
@@ -150,6 +150,7 @@ public:
 	TSharedPtr<FUICommandInfo> BeginPolyModelTool_Outset;
 	TSharedPtr<FUICommandInfo> BeginPolyModelTool_CutFaces;
 	TSharedPtr<FUICommandInfo> BeginPolyModelTool_InsertEdgeLoop;
+	TSharedPtr<FUICommandInfo> BeginPolyModelTool_ExtrudeEdges;
 	TSharedPtr<FUICommandInfo> BeginPolyModelTool_PushPull;
 	TSharedPtr<FUICommandInfo> BeginPolyModelTool_Bevel;
 	
diff --git a/Engine/Plugins/Runtime/MeshModelingToolset/Source/MeshModelingTools/Private/EditMeshPolygonsTool.cpp b/Engine/Plugins/Runtime/MeshModelingToolset/Source/MeshModelingTools/Private/EditMeshPolygonsTool.cpp
--- a/Engine/Plugins/Runtime/MeshModelingToolset/Source/MeshModelingTools/Private/EditMeshPolygonsTool.cpp
+++ b/Engine/Plugins/Runtime/MeshModelingToolset/Source/MeshModelingTools/Private/EditMeshPolygonsTool.cpp
@@ -41,6 +41,7 @@
 #include "TargetInterfaces/PrimitiveComponentBackedTarget.h"
 #include "ToolActivities/PolyEditActivityContext.h"
 #include "ToolActivities/PolyEditExtrudeActivity.h"
+#include "ToolActivities/PolyEditExtrudeEdgeActivity.h"
 #include "ToolActivities/PolyEditInsertEdgeActivity.h"
 #include "ToolActivities/PolyEditInsertEdgeLoopActivity.h"
 #include "ToolActivities/PolyEditInsetOutsetActivity.h"
@@ -169,6 +170,10 @@ bool UEditMeshPolygonsActionModeToolBuilder::CanBuildTool(const FToolBuilderStat
 			{
 				return (TopologyType == EGeometryTopologyType::Polygroup && ElementType != EGeometryElementType::Vertex && bIsEmpty == false);
 			}
+			else if (StartupAction == EEditMeshPolygonsToolActions::ExtrudeEdges)
+			{
+				return (ElementType == EGeometryElementType::Edge && !bIsEmpty);
+			}
 		}
 	}
 	return false;
@@ -178,6 +183,27 @@ bool UEditMeshPolygonsActionModeToolBuilder::CanBuildTool(const FToolBuilderStat
 void UEditMeshPolygonsActionModeToolBuilder::InitializeNewTool(USingleTargetWithSelectionTool* Tool, const FToolBuilderState& SceneState) const
 {
 	UEditMeshPolygonsToolBuilder::InitializeNewTool(Tool, SceneState);
+
+	// Need to enable triangle mode on the tool if our selection was a triangle (not group) selection.
+	// This normally gets done in the base class if bTriangleMode is true, but we can't change that in a
+	// const method.
+	if (UGeometrySelectionManager* SelectionManager = SceneState.ToolManager->GetContextObjectStore()->FindContext<UGeometrySelectionManager>())
+	{
+		EGeometryTopologyType TopologyType = EGeometryTopologyType::Triangle;
+		EGeometryElementType ElementType = EGeometryElementType::Face;
+		int NumTargets;
+		bool bIsEmpty = false;
+		SelectionManager->GetActiveSelectionInfo(TopologyType, ElementType, NumTargets, bIsEmpty);
+
+		if (TopologyType == EGeometryTopologyType::Triangle)
+		{
+			if (UEditMeshPolygonsTool* EditPolygonsTool = Cast<UEditMeshPolygonsTool>(Tool))
+			{
+				EditPolygonsTool->EnableTriangleMode();
+			}
+		}
+	}
+
 	UEditMeshPolygonsTool* EditPolygonsTool = CastChecked<UEditMeshPolygonsTool>(Tool);
 
 	EEditMeshPolygonsToolActions UseAction = StartupAction;
@@ -610,6 +636,11 @@ void UEditMeshPolygonsTool::Setup()
 	ActivityLabels.Add(BevelEdgeActivity, LOCTEXT("BevelActivityLabel", "Bevel"));
 	ActivityIconNames.Add(BevelEdgeActivity, "PolyEd.Bevel");
 
+	ExtrudeEdgeActivity = NewObject<UPolyEditExtrudeEdgeActivity>();
+	ExtrudeEdgeActivity->Setup(this);
+	ActivityLabels.Add(ExtrudeEdgeActivity, LOCTEXT("EdgeExtrudeActivityLabel", "Extrude Edges"));
+	ActivityIconNames.Add(ExtrudeEdgeActivity, "PolyEd.ExtrudeEdge");
+
 	// Now that we've initialized the activities, add in the selection settings and 
 	// CommonProps so that they are at the bottom.
 	AddToolPropertySource(SelectionMechanic->Properties);
@@ -734,6 +765,7 @@ void UEditMeshPolygonsTool::OnShutdown(EToolShutdownType ShutdownType)
 	InsertEdgeActivity->Shutdown(ShutdownType);
 	InsertEdgeLoopActivity->Shutdown(ShutdownType);
 	BevelEdgeActivity->Shutdown(ShutdownType);
+	ExtrudeEdgeActivity->Shutdown(ShutdownType);
 
 	GetToolManager()->GetPairedGizmoManager()->DestroyAllGizmosByOwner(this);
 
@@ -823,6 +855,7 @@ void UEditMeshPolygonsTool::OnShutdown(EToolShutdownType ShutdownType)
 	InsertEdgeActivity = nullptr;
 	InsertEdgeLoopActivity = nullptr;
 	BevelEdgeActivity = nullptr;
+	ExtrudeEdgeActivity = nullptr;
 
 	SelectionMechanic = nullptr;
 	DragAlignmentMechanic = nullptr;
@@ -1206,6 +1239,16 @@ void UEditMeshPolygonsTool::OnTick(float DeltaTime)
 		//Interactive operations:
 		case EEditMeshPolygonsToolActions::Extrude:
 		{
+			if (SelectionMechanic->GetActiveSelection().SelectedGroupIDs.IsEmpty()
+				&& !SelectionMechanic->GetActiveSelection().SelectedEdgeIDs.IsEmpty())
+			{
+				// This particular button happens to be under "face edits", but it's very tempting to click it anyway
+				// when you have an edge selection and expect it to extrude edges. We'll allow it to avoid frustrating
+				// the user. Not relevant for mesh element selection, where we don't use PolyEd and instead extrude
+				// the adjacent faces when edges are selected.
+				StartActivity(ExtrudeEdgeActivity);
+				break;
+			}
 			ExtrudeActivity->ExtrudeMode = FExtrudeOp::EExtrudeMode::MoveAndStitch;
 			ExtrudeActivity->PropertySetToUse = UPolyEditExtrudeActivity::EPropertySetToUse::Extrude;
 			StartActivity(ExtrudeActivity);
@@ -1265,6 +1308,26 @@ void UEditMeshPolygonsTool::OnTick(float DeltaTime)
 			StartActivity(InsertEdgeLoopActivity);
 			break;
 		}
+		case EEditMeshPolygonsToolActions::ExtrudeEdges:
+		{
+			// Hack: We currently don't support extra corners in mesh element selection, and
+			// the switch to using them can cause us to lose some of our selected edges. For
+			// now we just rebuild topology without the corners in this scenario, but we should
+			// fix instead just carry over the selection properly (and carry it back).
+			if (bTerminateOnPendingActionComplete && HasGeometrySelection())
+			{
+				const FGeometrySelection& CurSelection = GetGeometrySelection();
+				if (CurSelection.TopologyType == EGeometryTopologyType::Polygroup && bTriangleMode == false)
+				{
+					TSet<int32> EmptySet;
+					RebuildTopologyWithGivenExtraCorners(EmptySet);
+					SelectionMechanic->SetSelection_AsGroupTopology(CurSelection); // Have to reinitialize selection
+				}
+			}
+
+			StartActivity(ExtrudeEdgeActivity);
+			break;
+		}
 
 		case EEditMeshPolygonsToolActions::BevelFaces:
 		case EEditMeshPolygonsToolActions::BevelEdges:
diff --git a/Engine/Plugins/Runtime/MeshModelingToolset/Source/MeshModelingTools/Private/Operations/ExtrudeBoundaryEdges.cpp b/Engine/Plugins/Runtime/MeshModelingToolset/Source/MeshModelingTools/Private/Operations/ExtrudeBoundaryEdges.cpp
new file mode 100644
--- /dev/null
+++ b/Engine/Plugins/Runtime/MeshModelingToolset/Source/MeshModelingTools/Private/Operations/ExtrudeBoundaryEdges.cpp
@@ -0,0 +1,669 @@
+// Copyright Epic Games, Inc. All Rights Reserved.
+
+#include "Operations/ExtrudeBoundaryEdges.h"
+
+#include "BoxTypes.h"
+#include "DynamicMesh/DynamicMesh3.h"
+#include "DynamicMeshEditor.h"
+#include "Math/NumericLimits.h"
+#include "Util/ProgressCancel.h"
+#include "Parameterization/DynamicMeshUVEditor.h"
+#include "VectorUtil.h"
+
+namespace ExtrudeBoundaryEdgesLocals
+{
+	using namespace UE::Geometry;
+
+	// Gives some vector orthogonal to NormalizedVectorIn
+	FVector3d GetOrthogonalVector(const FVector3d& NormalizedVectorIn)
+	{
+		FVector3d Result = NormalizedVectorIn.Cross(FVector3d::UnitY());
+		if (Result.IsZero())
+		{
+			// Must have been colinear with Y, so Z will work
+			return FVector3d::UnitZ();
+		}
+		return Result;
+	}
+
+	// Given a non degenerate normal and edge, find direction in which we would want to extrude the edge to stay
+	//  in the plane of the face and move edge orthogonally outward.
+	FVector3d GetOutwardVector(const FVector3d& NormalVector, const FVector3d& EdgeVector)
+	{
+		// If we're looking at the triangle with edge vector pointing up and triangle to the left (and
+		//  Z towards us), then X should point right. In unreal's left handed coordinate system that
+		//  actually means this:
+		return NormalVector.Cross(EdgeVector);
+	}
+
+	FExtrudeBoundaryEdges::FExtrudeFrame GetExtrudeFrameAtVertex(const FVector3d SourcePosition,
+		const FVector3d& IncomingTriNormal, const FVector3d& IncomingEdgeVector,
+		const FVector3d& OutgoingTriNormal, const FVector3d& OutgoingEdgeVector,
+		double ScalingLimit)
+	{
+		// TODO: We try to do some handling of degenerate normals below, but it is generally not good
+		//  enough, and we are very likely to get nonsense extrude frames when we have degenerate normals.
+		//  One option is to propagate normal information from neighbors, or propagate frame information
+		//  from adjacent edges where a good one could be picked. These options are a fair amount of
+		//  work for something that may not have much benefit, because the user probably still won't have
+		//  much control and at best we would be hiding the effects of flawed geometry... Still, might
+		//  be something to come back to.
+
+		if (OutgoingEdgeVector.IsZero() && IncomingEdgeVector.IsZero())
+		{
+			// Not much we can do here.
+			return FExtrudeBoundaryEdges::FExtrudeFrame(FFrame3d(SourcePosition));
+		}
+
+
+		auto CreateFrameFromEdgeAndNormal = [&SourcePosition](const FVector3d& NormalVector, const FVector3d& EdgeVector)
+		{
+			FVector3d FrameZ = NormalVector;
+			FVector3d FrameX = GetOutwardVector(NormalVector, EdgeVector);
+			// With the way we calculate the outward vector, this will actually just always end up -EdgeVector,
+			//  but we'll just do the proper thing.
+			FVector3d FrameY = FrameZ.Cross(FrameX);
+
+			return FExtrudeBoundaryEdges::FExtrudeFrame(
+				FFrame3d(SourcePosition, FrameX, FrameY, FrameZ));
+		};
+
+		// Handle cases where we only have data on one neighbor.
+		if (IncomingEdgeVector.IsZero())
+		{
+			ensure(!OutgoingEdgeVector.IsZero());
+			return CreateFrameFromEdgeAndNormal(
+				OutgoingTriNormal.IsZero() ? GetOrthogonalVector(OutgoingEdgeVector) : OutgoingTriNormal, 
+				OutgoingEdgeVector);
+		}
+		if (OutgoingEdgeVector.IsZero())
+		{
+			ensure(!IncomingEdgeVector.IsZero());
+			return CreateFrameFromEdgeAndNormal(
+				IncomingTriNormal.IsZero() ? GetOrthogonalVector(IncomingTriNormal) : IncomingTriNormal, 
+				IncomingEdgeVector);
+		}
+
+		// If we got to here, we had neighbors on both sides. We'd like to use the normals in setting up the extrude
+		//  frame, but one or both of them could be zeroes if the triangles were degenerate. Not totally clear what
+		//  the best solution is, but for now we either find an orthogonal direction if they are both degenerate, or
+		//  else we pick a plane for the degenerate edge such that it lies in the same plane as the desired X for the
+		//  other face.
+		FVector3d IncomingNormalToUse = IncomingTriNormal;
+		FVector3d OutgoingNormalToUse = OutgoingTriNormal;
+		if (IncomingNormalToUse.IsZero())
+		{
+			if (OutgoingTriNormal.IsZero())
+			{
+				// Both faces degrenerate. X direction is orthogonal to both. Again, remember left handedness.
+				FVector3d FrameX = IncomingEdgeVector.Cross(OutgoingEdgeVector);
+				FrameX.Normalize();
+				if (FrameX.IsZero())
+				{
+					FrameX = GetOrthogonalVector(IncomingEdgeVector);
+				}
+
+				IncomingNormalToUse = IncomingEdgeVector.Cross(FrameX);
+				IncomingNormalToUse.Normalize();
+				OutgoingNormalToUse = OutgoingEdgeVector.Cross(FrameX);
+				OutgoingNormalToUse.Normalize();
+			}
+			else
+			{
+				FVector3d FrameX = GetOutwardVector(OutgoingTriNormal, OutgoingEdgeVector);
+
+				IncomingNormalToUse = IncomingEdgeVector.Cross(FrameX);
+				if (!IncomingNormalToUse.Normalize())
+				{
+					// If we were colinear with the X we picked, let's say that we're coplanar with the
+					// other face.
+					IncomingNormalToUse = OutgoingNormalToUse;
+				}
+			}
+		}
+		else if (OutgoingTriNormal.IsZero())
+		{
+			FVector3d FrameX = GetOutwardVector(IncomingTriNormal, IncomingEdgeVector);
+
+			OutgoingNormalToUse = OutgoingEdgeVector.Cross(FrameX);
+			if (!OutgoingNormalToUse.Normalize())
+			{
+				OutgoingNormalToUse = IncomingNormalToUse;
+			}
+		}
+
+		// At this point we have non degenerate edge vectors and normals.
+		// There are a few ways we can pick our frame. One option we've tried in the past is picking
+		//  the X to be in the plane of both neighboring triangles, which works for following the "crease"
+		//  of some edges nicely, but falls apart frequently if both edges lie in the plane of at least
+		//  one of the triangles, because you frequently end up extruding along one of the edges. (E.g. imagine
+		//  a classic pentagonal prism "house" shape and remove one of the roof faces, then try extruding that
+		//  boundary. This method will nicely extend the adjacent faces in their planes. But if you delete only
+		//  half of the roof face and try to extrude the boundary, you'll get very weird results as some vertices
+		//  try to move along boundary edges instead, due to the flat portion of the roof constraining them.)
+		// A significantly more reliable option is to make X an average of the two outward directions. This
+		//  tends to be a lot more predictable and less likely to generate garbage geometry. This is what we do.
+
+		FVector3d OutwardVector1 = GetOutwardVector(IncomingNormalToUse, IncomingEdgeVector);
+		FVector3d OutwardVector2 = GetOutwardVector(OutgoingNormalToUse, OutgoingEdgeVector);
+		FVector3d FrameX = OutwardVector1 + OutwardVector2;
+		FrameX.Normalize();
+		if (FrameX.IsZero())
+		{
+			// It seems like getting this situation would imply some very weird twisting topology
+			//  where the next triangle manages to flip which side of the edge it is on, or Z fights
+			//  the original triangle while moving backwards. There isn't a perfect thing to do here,
+			//  so just use one of the edges and ignore the other.
+			return CreateFrameFromEdgeAndNormal(
+				IncomingTriNormal.IsZero() ? GetOrthogonalVector(IncomingTriNormal) : IncomingTriNormal,
+				IncomingEdgeVector);
+		}
+		FVector3d FrameZ = IncomingNormalToUse + OutgoingNormalToUse;
+		if (FrameZ.IsZero())
+		{
+			// This means that the normals flipped, usually by having the triangles fold back onto each
+			//  other. We can pick a Z pointing in the plane of the edge by using the edge vectors here.
+			FrameZ = IncomingEdgeVector - OutgoingEdgeVector;
+		}
+		// Not sure whether X and Z are necessarily orthogonal at this point with our approaches... Seems
+		//  like it might be the case, but let's be safe and not assume it.
+		FVector3d FrameY = FrameZ.Cross(FrameX);
+		FrameY.Normalize();
+		FrameZ = FrameX.Cross(FrameY);
+		// X and Y normalized and orthogonal, so Z shouldn't need normalization.
+
+		FExtrudeBoundaryEdges::FExtrudeFrame ExtrudeFrame(FFrame3d(
+			SourcePosition, FrameX, FrameY, FrameZ));
+
+		if (ScalingLimit > 1)
+		{
+			FVector3d VectorToCompareAgainst = OutwardVector1.IsZero() ? OutwardVector2 : OutwardVector1;
+			double cosTheta = ExtrudeFrame.Frame.X().Dot(VectorToCompareAgainst);
+			if (FMath::Abs(cosTheta) > KINDA_SMALL_NUMBER)
+			{
+				ExtrudeFrame.Scaling.X = FMath::Min(1 / FMath::Abs(cosTheta), ScalingLimit);
+			}
+		}
+
+		return ExtrudeFrame;
+	}
+
+	int32 FindSharpestAngleMatchEdge(const FDynamicMesh3& Mesh, int32 EidToMatch, int32 Vid, 
+		const FVector3d& VertexNormal, const TArray<int32> CandidateEids)
+	{
+		// The approach here is pretty much the same as in FMeshBoundaryLoops::FindLeftTurnEdge, except that
+		//  we're actually looking for the right turn edge (better if you're concerned with following a hole/boundary
+		//  instead of the smallest island), we only look at candidate eids, and we can look from outgoing edge
+		//  instead.
+		
+		FIndex2i InputEdgeVids = Mesh.GetOrientedBoundaryEdgeV(EidToMatch);
+		bool bInputIsIncomingEdge = InputEdgeVids.B == Vid;
+
+		if (!ensure(bInputIsIncomingEdge || InputEdgeVids.A == Vid))
+		{
+			return IndexConstants::InvalidID;
+		}
+
+		// We're going to be measuring the angle from incoming to outgoing edge, and we want it to be the most
+		//  clockwise. Normally that would mean most negative in the plane, but with left handed coordinates
+		//  it is most positive...
+		double BestAngle = TNumericLimits<double>::Lowest();
+		int32 BestEid = IndexConstants::InvalidID;
+
+		FVector3d InputEdgeVector = Mesh.GetVertex(InputEdgeVids.B) - Mesh.GetVertex(InputEdgeVids.A);
+		
+		for (int32 Eid : CandidateEids)
+		{
+			if (Eid == EidToMatch)
+			{
+				continue;
+			}
+			FIndex2i Vids = Mesh.GetOrientedBoundaryEdgeV(Eid);
+			bool bIsIncomingEdge = Vids.B == Vid;
+			if (!ensure(bIsIncomingEdge || Vids.A == Vid))
+			{
+				continue;
+			}
+			if (bInputIsIncomingEdge == bIsIncomingEdge)
+			{
+				continue;
+			}
+
+			FVector3d EdgeVector = Mesh.GetVertex(Vids.B) - Mesh.GetVertex(Vids.A);
+			double Angle = TNumericLimits<double>::Min();
+			if (bInputIsIncomingEdge)
+			{
+				Angle = VectorUtil::PlaneAngleSignedD(InputEdgeVector, EdgeVector, VertexNormal);
+			}
+			else
+			{
+				Angle = VectorUtil::PlaneAngleSignedD(EdgeVector, InputEdgeVector, VertexNormal);
+			}
+			
+			if (Angle > BestAngle)
+			{
+				BestEid = Eid;
+				BestAngle = Angle;
+			}
+		}
+
+		return BestEid;
+	}
+}
+
+namespace UE::Geometry
+{
+FExtrudeBoundaryEdges::FExtrudeBoundaryEdges(FDynamicMesh3* MeshIn)
+	: Mesh(MeshIn)
+{
+}
+
+bool FExtrudeBoundaryEdges::Apply(FProgressCancel* Progress)
+{
+	using namespace ExtrudeBoundaryEdgesLocals;
+
+	// Complications in exturding edges:
+	// 1. We can't just always split bowties in our input, because we might be using the operation to fix
+	//  a bowtie. For instance imagine two circle holes tangent at a bowtie. We might be selecting just 
+	//  one of the holes and extruding inward. If that bowtie were broken, we would get an undesirable
+	//  split in the result, instead of a complete quad ring.
+	// 2. We DO need to do something about bowties if we're extruding more than two incident boundary edges 
+	//  (imagine extruding the boundaries of both holes in the example above). Keeping the vert alone in
+	//  the destination would attempt to create a non-manifold edge on it.
+	// 3. Degenerate edges are a problem when determining per-vertex extrude frames, because they don't
+	//   have proper edge/normal vectors, and as such could end up going to different places than their
+	//   neighbors.
+
+	// We take the following approach for 1 and 2:
+	// Every extruded edge gets a quad with verts a, b, c, d, where a -> b is one of the original edges,
+	//	oriented in triangle direction, and c -> d is the "extruded edge". When creating c or d, we'll 
+	//  enforce the constraint that each copy can be used as "c" in no more than one quad, and "d" in no
+	//  more than one quad. In other words, we'll create a new copy if the source vid is the destination
+	//  or the source for more than one edge in the original edges to extrude.
+	// This means that we will actually end up splitting verts in the destination even if they were not
+	//  a bowtie vert if the original vert had flipped neighbor triangles. That seems like a reasonable
+	//  thing to do in a situation where the user really shouldn't have any expectation.
+	//
+	// For 3, we currently decide to not do anything, instead allowing for whatever garbage results
+	//  might be created when a user tries to extrude degenerate edges (though the results won't be garbage
+	//  when extruding in a single direction). In an ideal world, we would implement some grouping of
+	//  vertices that are connected by degenerate edges, propage the vector information across them, and
+	//  make sure that vertices in the group get the same extrude frame. While possible, that seems not
+	//  worth the pain and potential pitfalls, especially if you consider pathological cases of some web
+	//  of degenerate boundary edges rather than a single chain, where there are multiple extrude frames
+	//  that each might want to use. The saner approach is to expect the user to collapse degenerate edges 
+	//  in the input first.
+
+	if (!ensure(Mesh && OffsetPositionFunc))
+	{
+		return false;
+	}
+
+	if (InputEids.IsEmpty())
+	{
+		return true;
+	}
+
+	if (Progress && Progress->Cancelled()) { return false; }
+
+	TArray<FNewVertSourceData> NewVertData;
+	TArray<int32> NewVertVids; // eventually 1:1 with NewVertData
+	TMap<int32, FIndex2i> EidToIndicesIntoNewVerts;
+
+	// Pair up the edges across vertices so that we know when to split verts
+	bool bSuccess = FExtrudeBoundaryEdges::GetInputEdgePairings(
+		*Mesh, InputEids, bAssignAnyBoundaryNeighborToUnmatched,
+		NewVertData, EidToIndicesIntoNewVerts);
+	if (!ensure(bSuccess)) { return false; }
+	
+	if (NewVertData.IsEmpty())
+	{
+		return true;
+	}
+	if (Progress && Progress->Cancelled()) { return false; }
+	
+	// Create all the new verts
+	for (const FNewVertSourceData& VertSourceData : NewVertData)
+	{
+		int32 CreatedVid = Mesh->AppendVertex(*Mesh, VertSourceData.SourceVid);
+		NewVertVids.Add(CreatedVid);
+	}
+	if (Progress && Progress->Cancelled()) { return false; }
+	
+	// Calculate extrude frames, if used
+	TArray<FExtrudeFrame> NewVertExtrudeFrames; // if used, 1:1 with NewVertData
+	if (bUsePerVertexExtrudeFrames)
+	{
+		for (const FNewVertSourceData& VertSourceData : NewVertData)
+		{
+			FExtrudeFrame& ExtrudeFrame = NewVertExtrudeFrames.Emplace_GetRef();
+			GetExtrudeFrame(*Mesh, VertSourceData.SourceVid, 
+				VertSourceData.SourceEidPair.A, VertSourceData.SourceEidPair.B, 
+				ExtrudeFrame, ScalingAdjustmentLimit);
+		}
+		if (Progress && Progress->Cancelled()) { return false; }
+	}
+
+	// Perform the offset for all of the created Vids
+	for (int32 i = 0; i < NewVertVids.Num(); ++i)
+	{
+		int32 NewVid = NewVertVids[i];
+		int32 SourceVid = NewVertData[i].SourceVid;
+
+		Mesh->SetVertex(NewVid, OffsetPositionFunc(
+			Mesh->GetVertex(NewVid),
+			bUsePerVertexExtrudeFrames ? NewVertExtrudeFrames[i] : FExtrudeFrame(),
+			SourceVid));
+	}
+	if (Progress && Progress->Cancelled()) { return false; }
+
+	// TODO: We could pretty easily add a version of stitching code to FDynamicMeshEditor that stitches
+	//  arbitrary vert sequences rather than loops, and then maybe we could follow FOffsetMeshRegion and
+	//  FInsetMeshRegion's examples to process the resulting strips. However it would require code that 
+	//  assembles our input eids into sequences first. That's not too hard once we have the pairings above,
+	//  but it isn't clear whether the extra effort will or will not be worthwhile. For now we do the easy
+	//  thing and just prep all the quads and stitch them, then just box project UV's onto the whole thing.
+	// The UV generation would be the main thing we'd want to improve, probably.
+	
+	int32 GroupID = 0;
+	if (GroupsToSetPerEid.IsSet() && !ensure(GroupsToSetPerEid->Num() == InputEids.Num()))
+	{
+		GroupsToSetPerEid.Reset();
+	}
+	if (Mesh->HasTriangleGroups() && !GroupsToSetPerEid.IsSet())
+	{
+		GroupID = Mesh->AllocateTriangleGroup();
+	}
+	
+	// Do the stitching
+	for (int32 i = 0; i < InputEids.Num(); ++i)
+	{
+		int32 Eid = InputEids[i];
+		GroupID = GroupsToSetPerEid.IsSet() ? (*GroupsToSetPerEid)[i] : GroupID;
+
+		if (!Mesh->IsBoundaryEdge(Eid))
+		{
+			continue;
+		}
+
+		const FIndex2i* NewVertIndices = EidToIndicesIntoNewVerts.Find(Eid);
+		if (!ensure(NewVertIndices 
+			&& NewVertIndices->A >= 0 && NewVertIndices->A < NewVertVids.Num()
+			&& NewVertIndices->B >= 0 && NewVertIndices->B < NewVertVids.Num()))
+		{
+			continue;
+		}
+
+		FIndex2i Vids = Mesh->GetOrientedBoundaryEdgeV(Eid);
+		int32 VidA = Vids.A;
+		int32 VidB = Vids.B;
+		int32 VidC = NewVertVids[NewVertIndices->A];
+		int32 VidD = NewVertVids[NewVertIndices->B];
+
+		FIndex3i Tri1(VidB, VidA, VidD);
+		NewTids.Add(Mesh->AppendTriangle(Tri1, GroupID));
+		FIndex3i Tri2(VidA, VidC, VidD);
+		NewTids.Add(Mesh->AppendTriangle(Tri2, GroupID));
+
+		NewExtrudedEids.Add(Mesh->FindEdgeFromTri(VidC, VidD, NewTids.Last()));
+	}
+
+	// Give UVs to the new triangles
+	if (Mesh->HasAttributes())
+	{
+		if (Progress && Progress->Cancelled()) { return false; }
+
+		FAxisAlignedBox3d Box;
+		for (int32 Tid : NewTids)
+		{
+			FVector3d Vert1, Vert2, Vert3;
+			Mesh->GetTriVertices(Tid, Vert1, Vert2, Vert3);
+			Box.Contain(Vert1);
+			Box.Contain(Vert2);
+			Box.Contain(Vert3);
+		}
+
+		FDynamicMeshUVEditor UVEd(Mesh, Mesh->Attributes()->PrimaryUV());
+		FFrame3d BoxFrame(Box.Center());
+		FVector3d BoxDimensions = Box.Extents() * 2;
+		UVEd.SetTriangleUVsFromBoxProjection(NewTids, [](FVector3d Position) {return Position; }, BoxFrame, BoxDimensions);
+	}
+	
+	return true;
+}
+
+// See comment in Apply()
+/**
+ * @param NewVertDataOut One entry for each vertex that should be created, with information on its relevant neighbor edges
+ *   and the source vertex.
+ * @param EidToIndicesIntoNewVertsOut Mapping from Eid to two indices into NewVertDataOut that indicate the two new vertices
+ *   that would be stitched into the quad containing this edge.
+ */
+bool FExtrudeBoundaryEdges::GetInputEdgePairings(
+	const FDynamicMesh3& Mesh, TArray<int32>& InputEids, bool bAssignAnyBoundaryNeighborToUnmatched,
+	TArray<FNewVertSourceData>& NewVertDataOut, TMap<int32, FIndex2i>& EidToIndicesIntoNewVertsOut)
+{
+	using namespace ExtrudeBoundaryEdgesLocals;
+
+	NewVertDataOut.Reset();
+	EidToIndicesIntoNewVertsOut.Reset();
+
+	// Helper that updates both NewVertDataOut and EidToIndicesIntoNewVertsOut
+	auto AddEdgePairing = [&NewVertDataOut, &EidToIndicesIntoNewVertsOut](int32 SourceVid, int32 IncomingEid, int32 OutgoingEid)
+	{
+		int32 NewVertIndex = NewVertDataOut.Emplace();
+		NewVertDataOut[NewVertIndex].SourceVid = SourceVid;
+		NewVertDataOut[NewVertIndex].SourceEidPair = FIndex2i(IncomingEid, OutgoingEid);
+		
+		auto UpdateMapping = [NewVertIndex, &EidToIndicesIntoNewVertsOut](int32 Eid, bool bUpdatingOutgoingEdge)
+		{
+			if (Eid == IndexConstants::InvalidID)
+			{
+				return;
+			}
+
+			FIndex2i* ExistingEntry = EidToIndicesIntoNewVertsOut.Find(Eid);
+			if (ExistingEntry)
+			{
+				// If we're updating the outgoing edge, then we are its first vertex
+				(*ExistingEntry)[bUpdatingOutgoingEdge ? 0 : 1] = NewVertIndex;
+			}
+			else
+			{
+				FIndex2i Entry;
+				Entry[bUpdatingOutgoingEdge ? 0 : 1] = NewVertIndex;
+				EidToIndicesIntoNewVertsOut.Add(Eid, Entry);
+			}
+		};
+
+		UpdateMapping(IncomingEid, false);
+		UpdateMapping(OutgoingEid, true);
+	};
+
+	// We will be building up a list of associations of verts to two neighbor edges. If we run into
+	//  a situation where there are more than two candidates, we will put the entry into AmbiguousVids
+	//  and (later) remove it from this map.
+	TMap<int32, FIndex2i> UnambiguousVidToEdgePairings;
+	
+	TArray<int32> AmbiguousVids;
+
+	for (int32 Eid : InputEids)
+	{
+		if (!Mesh.IsBoundaryEdge(Eid))
+		{
+			continue;
+		}
+		
+		FIndex2i Vids = Mesh.GetOrientedBoundaryEdgeV(Eid);
+		for (int SubIdx = 0; SubIdx < 2; ++SubIdx)
+		{
+			int32 Vid = Vids[SubIdx];
+			FIndex2i* EdgePairing = UnambiguousVidToEdgePairings.Find(Vid);
+			if (!EdgePairing)
+			{
+				EdgePairing = &UnambiguousVidToEdgePairings.Add(Vid, FIndex2i::Invalid());
+			}
+			if ((*EdgePairing)[1-SubIdx] != IndexConstants::InvalidID)
+			{
+				// Found multiple incoming or outgoing edges. We'll handle this later.
+				AmbiguousVids.Add(Vid);
+			}
+			else
+			{
+				(*EdgePairing)[1-SubIdx] = Eid;
+			}
+		}
+	}
+
+	// Crude (not angle-weighed) vertex normal calculation used for projection below, same as 
+	//  FMeshBoundaryLoops::GetVertexNormal.
+	auto CalculateVertexNormal = [&Mesh](int32 Vid)
+	{
+		FVector3d VertexNormal = FVector3d::Zero();
+		for (int32 Tid : Mesh.VtxTrianglesItr(Vid))
+		{
+			VertexNormal += Mesh.GetTriNormal(Tid);
+		}
+		VertexNormal.Normalize();
+		return VertexNormal;
+	};
+
+	TSet<int32> SelectedEids(InputEids);
+
+	// Sort out the ambiguous verts first, and remove those entries from UnambiguousVidToEdgePairings
+	for (int32 AmbiguousVid : AmbiguousVids)
+	{
+		UnambiguousVidToEdgePairings.Remove(AmbiguousVid);
+
+		TArray<int32> IncomingEids;
+		TArray<int32> OutgoingEids;
+		for (int32 Eid : Mesh.VtxEdgesItr(AmbiguousVid))
+		{
+			if (!Mesh.IsBoundaryEdge(Eid) || !SelectedEids.Contains(Eid))
+			{
+				continue;
+			}
+
+			FIndex2i Vids = Mesh.GetOrientedBoundaryEdgeV(Eid);
+			if (Vids.A == AmbiguousVid)
+			{
+				OutgoingEids.Add(Eid);
+			}
+			else
+			{
+				ensure(Vids.B == AmbiguousVid);
+				IncomingEids.Add(Eid);
+			}
+		}
+
+		FVector3d VertexNormal = CalculateVertexNormal(AmbiguousVid);
+
+		for (int32 IncomingEid : IncomingEids)
+		{
+			int32 OutgoingEid = FindSharpestAngleMatchEdge(Mesh, IncomingEid, AmbiguousVid, VertexNormal, OutgoingEids);
+			if (OutgoingEid != IndexConstants::InvalidID)
+			{
+				// One candidate used up
+				OutgoingEids.Remove(OutgoingEid);
+			}
+			AddEdgePairing(AmbiguousVid, IncomingEid, OutgoingEid);
+		}
+		// If we still have some outgoing edges, we know they are unmatched now.
+		for (int32 OutgoingEid : OutgoingEids)
+		{
+			AddEdgePairing(AmbiguousVid, IndexConstants::InvalidID, OutgoingEid);
+		}
+	}// done processing ambiguous verts
+
+	for (const TPair<int32, FIndex2i>& KeyValue : UnambiguousVidToEdgePairings)
+	{
+		AddEdgePairing(KeyValue.Key, KeyValue.Value.A, KeyValue.Value.B);
+	}
+
+	if (bAssignAnyBoundaryNeighborToUnmatched)
+	{
+		for (FNewVertSourceData& VertData : NewVertDataOut)
+		{
+			if (!VertData.SourceEidPair.Contains(IndexConstants::InvalidID))
+			{
+				continue;
+			}
+
+			FVector3d VertexNormal = CalculateVertexNormal(VertData.SourceVid);
+			TArray<int32> NeighborhoodEids;
+			Mesh.GetAllVtxBoundaryEdges(VertData.SourceVid, NeighborhoodEids);
+
+			if (VertData.SourceEidPair.A == IndexConstants::InvalidID 
+				&& ensure(VertData.SourceEidPair.B != IndexConstants::InvalidID))
+			{
+				VertData.SourceEidPair.A = FindSharpestAngleMatchEdge(Mesh, 
+					VertData.SourceEidPair.B, VertData.SourceVid, VertexNormal, NeighborhoodEids);
+			}
+			else if (VertData.SourceEidPair.B == IndexConstants::InvalidID
+				&& ensure(VertData.SourceEidPair.A != IndexConstants::InvalidID))
+			{
+				VertData.SourceEidPair.B = FindSharpestAngleMatchEdge(Mesh,
+					VertData.SourceEidPair.A, VertData.SourceVid, VertexNormal, NeighborhoodEids);
+			}
+		}
+	}
+
+	return true;
+}
+
+bool FExtrudeBoundaryEdges::GetExtrudeFrame(const FDynamicMesh3& Mesh, int32 Vid,
+	int32 IncomingEid, int32 OutgoingEid, FExtrudeFrame& ExtrudeFrameOut,
+	double ScalingLimit)
+{
+	using namespace ExtrudeBoundaryEdgesLocals;
+
+	auto GetEdgeVectorAndNormal = [&Mesh, Vid](int32 Eid, bool bIsIncoming, FVector3d& EdgeVectorOut, FVector3d& NormalVectorOut) -> bool
+	{
+		if (Eid == IndexConstants::InvalidID)
+		{
+			EdgeVectorOut = FVector3d::Zero();
+			NormalVectorOut = FVector3d::Zero();
+			return true;
+		}
+
+		if (!ensure(Mesh.IsBoundaryEdge(Eid)))
+		{
+			return false;
+		}
+
+		FIndex2i Vids = Mesh.GetOrientedBoundaryEdgeV(Eid);
+		if (bIsIncoming && !ensure(Vids.B == Vid))
+		{
+			return false;
+		}
+		else if (!bIsIncoming && !ensure(Vids.A == Vid))
+		{
+			return false;
+		}
+
+		int32 Tid = Mesh.GetEdgeT(Eid).A;
+		NormalVectorOut = Mesh.GetTriNormal(Tid);
+
+		EdgeVectorOut = Mesh.GetVertex(Vids.B) - Mesh.GetVertex(Vids.A);
+		EdgeVectorOut.Normalize();
+
+		return true;
+	};
+
+	FVector3d EdgeVector1, EdgeVector2, Normal1, Normal2;
+
+	if (!GetEdgeVectorAndNormal(IncomingEid, true, EdgeVector1, Normal1)
+		|| !GetEdgeVectorAndNormal(OutgoingEid, false, EdgeVector2, Normal2))
+	{
+		return false;
+	}
+
+	ExtrudeFrameOut = GetExtrudeFrameAtVertex(
+		Mesh.GetVertex(Vid),
+		Normal1, EdgeVector1,
+		Normal2, EdgeVector2, ScalingLimit);
+
+	return true;
+}
+
+}//end UE::Geometry
\ No newline at end of file
diff --git a/Engine/Plugins/Runtime/MeshModelingToolset/Source/MeshModelingTools/Private/Operations/ExtrudeBoundaryEdges.h b/Engine/Plugins/Runtime/MeshModelingToolset/Source/MeshModelingTools/Private/Operations/ExtrudeBoundaryEdges.h
new file mode 100644
--- /dev/null
+++ b/Engine/Plugins/Runtime/MeshModelingToolset/Source/MeshModelingTools/Private/Operations/ExtrudeBoundaryEdges.h
@@ -0,0 +1,125 @@
+// Copyright Epic Games, Inc. All Rights Reserved.
+
+#pragma once
+
+#include "Containers/Array.h"
+#include "FrameTypes.h"
+#include "HAL/Platform.h" //int32
+#include "Misc/Optional.h"
+#include "VectorTypes.h"
+
+class FProgressCancel;
+
+namespace UE::Geometry
+{
+
+class FDynamicMesh3;
+
+class FExtrudeBoundaryEdges
+{
+public:
+	using FFrame3d = UE::Geometry::FFrame3d;
+	
+	// Represents a frame where the axes might not be unit scaled (but are still orthogonal).
+	// Allows vertices to adjust the extrusion distance along one of their extrusion frame
+	//  axes when trying to keep edges parallel.
+	// When using extrusion frames, vertices will be moved in the XZ plane, usually along X.
+	struct FExtrudeFrame
+	{
+		FExtrudeFrame() {}
+
+		FExtrudeFrame(const FFrame3d& FrameIn, const FVector3d& ScalingIn = FVector3d::One())
+			: Frame(FrameIn)
+			, Scaling(ScalingIn)
+		{}
+
+		FFrame3d Frame;
+		FVector3d Scaling = FVector3d::One();
+	};
+
+	// Data needed to create a new vert and its extrude frame
+	struct FNewVertSourceData
+	{
+		int32 SourceVid = IndexConstants::InvalidID;
+		// Neighboring edges, ordered by (incoming, outgoing)
+		FIndex2i SourceEidPair = FIndex2i::Invalid();
+	};
+
+	// Inputs:
+
+	/** The mesh that we are modifying */
+	FDynamicMesh3* Mesh;
+
+	/** The edges we're extruding */
+	TArray<int32> InputEids;
+
+	/** Whether to calculate local extrude frames and supply them to OffsetPositionFunc */
+	bool bUsePerVertexExtrudeFrames = true;
+
+	/**
+	 * Function queried for new verex positions. ExtrudeFrame origin is Position, unless it is not initialized
+	 *  due to bUsePerVertexExtrudeFrames being false.
+	 * The default given here assumes that bUsePerVertexExtrudeFrames is true, and extrudes along the X axis of the extrude frame.
+	 */
+	TFunction<FVector3d(const FVector3d& Position, const FExtrudeFrame& ExtrudeFrame, int32 SourceVid)> OffsetPositionFunc =
+		[this](const FVector3d& Position, const FExtrudeFrame& ExtrudeFrame, int32 SourceVid)
+	{
+		return ExtrudeFrame.Frame.FromFramePoint(FVector3d(this->DefaultOffsetDistance, 0, 0) * ExtrudeFrame.Scaling);
+	};
+
+	/** When generating extrude frames, whether to use unselected neighbors for setting the frame. */
+	bool bAssignAnyBoundaryNeighborToUnmatched = false;
+
+	/** 
+	 * If greater than 1, maximal amount by which a vertex can be moved in an attempt to keep edges parallel to
+	 *  original edges while extruding. This "movement" is done by scaling the X axis of the extrude frame.
+	 */
+	double ScalingAdjustmentLimit = 1.0;
+
+	/** 
+	 * Optional mapping, 1:1 with Eids, that gives the group id to use for each generated quad. Otherwise all generated
+	 *  triangles will be given the same new group id.
+	 */
+	TOptional<TArray<int32>> GroupsToSetPerEid;
+
+	/** Used in the default OffsetPositionFunc. */
+	double DefaultOffsetDistance = 1.0;
+
+
+	// Outputs:
+	TArray<int32> NewTids;
+	TArray<int32> NewExtrudedEids;
+
+public:
+	FExtrudeBoundaryEdges(FDynamicMesh3* mesh);
+
+	virtual ~FExtrudeBoundaryEdges() {}
+
+	/**
+	 * Apply the operation to the input mesh.
+	 * @return true if the algorithm succeeds
+	 */
+	virtual bool Apply(FProgressCancel* Progress);
+
+	/** 
+	 * Pairs up edges across vertices to help in extrusion frame calculation. Public because it is used by
+	 * the extrude edges activity to find an operational space for the gizmos used to set extrude distance.
+	 * 
+	 * @return false if there is an error.
+	 */
+	static bool GetInputEdgePairings(
+		const FDynamicMesh3& Mesh, TArray<int32>& InputEids, bool bAssignAnyBoundaryNeighborToUnmatched,
+		TArray<FNewVertSourceData>& NewVertDataOut, TMap<int32, FIndex2i>& EidToIndicesIntoNewVertsOut);
+
+	/**
+	 * Gets an extrude frame given a vert and its neighboring boundary edges. Like GetInputEdgePairings,
+	 * public because it is used to set up UX for setting extrude distance.
+	 * 
+	 @return false if there is an error.
+	 */
+	static bool GetExtrudeFrame(const FDynamicMesh3& Mesh, int32 Vid,
+		int32 IncomingEid, int32 OutgoingEid, FExtrudeFrame& ExtrudeFrameOut,
+		double ScalingLimit);
+};
+
+} // end namespace UE::Geometry
\ No newline at end of file
diff --git a/Engine/Plugins/Runtime/MeshModelingToolset/Source/MeshModelingTools/Private/ToolActivities/PolyEditExtrudeEdgeActivity.cpp b/Engine/Plugins/Runtime/MeshModelingToolset/Source/MeshModelingTools/Private/ToolActivities/PolyEditExtrudeEdgeActivity.cpp
new file mode 100644
--- /dev/null
+++ b/Engine/Plugins/Runtime/MeshModelingToolset/Source/MeshModelingTools/Private/ToolActivities/PolyEditExtrudeEdgeActivity.cpp
@@ -0,0 +1,739 @@
+// Copyright Epic Games, Inc. All Rights Reserved.
+
+#include "ToolActivities/PolyEditExtrudeEdgeActivity.h"
+
+#include "BaseGizmos/CombinedTransformGizmo.h"
+#include "BaseGizmos/TransformGizmoUtil.h"
+#include "BaseGizmos/TransformProxy.h"
+#include "Containers/Map.h"
+#include "ContextObjectStore.h"
+#include "Drawing/PreviewGeometryActor.h"
+#include "DynamicMesh/DynamicMeshChangeTracker.h"
+#include "DynamicMeshEditor.h"
+#include "InteractiveToolManager.h"
+#include "MeshOpPreviewHelpers.h"
+#include "ModelingOperators.h" // FDynamicMeshOperator
+#include "Operations/ExtrudeBoundaryEdges.h"
+#include "SceneManagement.h"
+#include "Selection/PolygonSelectionMechanic.h"
+#include "Selection/StoredMeshSelectionUtil.h"
+#include "ToolActivities/PolyEditActivityContext.h"
+#include "ToolContextInterfaces.h" // IToolsContextRenderAPI
+
+#define LOCTEXT_NAMESPACE "PolyEditExtrudeEdgeActivity"
+
+namespace PolyEditExtrudeEdgeActivityLocals
+{
+	// When adjusting extrusion distance to keep edges parallel, maximal allowed increase factor
+	const double MAX_VERT_MOVEMENT_ADJUSTMENT_SCALE = 4.0;
+
+	// These should match what is used in UMeshTopologySelectionMechanic (or better, live in a common place)
+	const FColor GroupLineColor = FColor::Red;
+	const float GroupLineThickness = 1.0f;
+
+	const FColor ExtrudeFrameLineColor = FColor::Orange;
+	const float ExtrudeFrameLineThickness = 0.5f;
+	const double ExtrudeFrameLineLength = 1000;
+
+	FText TransactionLabel = LOCTEXT("ExtrudeEdgeTransactionLabel", "Extrude Edges");
+
+	class FExtrudeEdgeOp : public UE::Geometry::FDynamicMeshOperator
+	{
+	public:
+		virtual ~FExtrudeEdgeOp() {}
+
+		// Inputs:
+		TSharedPtr<const FDynamicMesh3, ESPMode::ThreadSafe> OriginalMesh;
+		FTransform OriginalMeshTransform;
+		TArray<int32> SelectedEids;
+		TArray<int32> GroupsToSetPerEid;
+
+		EPolyEditExtrudeEdgeDirectionMode DirectionMode;
+		FVector3d SingleDirectionVector;
+		FVector3d LocalFrameParams;
+		double ScalingAdjustmentLimit = 1;
+		bool bAssignAnyBoundaryNeighborToUnmatched = false;
+
+		bool bGatherGroupBoundariesForRender = false;
+		bool bGatherAllEdgesForRender = false;
+
+		// Outputs
+		TArray<int32> NewSelectionEids;
+		TArray<int32> EdgesToRender;
+	protected:
+
+		TMap<int32, FVector3d> DesiredExtrudeDirections;
+
+	public:
+		// FDynamicMeshOperator
+		virtual void CalculateResult(FProgressCancel* Progress) override
+		{
+			using namespace UE::Geometry;
+
+			if (Progress && Progress->Cancelled()) { return; }
+
+			ResultMesh->Copy(*OriginalMesh);
+			ResultTransform = OriginalMeshTransform;
+
+			if (Progress && Progress->Cancelled()) { return; }
+
+			if ((DirectionMode == EPolyEditExtrudeEdgeDirectionMode::LocalExtrudeFrames && LocalFrameParams.IsNearlyZero())
+				|| (DirectionMode == EPolyEditExtrudeEdgeDirectionMode::SingleDirection && SingleDirectionVector.IsNearlyZero()))
+			{
+				return;
+			}
+
+			FExtrudeBoundaryEdges Extruder(ResultMesh.Get());
+
+			switch (DirectionMode)
+			{
+			case EPolyEditExtrudeEdgeDirectionMode::LocalExtrudeFrames:
+				Extruder.OffsetPositionFunc = [this](const FVector3d& Position, const FExtrudeBoundaryEdges::FExtrudeFrame& ExtrudeFrame, int32 SourceVid)
+				{
+					return ExtrudeFrame.Frame.FromFramePoint(ExtrudeFrame.Scaling * LocalFrameParams);
+				};
+				break;
+			case EPolyEditExtrudeEdgeDirectionMode::SingleDirection:
+				Extruder.OffsetPositionFunc = [this](const FVector3d& Position, const FExtrudeBoundaryEdges::FExtrudeFrame& ExtrudeFrame, int32 SourceVid)
+				{
+					return Position + SingleDirectionVector;
+				};
+				break;
+			}
+
+			Extruder.InputEids = SelectedEids;
+			Extruder.GroupsToSetPerEid = GroupsToSetPerEid;
+			Extruder.ScalingAdjustmentLimit = ScalingAdjustmentLimit;
+			Extruder.bUsePerVertexExtrudeFrames = DirectionMode == EPolyEditExtrudeEdgeDirectionMode::LocalExtrudeFrames;
+			Extruder.bAssignAnyBoundaryNeighborToUnmatched = bAssignAnyBoundaryNeighborToUnmatched;
+
+			if (Progress && Progress->Cancelled()) { return; }
+
+			Extruder.Apply(Progress);
+
+			if (Progress && Progress->Cancelled()) { return; }
+
+			NewSelectionEids = Extruder.NewExtrudedEids;
+
+			TSet<int32> EdgesToRenderSet;
+			if (bGatherAllEdgesForRender)
+			{
+				for (int32 Tid : Extruder.NewTids)
+				{
+					FIndex3i TriEids = ResultMesh->GetTriEdges(Tid);
+					for (int i = 0; i < 3; ++i)
+					{
+						EdgesToRenderSet.Add(TriEids[i]);
+					}
+				}
+			}
+			else if (bGatherGroupBoundariesForRender)
+			{
+				for (int32 Tid : Extruder.NewTids)
+				{
+					FIndex3i TriEids = ResultMesh->GetTriEdges(Tid);
+					for (int i = 0; i < 3; ++i)
+					{
+						if (ResultMesh->IsBoundaryEdge(TriEids[i]) || ResultMesh->IsGroupBoundaryEdge(TriEids[i]))
+						{
+							EdgesToRenderSet.Add(TriEids[i]);
+						}
+					}
+				}
+			}
+			EdgesToRender = EdgesToRenderSet.Array();
+		}
+	};
+}
+
+TUniquePtr<UE::Geometry::FDynamicMeshOperator> UPolyEditExtrudeEdgeActivity::MakeNewOperator()
+{
+	using namespace PolyEditExtrudeEdgeActivityLocals;
+
+	TUniquePtr<FExtrudeEdgeOp> Op = MakeUnique<FExtrudeEdgeOp>();
+	Op->OriginalMesh = MakeShared<FDynamicMesh3, ESPMode::ThreadSafe>(*ActivityContext->CurrentMesh);
+	Op->OriginalMeshTransform = ActivityContext->Preview->PreviewMesh->GetTransform();
+	Op->SelectedEids = SelectedEids;
+	Op->GroupsToSetPerEid = GroupsToSetPerEid;
+	Op->ScalingAdjustmentLimit = Settings->bAdjustToExtrudeEvenly ? MAX_VERT_MOVEMENT_ADJUSTMENT_SCALE : 1.0;
+	Op->DirectionMode = Settings->DirectionMode;
+	Op->bGatherAllEdgesForRender = ActivityContext->bTriangleMode;
+	Op->bGatherGroupBoundariesForRender = true;
+	Op->bAssignAnyBoundaryNeighborToUnmatched = Settings->bUseUnselectedForFrames;
+
+	if (Settings->DirectionMode == EPolyEditExtrudeEdgeDirectionMode::SingleDirection)
+	{
+		FVector3d WorldDestination = SingleDirectionVectorWorldSpace + ExtrudeFrameForGizmoWorldSpace.Origin;
+
+		Op->SingleDirectionVector = CurrentMeshTransform.InverseTransformPosition(WorldDestination)
+			- ExtrudeFrameForGizmoMeshSpace.Origin;
+	}
+	else
+	{
+		FVector3d WorldDestination = ExtrudeFrameForGizmoWorldSpace.FromFramePoint(ParamsInWorldExtrudeFrame);
+		
+		Op->LocalFrameParams = ExtrudeFrameForGizmoMeshSpace.ToFramePoint(
+			CurrentMeshTransform.InverseTransformPosition(WorldDestination));
+		Op->LocalFrameParams /= ExtrudeFrameScaling;
+	}
+
+	return Op;
+}
+
+
+void UPolyEditExtrudeEdgeActivity::Setup(UInteractiveTool* ParentToolIn)
+{
+	Super::Setup(ParentToolIn);
+
+	ActivityContext = ParentTool->GetToolManager()->GetContextObjectStore()->FindContext<UPolyEditActivityContext>();
+
+	Settings = NewObject<UPolyEditExtrudeEdgeActivityProperties>();
+	Settings->RestoreProperties(ParentTool.Get());
+	AddToolPropertySource(Settings);
+	SetToolPropertySourceEnabled(Settings, false);
+	
+	ExtrudeFrameProxy = NewObject<UTransformProxy>(this);
+	ExtrudeFrameProxy->OnTransformChanged.AddWeakLambda(this, [this](UTransformProxy*, FTransform Transform)
+	{
+		// TODO: It's possible to get these callbacks from undo while the activity is not running because the
+		// transform changes are only expired by exiting the tool, not the activity. This is easy enough to
+		// guard against, but may need a broader solution someday.
+		if (bIsRunning && ActivityContext && ActivityContext->Preview)
+		{
+			ParamsInWorldExtrudeFrame = ExtrudeFrameForGizmoWorldSpace.ToFramePoint(Transform.GetLocation());
+			ActivityContext->Preview->InvalidateResult();
+		}
+	});
+
+	SingleDirectionProxy = NewObject<UTransformProxy>(this);
+	SingleDirectionProxy->OnTransformChanged.AddWeakLambda(this, [this](UTransformProxy*, FTransform Transform)
+	{
+		if (bIsRunning && ActivityContext && ActivityContext->Preview)
+		{
+			SingleDirectionVectorWorldSpace = Transform.GetLocation() - ExtrudeFrameForGizmoWorldSpace.Origin;
+			ActivityContext->Preview->InvalidateResult();
+		}
+	});
+
+	Settings->WatchProperty(Settings->DirectionMode, [this](EPolyEditExtrudeEdgeDirectionMode) 
+	{
+		ConvertToNewDirectionMode(Settings->DirectionMode == EPolyEditExtrudeEdgeDirectionMode::SingleDirection);
+	});
+	Settings->WatchProperty(Settings->Distance, [this](double) 
+	{
+		if (!bIsRunning) { return; }
+		
+		// When distance changes, we keep the same extrude direction (set by gizmo or automatically) and
+		//  just adjust the length of extrusion in world space.
+
+		auto AdjustVectorLength = [this](FVector3d& VectorToAdjust)
+		{
+			VectorToAdjust.Normalize();
+			if (VectorToAdjust.IsNearlyZero())
+			{
+				VectorToAdjust = FVector3d::UnitX();
+			}
+			VectorToAdjust *= Settings->Distance;
+		};
+
+		if (Settings->DirectionMode == EPolyEditExtrudeEdgeDirectionMode::SingleDirection)
+		{
+			AdjustVectorLength(SingleDirectionVectorWorldSpace);
+		}
+		else
+		{
+			AdjustVectorLength(ParamsInWorldExtrudeFrame);
+		}
+
+		if (bIsRunning)
+		{
+			ActivityContext->Preview->InvalidateResult();
+		}
+	});
+	Settings->WatchProperty(Settings->DistanceMode, [this](EPolyEditExtrudeEdgeDistanceMode) 
+	{
+		if (!bIsRunning) { return; }
+
+		// When swapping distance mode, we would like things to stay pretty much where they are.
+		// This can be done by updating from params, since both methods of operation are just
+		//  setting those.
+		if (Settings->DistanceMode == EPolyEditExtrudeEdgeDistanceMode::Fixed)
+		{
+			UpdateDistanceFromParams();
+		}
+		else
+		{
+			UpdateGizmosFromCurrentParams();
+		}
+		UpdateGizmoVisibility(); 
+	});
+
+	auto UpdateExtrudeFrame = [this]()
+	{
+		if (!bIsRunning) { return; }
+
+		RecalculateGizmoExtrudeFrame();
+		UpdateGizmosFromCurrentParams();
+		ActivityContext->Preview->InvalidateResult();
+	};
+
+	Settings->WatchProperty(Settings->bUseUnselectedForFrames, [UpdateExtrudeFrame](bool)
+	{
+		UpdateExtrudeFrame();
+	});
+	Settings->WatchProperty(Settings->bAdjustToExtrudeEvenly, [UpdateExtrudeFrame](bool)
+	{
+		UpdateExtrudeFrame();
+	});
+
+	Settings->SilentUpdateWatched();
+}
+
+
+EToolActivityStartResult UPolyEditExtrudeEdgeActivity::Start()
+{
+	using namespace PolyEditExtrudeEdgeActivityLocals;
+
+	if (!CanStart())
+	{
+		ParentTool->GetToolManager()->DisplayMessage(
+			LOCTEXT("OnExtrudeFailedMesssage", "Action requires boundary edge selection."),
+			EToolMessageLevel::UserWarning);
+		return EToolActivityStartResult::FailedStart;
+	}
+
+	CurrentMeshTransform = ActivityContext->Preview->PreviewMesh->GetTransform();
+
+	ActivityContext->Preview->ChangeOpFactory(this);
+	ActivityContext->Preview->OnOpCompleted.AddWeakLambda(this,
+		[this](const UE::Geometry::FDynamicMeshOperator* UncastOp)
+	{
+		const FExtrudeEdgeOp* Op = static_cast<const FExtrudeEdgeOp*>(UncastOp);
+		NewSelectionEids = Op->NewSelectionEids;
+		EidsToRender = Op->EdgesToRender;
+	});
+	ActivityContext->Preview->OnMeshUpdated.AddWeakLambda(this,
+		[this](UMeshOpPreviewWithBackgroundCompute*)
+	{
+		// Do this here instead of OnOpCompleted, because we need access to the updated mesh.
+		UpdateDrawnPreviewEdges();
+	});
+
+	PreviewGeometry = NewObject<UPreviewGeometry>();
+	PreviewGeometry->CreateInWorld(ActivityContext->Preview->GetWorld(), FTransform::Identity);
+	PreviewGeometry->SetTransform(CurrentMeshTransform);
+
+	SetToolPropertySourceEnabled(Settings, true);
+	bIsRunning = true;
+
+	// Things seem to look a little better if we clear the selection highlighting. But this means we need to
+	// revert the selection before issuing any undo/redo transactions. This is done either in ApplyExtrude or
+	// EndInternal.
+	ActiveSelection = ActivityContext->SelectionMechanic->GetActiveSelection();
+	ActivityContext->SelectionMechanic->ClearSelection();
+	bRevertedSelection = false;
+	GatherSelectedEids();
+
+	UInteractiveGizmoManager* GizmoManager = GetParentTool()->GetToolManager()->GetPairedGizmoManager();
+	ExtrudeFrameGizmo = UE::TransformGizmoUtil::CreateCustomTransformGizmo(GizmoManager,
+		ETransformGizmoSubElements::TranslateAxisX | ETransformGizmoSubElements::TranslateAxisZ | ETransformGizmoSubElements::TranslatePlaneXZ,
+		this);
+	ExtrudeFrameGizmo->SetActiveTarget(ExtrudeFrameProxy, GetParentTool()->GetToolManager());
+	ExtrudeFrameGizmo->SetVisibility(false);
+	// We force the coordinate system to be local so that the gizmo only moves in the plane we specify
+	ExtrudeFrameGizmo->bUseContextCoordinateSystem = false;
+	ExtrudeFrameGizmo->CurrentCoordinateSystem = EToolContextCoordinateSystem::Local;
+
+	SingleDirectionGizmo = UE::TransformGizmoUtil::CreateCustomTransformGizmo(GizmoManager,
+		ETransformGizmoSubElements::TranslateAllAxes | ETransformGizmoSubElements::TranslateAllPlanes,
+		this);
+	SingleDirectionGizmo->SetActiveTarget(SingleDirectionProxy);
+	SingleDirectionGizmo->SetVisibility(false);
+
+	RecalculateGizmoExtrudeFrame();
+
+	ResetParams();
+	UpdateGizmosFromCurrentParams();
+
+	UpdateGizmoVisibility();
+
+	// Do this as the last thing so that everything is set up for it.
+	ActivityContext->Preview->InvalidateResult();
+
+	ActivityContext->EmitActivityStart(LOCTEXT("BeginExtrudeEdgeActivity", "Begin Extrude Edges"));
+	return EToolActivityStartResult::Running;
+}
+
+
+EToolActivityEndResult UPolyEditExtrudeEdgeActivity::End(EToolShutdownType ShutdownType)
+{
+	if (!ensure(bIsRunning))
+	{
+		EndInternal();
+		return EToolActivityEndResult::ErrorDuringEnd;
+	}
+
+	if (ShutdownType == EToolShutdownType::Cancel)
+	{
+		EndInternal();
+
+		// Reset the preview
+		ActivityContext->Preview->PreviewMesh->UpdatePreview(ActivityContext->CurrentMesh.Get());
+
+		return EToolActivityEndResult::Cancelled;
+	}
+	else
+	{
+		// Stop the current compute if there is one.
+		ActivityContext->Preview->CancelCompute();
+
+		// Apply whatever we happen to have at the moment.
+		ApplyExtrude();
+
+		EndInternal();
+		return EToolActivityEndResult::Completed;
+	}
+}
+
+
+void UPolyEditExtrudeEdgeActivity::ApplyExtrude()
+{
+	using namespace UE::Geometry;
+	using namespace PolyEditExtrudeEdgeActivityLocals;
+
+	const FDynamicMesh3* ResultMesh = ActivityContext->Preview->PreviewMesh->GetMesh();
+
+	UE::Geometry::FMeshEdgesFromTriangleSubIndices StableEdgeIDs;
+	StableEdgeIDs.InitializeFromEdgeIDs(*ResultMesh, NewSelectionEids);
+
+	FDynamicMeshChangeTracker ChangeTracker(ActivityContext->CurrentMesh.Get());
+	ChangeTracker.BeginChange();
+
+	ActivityContext->CurrentMesh->Copy(*ResultMesh);
+
+	// We need to reset the old selection back before we give the new one, so that undo reverts back
+	//  to the correct selection state. Currently, the boolean check is only necessary in EndInternal,
+	//  but we'll keep it in case we move code around.
+	if (!bRevertedSelection)
+	{
+		ActivityContext->SelectionMechanic->SetSelection(ActiveSelection, false);
+		bRevertedSelection = true;
+	}
+
+	ParentTool->GetToolManager()->BeginUndoTransaction(TransactionLabel);
+	FGroupTopologySelection NewSelection;
+	ActivityContext->EmitCurrentMeshChangeAndUpdate(TransactionLabel, ChangeTracker.EndChange(), NewSelection);
+	ActivityContext->SelectionMechanic->BeginChange();
+
+	StableEdgeIDs.GetEdgeIDs(*ActivityContext->CurrentMesh, NewSelectionEids);
+
+	if (ActivityContext->bTriangleMode)
+	{
+		NewSelection.SelectedEdgeIDs.Append(NewSelectionEids);
+	}
+	else
+	{
+		for (int32 Eid : NewSelectionEids)
+		{
+			NewSelection.SelectedEdgeIDs.Add(ActivityContext->CurrentTopology->FindGroupEdgeID(Eid));
+		}
+	}
+	ActivityContext->SelectionMechanic->SetSelection(NewSelection);
+	ActivityContext->SelectionMechanic->EndChangeAndEmitIfModified();
+	ParentTool->GetToolManager()->EndUndoTransaction();
+}
+
+
+void UPolyEditExtrudeEdgeActivity::EndInternal()
+{
+	// Fix up the selection we cleared to remove the highlighting (if we haven't already done it via an Apply)
+	if (!bRevertedSelection)
+	{
+		ActivityContext->SelectionMechanic->SetSelection(ActiveSelection, false);
+		bRevertedSelection = true;
+	}
+
+	ActivityContext->Preview->ClearOpFactory();
+	ActivityContext->Preview->OnOpCompleted.RemoveAll(this);
+	ActivityContext->Preview->OnMeshUpdated.RemoveAll(this);
+	SetToolPropertySourceEnabled(Settings, false);
+
+	UInteractiveGizmoManager* GizmoManager = GetParentTool()->GetToolManager()->GetPairedGizmoManager();
+	GizmoManager->DestroyAllGizmosByOwner(this);
+
+	PreviewGeometry->Disconnect();
+	PreviewGeometry = nullptr;
+
+	bIsRunning = false;
+}
+
+
+
+void UPolyEditExtrudeEdgeActivity::Shutdown(EToolShutdownType ShutdownType)
+{
+	if (IsRunning())
+	{
+		End(ShutdownType);
+	}
+	Super::Shutdown(ShutdownType);
+}
+
+bool UPolyEditExtrudeEdgeActivity::CanStart() const
+{
+	using namespace UE::Geometry;
+	if (!ActivityContext)
+	{
+		return false;
+	}
+	const FGroupTopologySelection& Selection = ActivityContext->SelectionMechanic->GetActiveSelection();
+	if (Selection.SelectedEdgeIDs.IsEmpty())
+	{
+		return false;
+	}
+	for (int32 GroupEdgeID : Selection.SelectedEdgeIDs)
+	{
+		if (ActivityContext->CurrentTopology->IsBoundaryEdge(GroupEdgeID))
+		{
+			return true;
+		}
+	}
+	return false;
+}
+
+bool UPolyEditExtrudeEdgeActivity::CanAccept() const
+{
+	return true;
+}
+
+
+void UPolyEditExtrudeEdgeActivity::UpdateGizmoVisibility()
+{
+	if (!bIsRunning || !ensure(SingleDirectionGizmo && ExtrudeFrameGizmo))
+	{
+		return;
+	}
+	if (Settings->DistanceMode == EPolyEditExtrudeEdgeDistanceMode::Fixed)
+	{
+		SingleDirectionGizmo->SetVisibility(false);
+		ExtrudeFrameGizmo->SetVisibility(false);
+	}
+	else
+	{
+		ExtrudeFrameGizmo->SetVisibility(Settings->DirectionMode == EPolyEditExtrudeEdgeDirectionMode::LocalExtrudeFrames);
+		SingleDirectionGizmo->SetVisibility(Settings->DirectionMode == EPolyEditExtrudeEdgeDirectionMode::SingleDirection);
+	}
+}
+
+void UPolyEditExtrudeEdgeActivity::GatherSelectedEids()
+{
+	SelectedEids.Reset();
+	GroupsToSetPerEid.Reset();
+
+	int32 GroupID = ActivityContext->CurrentMesh->MaxGroupID();
+
+	// Only used in triangle mode, to avoid creating a bunch of new groups
+	TMap<int32, int32> NeighborGroupToNewGroup; 
+	
+	for (int32 GroupEdgeID : ActiveSelection.SelectedEdgeIDs)
+	{
+		if (!ActivityContext->CurrentTopology->IsBoundaryEdge(GroupEdgeID))
+		{
+			continue;
+		}
+
+		if (ActivityContext->bTriangleMode)
+		{
+			SelectedEids.Add(GroupEdgeID);
+
+			int32 Tid = ActivityContext->CurrentMesh->GetEdgeT(GroupEdgeID).A;
+			int32 NeighborGroup = ActivityContext->CurrentMesh->GetTriangleGroup(Tid);
+
+			int32* NewGroupID = NeighborGroupToNewGroup.Find(NeighborGroup);
+			if (!NewGroupID)
+			{
+				NewGroupID = &NeighborGroupToNewGroup.Add(NeighborGroup, GroupID);
+				++GroupID;
+			}
+			GroupsToSetPerEid.Add(*NewGroupID);
+		}
+		else
+		{
+			for (int32 Eid : ActivityContext->CurrentTopology->GetGroupEdgeEdges(GroupEdgeID))
+			{
+				SelectedEids.Add(Eid);
+				GroupsToSetPerEid.Add(GroupID);
+			}
+
+			++GroupID;
+		}
+	}
+}
+
+// We draw the boundary edges for our extrusion preview. This is not something we do in other
+//  cases, for instance while in push/pull and the face extrude activity. But for edge extrusion,
+//  it is helpful to see the boundary because you may be looking behind the face that you are
+//  extruding, making it invisible.
+void UPolyEditExtrudeEdgeActivity::UpdateDrawnPreviewEdges()
+{
+	using namespace PolyEditExtrudeEdgeActivityLocals;
+
+	if (!bIsRunning 
+		|| !ActivityContext || !ActivityContext->Preview || !ActivityContext->Preview->PreviewMesh 
+		|| !PreviewGeometry)
+	{
+		return;
+	}
+
+	ActivityContext->Preview->PreviewMesh->ProcessMesh([this](const FDynamicMesh3& Mesh) 
+	{
+		PreviewGeometry->CreateOrUpdateLineSet(TEXT("NewBoundaries"), EidsToRender.Num(), [this, &Mesh]
+			(int32 Index, TArray<FRenderableLine>& Lines)
+		{
+			int32 Eid = EidsToRender[Index];
+			if (!ensure(Mesh.IsEdge(Eid)))
+			{
+				return;
+			}
+			
+			FIndex2i EdgeVids = Mesh.GetEdgeV(Eid);
+			Lines.Add(FRenderableLine(Mesh.GetVertex(EdgeVids.A), Mesh.GetVertex(EdgeVids.B),
+				GroupLineColor, GroupLineThickness));
+		});
+	});
+}
+
+
+// Updates Settings->Distance from the parameters
+void UPolyEditExtrudeEdgeActivity::UpdateDistanceFromParams()
+{
+	if (Settings->DirectionMode == EPolyEditExtrudeEdgeDirectionMode::SingleDirection)
+	{
+		Settings->Distance = SingleDirectionVectorWorldSpace.Length();
+	}
+	else // if (Settings->DirectionMode == EPolyEditExtrudeEdgeDirectionMode::LocalExtrudeFrames)
+	{
+		Settings->Distance = ExtrudeFrameForGizmoWorldSpace.FromFrameVector(ParamsInWorldExtrudeFrame).Length();
+	}
+}
+
+void UPolyEditExtrudeEdgeActivity::UpdateGizmosFromCurrentParams()
+{
+	FVector3d SingleDirectionGizmoLocation = SingleDirectionVectorWorldSpace + ExtrudeFrameForGizmoWorldSpace.Origin;
+
+	FVector3d ExtrudeFrameGizmoLocation = ExtrudeFrameForGizmoWorldSpace.FromFramePoint(ParamsInWorldExtrudeFrame);
+
+	SingleDirectionGizmo->ReinitializeGizmoTransform(FTransform(
+		FQuat(ExtrudeFrameForGizmoWorldSpace.Rotation),
+		SingleDirectionGizmoLocation));
+
+	ExtrudeFrameGizmo->ReinitializeGizmoTransform(FTransform(
+		FQuat(ExtrudeFrameForGizmoWorldSpace.Rotation),
+		ExtrudeFrameGizmoLocation));
+
+	ExtrudeFrameGizmo->SetDisplaySpaceTransform(ExtrudeFrameForGizmoWorldSpace.ToFTransform());
+}
+
+void UPolyEditExtrudeEdgeActivity::ConvertToNewDirectionMode(bool bToSingleDirection)
+{
+	if (bToSingleDirection)
+	{
+		// Convert fixed parameter
+		FVector3d CurrentPosition = ExtrudeFrameForGizmoWorldSpace.FromFramePoint(ParamsInWorldExtrudeFrame);
+		SingleDirectionVectorWorldSpace = CurrentPosition - ExtrudeFrameForGizmoWorldSpace.Origin;
+
+		// Convert gizmo location
+		CurrentPosition = ExtrudeFrameGizmo->GetGizmoTransform().GetLocation();
+		FTransform NewTransform(FQuat(ExtrudeFrameForGizmoWorldSpace.Rotation), CurrentPosition);
+		SingleDirectionGizmo->ReinitializeGizmoTransform(NewTransform);
+	}
+	else
+	{
+		// Convert fixed parameter
+		FVector3d CurrentPosition = ExtrudeFrameForGizmoWorldSpace.Origin + SingleDirectionVectorWorldSpace;
+		ParamsInWorldExtrudeFrame = ExtrudeFrameForGizmoWorldSpace.ToFramePoint(CurrentPosition);
+		ParamsInWorldExtrudeFrame.Y = 0;
+
+		// Convert gizmo location
+		CurrentPosition = SingleDirectionGizmo->GetGizmoTransform().GetLocation();
+		FVector3d FramePoint = ExtrudeFrameForGizmoWorldSpace.ToFramePoint(CurrentPosition);
+		FramePoint.Y = 0;
+		FTransform NewTransform(FQuat(ExtrudeFrameForGizmoWorldSpace.Rotation), 
+			ExtrudeFrameForGizmoWorldSpace.FromFramePoint(FramePoint));
+		ExtrudeFrameGizmo->ReinitializeGizmoTransform(NewTransform);
+	}
+	UpdateGizmoVisibility();
+
+	if (bIsRunning)
+	{
+		ActivityContext->Preview->InvalidateResult();
+	}
+}
+
+void UPolyEditExtrudeEdgeActivity::ResetParams()
+{
+	SingleDirectionVectorWorldSpace = FVector3d::UnitZ() * Settings->Distance;
+	ParamsInWorldExtrudeFrame = FVector3d::UnitX() * Settings->Distance;
+}
+
+void UPolyEditExtrudeEdgeActivity::RecalculateGizmoExtrudeFrame()
+{
+	using namespace PolyEditExtrudeEdgeActivityLocals;
+	using namespace UE::Geometry;
+
+	// Getting an extrude frame is actually tricky because we have to grab two boundary edges that are
+	// adjacent to some vertex in the selection, and the way these get selected is dependent on the 
+	// implementation of FExtrudeBoundaryEdges (because there is ambiguity in the case of bowties).
+
+	TArray<FExtrudeBoundaryEdges::FNewVertSourceData> NewVertData;
+	TMap<int32, FIndex2i> EidToIndicesIntoNewVerts;
+	FExtrudeBoundaryEdges::GetInputEdgePairings(
+		*ActivityContext->CurrentMesh, SelectedEids, Settings->bUseUnselectedForFrames,
+		NewVertData, EidToIndicesIntoNewVerts);
+
+	if (!ensure(!NewVertData.IsEmpty()))
+	{
+		return;
+	}
+
+	FExtrudeBoundaryEdges::FExtrudeFrame ExtrudeFrame;
+	FExtrudeBoundaryEdges::GetExtrudeFrame(
+		*ActivityContext->CurrentMesh, 
+		NewVertData[0].SourceVid, 
+		NewVertData[0].SourceEidPair.A,
+		NewVertData[0].SourceEidPair.B, 
+		ExtrudeFrame, Settings->bAdjustToExtrudeEvenly ? MAX_VERT_MOVEMENT_ADJUSTMENT_SCALE : 1.0);
+
+	ExtrudeFrameForGizmoMeshSpace = ExtrudeFrame.Frame;
+	ExtrudeFrameScaling = ExtrudeFrame.Scaling;
+
+	// This method of finding the extrude frame in world space isn't quite right in the case of nonuniform scale. However
+	// it seems hard to find any perfect solution, and in the end, it's not critical because we'll probably be able to get
+	// extrusion parameters out of an imperfect frame anyway.
+	ExtrudeFrameForGizmoWorldSpace = FFrame3d(
+		CurrentMeshTransform.TransformPosition(ExtrudeFrameForGizmoMeshSpace.Origin),
+		CurrentMeshTransform.TransformRotation(FQuat(ExtrudeFrameForGizmoMeshSpace.Rotation)));
+}
+
+void UPolyEditExtrudeEdgeActivity::Render(IToolsContextRenderAPI* RenderAPI)
+{
+	using namespace PolyEditExtrudeEdgeActivityLocals;
+
+	if (bIsRunning)
+	{
+		// Draw a line showing where the X axis of the extrude frame is. We could use the preview actor for this but
+		// this is simpler and doesn't require us to set up the foreground material.
+		FPrimitiveDrawInterface* PDI = RenderAPI->GetPrimitiveDrawInterface();
+		PDI->DrawLine(ExtrudeFrameForGizmoWorldSpace.Origin,
+			ExtrudeFrameForGizmoWorldSpace.FromFramePoint(FVector3d(1000, 0, 0)),
+			ExtrudeFrameLineColor, ExtrudeFrameLineThickness, SDPG_Foreground, 1.0, true);
+	}
+	
+}
+
+void UPolyEditExtrudeEdgeActivity::Tick(float DeltaTime)
+{
+}
+
+
+
+#undef LOCTEXT_NAMESPACE
\ No newline at end of file
diff --git a/Engine/Plugins/Runtime/MeshModelingToolset/Source/MeshModelingTools/Private/ToolActivities/PolyEditExtrudeEdgeActivity.h b/Engine/Plugins/Runtime/MeshModelingToolset/Source/MeshModelingTools/Private/ToolActivities/PolyEditExtrudeEdgeActivity.h
new file mode 100644
--- /dev/null
+++ b/Engine/Plugins/Runtime/MeshModelingToolset/Source/MeshModelingTools/Private/ToolActivities/PolyEditExtrudeEdgeActivity.h
@@ -0,0 +1,153 @@
+// Copyright Epic Games, Inc. All Rights Reserved.
+
+#pragma once
+
+#include "FrameTypes.h"
+#include "InteractiveTool.h" // UInteractiveToolPropertySet
+#include "InteractiveToolActivity.h"
+#include "GroupTopology.h" // FGroupTopologySelection
+#include "ModelingOperators.h" // IDynamicMeshOperatorFactory
+
+#include "PolyEditExtrudeEdgeActivity.generated.h"
+
+class UCombinedTransformGizmo;
+class UPolyEditActivityContext;
+class UPreviewGeometry;
+class UTransformProxy;
+
+UENUM()
+enum class EPolyEditExtrudeEdgeDirectionMode
+{
+	/** Each vertex gets its own local frame to extrude along. */
+	LocalExtrudeFrames,
+	/** All vertices are extruded in the same direction */
+	SingleDirection,
+};
+
+
+UENUM()
+enum class EPolyEditExtrudeEdgeDistanceMode
+{
+	/** Sets the distance numerically using the Distance parameter. */
+	Fixed,
+	/** Uses a gizmo attached to one of the vertices to determine distance to extrude. */
+	Gizmo
+};
+
+UCLASS()
+class UPolyEditExtrudeEdgeActivityProperties : public UInteractiveToolPropertySet
+{
+	GENERATED_BODY()
+
+public:
+	
+	/** How direction to move the vertices is determined. */
+	UPROPERTY(EditAnywhere, Category = Extrude)
+	EPolyEditExtrudeEdgeDirectionMode DirectionMode = EPolyEditExtrudeEdgeDirectionMode::LocalExtrudeFrames;
+
+	/** How distance to extrude along a given direction is determined. */
+	UPROPERTY(EditAnywhere, Category = Extrude)
+	EPolyEditExtrudeEdgeDistanceMode DistanceMode = EPolyEditExtrudeEdgeDistanceMode::Gizmo;
+
+	UPROPERTY(EditAnywhere, Category = Extrude, meta = (EditCondition = "DistanceMode == EPolyEditExtrudeEdgeDistanceMode::Fixed", EditConditionHides,
+		UIMin = "0", UIMax = "200", ClampMin = "-100000", ClampMax = "100000"))
+	double Distance = 20;
+
+	/** When a vertex has both a selected and nonselected neighbor, use the unselected neighbor in picking an extrude frame as well. */
+	UPROPERTY(EditAnywhere, Category = Extrude, meta = (
+		EditCondition = "DirectionMode == EPolyEditExtrudeEdgeDirectionMode::LocalExtrudeFrames", EditConditionHides))
+	bool bUseUnselectedForFrames = false;
+
+	/** Adjust individual extrude directions in an effort to make extruded edges parallel to their original edges. */
+	UPROPERTY(EditAnywhere, Category = Extrude, meta = (
+		EditCondition = "DirectionMode == EPolyEditExtrudeEdgeDirectionMode::LocalExtrudeFrames", EditConditionHides))
+	bool bAdjustToExtrudeEvenly = false;
+};
+
+
+/**
+ * Extrudes boundary edges.
+ */
+UCLASS()
+class UPolyEditExtrudeEdgeActivity : public UInteractiveToolActivity,
+	public UE::Geometry::IDynamicMeshOperatorFactory
+
+{
+	GENERATED_BODY()
+
+public:
+	using FFrame3d = UE::Geometry::FFrame3d;
+
+	// IInteractiveToolActivity
+	virtual void Setup(UInteractiveTool* ParentTool) override;
+	virtual void Shutdown(EToolShutdownType ShutdownType) override;
+	virtual bool CanStart() const override;
+	virtual EToolActivityStartResult Start() override;
+	virtual bool IsRunning() const override { return bIsRunning; }
+	virtual bool HasAccept() const override { return true; }
+	virtual bool CanAccept() const override;
+	virtual EToolActivityEndResult End(EToolShutdownType) override;
+	virtual void Render(IToolsContextRenderAPI* RenderAPI) override;
+	virtual void Tick(float DeltaTime) override;
+
+	// IDynamicMeshOperatorFactory
+	virtual TUniquePtr<UE::Geometry::FDynamicMeshOperator> MakeNewOperator() override;
+
+	UPROPERTY()
+	TObjectPtr<UPolyEditExtrudeEdgeActivityProperties> Settings = nullptr;
+
+protected:
+
+	UPROPERTY()
+	TObjectPtr<UPolyEditActivityContext> ActivityContext;
+
+	// Support for gizmos.
+	UPROPERTY()
+	TObjectPtr<UTransformProxy> ExtrudeFrameProxy;
+	UPROPERTY()
+	TObjectPtr<UCombinedTransformGizmo> ExtrudeFrameGizmo;
+	UPROPERTY()
+	TObjectPtr<UTransformProxy> SingleDirectionProxy;
+	UPROPERTY()
+	TObjectPtr<UCombinedTransformGizmo> SingleDirectionGizmo;
+
+	// Used for drawing the boundaries of the new faces to make them look similar to the
+	// rest of the mesh.
+	UPROPERTY()
+	TObjectPtr<UPreviewGeometry> PreviewGeometry;
+
+	bool bIsRunning = false;
+
+	void ApplyExtrude();
+	void EndInternal();
+
+	FFrame3d ExtrudeFrameForGizmoMeshSpace;
+	FVector3d ExtrudeFrameScaling; // in mesh space, when using bAdjustToExtrudeEvenly 
+	FFrame3d ExtrudeFrameForGizmoWorldSpace;
+	
+	// These are the parameters we actually end up using, initialized from gizmos or detail panel.
+	FVector3d SingleDirectionVectorWorldSpace;
+	FVector3d ParamsInWorldExtrudeFrame;
+
+	void UpdateGizmosFromCurrentParams();
+	void UpdateDistanceFromParams();
+	void ConvertToNewDirectionMode(bool bToSingleDirection);
+	void ResetParams();
+
+	void RecalculateGizmoExtrudeFrame();
+	void UpdateGizmoVisibility();
+
+	UE::Geometry::FGroupTopologySelection ActiveSelection;
+	bool bRevertedSelection = false;
+	TArray<int32> SelectedEids;
+	TArray<int32> GroupsToSetPerEid;
+	void GatherSelectedEids();
+
+	TArray<int32> EidsToRender;
+	void UpdateDrawnPreviewEdges();
+
+	TArray<int32> NewSelectionEids;
+
+	// TODO: Might be worth having a getter in ActivityContext
+	FTransform CurrentMeshTransform;
+};
diff --git a/Engine/Plugins/Runtime/MeshModelingToolset/Source/MeshModelingTools/Public/EditMeshPolygonsTool.h b/Engine/Plugins/Runtime/MeshModelingToolset/Source/MeshModelingTools/Public/EditMeshPolygonsTool.h
--- a/Engine/Plugins/Runtime/MeshModelingToolset/Source/MeshModelingTools/Public/EditMeshPolygonsTool.h
+++ b/Engine/Plugins/Runtime/MeshModelingToolset/Source/MeshModelingTools/Public/EditMeshPolygonsTool.h
@@ -35,6 +35,7 @@ class UPolyEditInsetOutsetActivity;
 class UPolyEditCutFacesActivity;
 class UPolyEditPlanarProjectionUVActivity;
 class UPolyEditBevelEdgeActivity;
+class UPolyEditExtrudeEdgeActivity;
 class UPolygonSelectionMechanic;
 class UTransformProxy;
 
@@ -131,6 +132,7 @@ enum class EEditMeshPolygonsToolActions
 	StraightenEdge,
 	FillHole,
 	BridgeEdges,
+	ExtrudeEdges,
 	BevelEdges,
 	SimplifyAlongEdges,
 
@@ -420,12 +422,16 @@ public:
 	UFUNCTION(CallInEditor, Category = EdgeEdits, meta = (DisplayName = "Bridge", DisplayPriority = 5))
 	void Bridge() { PostAction(EEditMeshPolygonsToolActions::BridgeEdges); }
 
-	/** Simplify the underlying triangulation along the selected edges, when doing so won't change the shape or UVs, or make low-quality triangles */
+	/** Duplicate and move boundary edge vertices outwards and connect them to the original boundary to create new faces. */
 	UFUNCTION(CallInEditor, Category = EdgeEdits, meta = (DisplayPriority = 6))
+	void Extrude() { PostAction(EEditMeshPolygonsToolActions::ExtrudeEdges); }
+
+	/** Simplify the underlying triangulation along the selected edges, when doing so won't change the shape or UVs, or make low-quality triangles */
+	UFUNCTION(CallInEditor, Category = EdgeEdits, meta = (DisplayPriority = 7))
 	void Simplify() { PostAction(EEditMeshPolygonsToolActions::SimplifyAlongEdges); }
 	
 	/** Delete selected edge, implicitly merging any connected faces */
-	UFUNCTION(CallInEditor, Category = EdgeEdits, meta = (DisplayName = "Delete Edges", DisplayPriority = 7))
+	UFUNCTION(CallInEditor, Category = EdgeEdits, meta = (DisplayName = "Delete Edges", DisplayPriority = 8))
 	void DeleteEdge() { PostAction(EEditMeshPolygonsToolActions::Delete); }
 };
 
@@ -443,6 +449,10 @@ public:
 	UFUNCTION(CallInEditor, Category = EdgeEdits, meta = (DisplayName = "Fill Hole", DisplayPriority = 1))
 	void FillHole() { PostAction(EEditMeshPolygonsToolActions::FillHole); }
 
+	/** Duplicate and move boundary vertices outwards and connect them to the original boundary to create new faces. */
+	UFUNCTION(CallInEditor, Category = EdgeEdits, meta = (DisplayName = "Extrude", DisplayPriority = 1))
+	void Extrude() { PostAction(EEditMeshPolygonsToolActions::ExtrudeEdges); }
+
 	/** Collapse the selected edges, deleting the attached triangles and merging its two vertices into one */
 	UFUNCTION(CallInEditor, Category = EdgeEdits, meta = (DisplayName = "Collapse", DisplayPriority = 1))
 	void Collapse() { PostAction(EEditMeshPolygonsToolActions::CollapseSingleEdge); }
@@ -567,6 +577,8 @@ protected:
 	TObjectPtr<UPolyEditInsertEdgeLoopActivity> InsertEdgeLoopActivity = nullptr;
 	UPROPERTY()
 	TObjectPtr<UPolyEditBevelEdgeActivity> BevelEdgeActivity = nullptr;
+	UPROPERTY()
+	TObjectPtr<UPolyEditExtrudeEdgeActivity> ExtrudeEdgeActivity = nullptr;
 
 	TMap<UInteractiveToolActivity*, FText> ActivityLabels;
 	TMap<UInteractiveToolActivity*, FName> ActivityIconNames;
