MeshModelingTools: Fix CIS static analysis warning about a place where parameters turned out to be switched.

diff --git a/Engine/Plugins/Runtime/MeshModelingToolset/Source/MeshModelingTools/Private/ToolActivities/PolyEditExtrudeEdgeActivity.cpp b/Engine/Plugins/Runtime/MeshModelingToolset/Source/MeshModelingTools/Private/ToolActivities/PolyEditExtrudeEdgeActivity.cpp
--- a/Engine/Plugins/Runtime/MeshModelingToolset/Source/MeshModelingTools/Private/ToolActivities/PolyEditExtrudeEdgeActivity.cpp
+++ b/Engine/Plugins/Runtime/MeshModelingToolset/Source/MeshModelingTools/Private/ToolActivities/PolyEditExtrudeEdgeActivity.cpp
@@ -725,7 +725,7 @@ void UPolyEditExtrudeEdgeActivity::Render(IToolsContextRenderAPI* RenderAPI)
 		FPrimitiveDrawInterface* PDI = RenderAPI->GetPrimitiveDrawInterface();
 		PDI->DrawLine(ExtrudeFrameForGizmoWorldSpace.Origin,
 			ExtrudeFrameForGizmoWorldSpace.FromFramePoint(FVector3d(1000, 0, 0)),
-			ExtrudeFrameLineColor, ExtrudeFrameLineThickness, SDPG_Foreground, 1.0, true);
+			ExtrudeFrameLineColor, SDPG_Foreground, ExtrudeFrameLineThickness, 1.0, true);
 	}
 	
 }
