MeshModelingTools: Fix edge extrude single direction initialization not matching what happens if you do extrude frames first and then change over.

diff --git a/Engine/Plugins/Runtime/MeshModelingToolset/Source/MeshModelingTools/Private/ToolActivities/PolyEditExtrudeEdgeActivity.cpp b/Engine/Plugins/Runtime/MeshModelingToolset/Source/MeshModelingTools/Private/ToolActivities/PolyEditExtrudeEdgeActivity.cpp
--- a/Engine/Plugins/Runtime/MeshModelingToolset/Source/MeshModelingTools/Private/ToolActivities/PolyEditExtrudeEdgeActivity.cpp
+++ b/Engine/Plugins/Runtime/MeshModelingToolset/Source/MeshModelingTools/Private/ToolActivities/PolyEditExtrudeEdgeActivity.cpp
@@ -671,7 +671,7 @@ void UPolyEditExtrudeEdgeActivity::ConvertToNewDirectionMode(bool bToSingleDirec
 
 void UPolyEditExtrudeEdgeActivity::ResetParams()
 {
-	SingleDirectionVectorWorldSpace = FVector3d::UnitZ() * Settings->Distance;
+	SingleDirectionVectorWorldSpace = ExtrudeFrameForGizmoWorldSpace.X() * Settings->Distance;
 	ParamsInWorldExtrudeFrame = FVector3d::UnitX() * Settings->Distance;
 }
 
