add support for instance colors in scene component -> dynamic mesh geometry script conversion function

diff --git a/Engine/Plugins/Runtime/GeometryScripting/Source/GeometryScriptingCore/Private/SceneUtilityFunctions.cpp b/Engine/Plugins/Runtime/GeometryScripting/Source/GeometryScriptingCore/Private/SceneUtilityFunctions.cpp
--- a/Engine/Plugins/Runtime/GeometryScripting/Source/GeometryScriptingCore/Private/SceneUtilityFunctions.cpp
+++ b/Engine/Plugins/Runtime/GeometryScripting/Source/GeometryScriptingCore/Private/SceneUtilityFunctions.cpp
@@ -73,6 +73,7 @@ UDynamicMesh* UGeometryScriptLibrary_SceneUtilityFunctions::CopyMeshFromComponen
 	ToMeshOptions.LODType = SafeConvertLODType(Options.RequestedLOD.LODType);
 	ToMeshOptions.bWantNormals = Options.bWantNormals;
 	ToMeshOptions.bWantTangents = Options.bWantTangents;
+	ToMeshOptions.bWantInstanceColors = Options.bWantInstanceColors;
 	UE::Geometry::FDynamicMesh3 NewMesh;
 	FText ErrorMessage;
 	bool bSuccess = UE::Conversion::SceneComponentToDynamicMesh(Component, ToMeshOptions, bTransformToWorld, NewMesh, LocalToWorld, ErrorMessage);
diff --git a/Engine/Plugins/Runtime/GeometryScripting/Source/GeometryScriptingCore/Public/GeometryScript/SceneUtilityFunctions.h b/Engine/Plugins/Runtime/GeometryScripting/Source/GeometryScriptingCore/Public/GeometryScript/SceneUtilityFunctions.h
--- a/Engine/Plugins/Runtime/GeometryScripting/Source/GeometryScriptingCore/Public/GeometryScript/SceneUtilityFunctions.h
+++ b/Engine/Plugins/Runtime/GeometryScripting/Source/GeometryScriptingCore/Public/GeometryScript/SceneUtilityFunctions.h
@@ -23,6 +23,10 @@ public:
 	UPROPERTY(EditAnywhere, BlueprintReadWrite, Category = Options)
 	bool bWantTangents = true;
 
+	// Whether to request per-instance vertex colors (where applicable; applies to RenderData LODs of Static Mesh components)
+	UPROPERTY(EditAnywhere, BlueprintReadWrite, Category = Options)
+	bool bWantInstanceColors = false;
+
 	UPROPERTY(EditAnywhere, BlueprintReadWrite, Category = Options)
 	FGeometryScriptMeshReadLOD RequestedLOD = FGeometryScriptMeshReadLOD();
 };
diff --git a/Engine/Plugins/Runtime/MeshModelingToolset/Source/ModelingComponents/Private/ConversionUtils/SceneComponentToDynamicMesh.cpp b/Engine/Plugins/Runtime/MeshModelingToolset/Source/ModelingComponents/Private/ConversionUtils/SceneComponentToDynamicMesh.cpp
--- a/Engine/Plugins/Runtime/MeshModelingToolset/Source/ModelingComponents/Private/ConversionUtils/SceneComponentToDynamicMesh.cpp
+++ b/Engine/Plugins/Runtime/MeshModelingToolset/Source/ModelingComponents/Private/ConversionUtils/SceneComponentToDynamicMesh.cpp
@@ -14,6 +14,7 @@
 #include "Components/BrushComponent.h"
 #include "Components/DynamicMeshComponent.h"
 #include "Components/SplineMeshComponent.h"
+#include "StaticMeshComponentLODInfo.h"
 #include "ConversionUtils/VolumeToDynamicMesh.h"
 #include "ConversionUtils/SkinnedMeshToDynamicMesh.h"
 #include "ConversionUtils/SplineComponentDeformDynamicMesh.h"
@@ -90,6 +91,9 @@ namespace Private::ConversionHelper
 
 		// Whether to scale the copied mesh by the Build Setting's 'Build Scale'. Note: This is considered separately from the 'Apply Build Settings' option.
 		bool bUseBuildScale = true;
+
+		// Whether to request the vertex colors of the component instancing the static mesh, rather than the static mesh asset
+		bool bRequestInstanceVertexColors = true;
 	};
 
 	static bool CopyMeshFromStaticMesh_SourceData(
@@ -210,6 +214,7 @@ namespace Private::ConversionHelper
 
 	static bool CopyMeshFromStaticMesh_RenderData(
 		UStaticMesh* FromStaticMeshAsset,
+		UStaticMeshComponent* StaticMeshComponent,
 		FStaticMeshConversionOptions AssetOptions,
 		EMeshLODType LODType,
 		int32 LODIndex,
@@ -265,13 +270,30 @@ namespace Private::ConversionHelper
 #endif
 
 		FStaticMeshLODResourcesToDynamicMesh Converter;
-		Converter.Convert(LODResources, ConvertOptions, OutMesh);
+		if (AssetOptions.bRequestInstanceVertexColors && StaticMeshComponent && StaticMeshComponent->LODData.IsValidIndex(UseLODIndex))
+		{
+			FStaticMeshComponentLODInfo* InstanceMeshLODInfo = &StaticMeshComponent->LODData[UseLODIndex];
+			const bool bValidInstanceData = InstanceMeshLODInfo
+				&& InstanceMeshLODInfo->OverrideVertexColors
+				&& InstanceMeshLODInfo->OverrideVertexColors->GetAllowCPUAccess()
+				&& InstanceMeshLODInfo->OverrideVertexColors->GetNumVertices() == LODResources->GetNumVertices();
+			Converter.Convert(LODResources, ConvertOptions, OutMesh, bValidInstanceData,
+				[InstanceMeshLODInfo](int32 LODVID)
+				{
+					return InstanceMeshLODInfo->OverrideVertexColors->VertexColor(LODVID);
+				});
+		}
+		else
+		{
+			Converter.Convert(LODResources, ConvertOptions, OutMesh);
+		}
 
 		return true;
 	}
 
 	static bool CopyMeshFromStaticMesh(
 		UStaticMesh* FromStaticMeshAsset,
+		UStaticMeshComponent* StaticMeshComponent,
 		FStaticMeshConversionOptions AssetOptions,
 		EMeshLODType LODType,
 		int32 LODIndex,
@@ -322,7 +344,7 @@ namespace Private::ConversionHelper
 
 		if (LODType == EMeshLODType::RenderData)
 		{
-			return CopyMeshFromStaticMesh_RenderData(FromStaticMeshAsset, AssetOptions, LODType, LODIndex, OutMesh, OutErrorMessage);
+			return CopyMeshFromStaticMesh_RenderData(FromStaticMeshAsset, StaticMeshComponent, AssetOptions, LODType, LODIndex, OutMesh, OutErrorMessage);
 		}
 		else
 		{
@@ -341,6 +363,11 @@ bool SceneComponentToDynamicMesh(USceneComponent* Component, const FToMeshOption
 	bool bSuccess = false;
 	OutMesh.Clear();
 
+	if (!Component)
+	{
+		OutErrorMessage = LOCTEXT("CopyMeshFromComponent_NullComponent", "Scene Component is null");
+		return false;
+	}
 	OutLocalToWorld = Component->GetComponentTransform();
 
 	auto GetPrimitiveComponentMaterials = [](UPrimitiveComponent* PrimComp, TArray<UMaterialInterface*> Materials)
@@ -411,8 +438,9 @@ bool SceneComponentToDynamicMesh(USceneComponent* Component, const FToMeshOption
 			Private::ConversionHelper::FStaticMeshConversionOptions AssetOptions;
 			AssetOptions.bApplyBuildSettings = (Options.bWantNormals || Options.bWantTangents);
 			AssetOptions.bRequestTangents = Options.bWantTangents;
+			AssetOptions.bRequestInstanceVertexColors = Options.bWantInstanceColors;
 			bSuccess = Private::ConversionHelper::CopyMeshFromStaticMesh(
-				StaticMesh, AssetOptions, Options.LODType, Options.LODIndex, Options.bUseClosestLOD, OutMesh, OutErrorMessage);
+				StaticMesh, SplineMeshComponent, AssetOptions, Options.LODType, Options.LODIndex, Options.bUseClosestLOD, OutMesh, OutErrorMessage);
 
 			// deform the dynamic mesh and its tangent space with the spline
 			if (bSuccess)
@@ -444,8 +472,9 @@ bool SceneComponentToDynamicMesh(USceneComponent* Component, const FToMeshOption
 			Private::ConversionHelper::FStaticMeshConversionOptions AssetOptions;
 			AssetOptions.bApplyBuildSettings = (Options.bWantNormals || Options.bWantTangents);
 			AssetOptions.bRequestTangents = Options.bWantTangents;
+			AssetOptions.bRequestInstanceVertexColors = Options.bWantInstanceColors;
 			bSuccess = Private::ConversionHelper::CopyMeshFromStaticMesh(
-				StaticMesh, AssetOptions, Options.LODType, Options.LODIndex, Options.bUseClosestLOD, OutMesh, OutErrorMessage);
+				StaticMesh, StaticMeshComponent, AssetOptions, Options.LODType, Options.LODIndex, Options.bUseClosestLOD, OutMesh, OutErrorMessage);
 
 			// if we have an ISMC, append instances
 			if (UInstancedStaticMeshComponent* ISMComponent = Cast<UInstancedStaticMeshComponent>(StaticMeshComponent))
diff --git a/Engine/Plugins/Runtime/MeshModelingToolset/Source/ModelingComponents/Public/ConversionUtils/SceneComponentToDynamicMesh.h b/Engine/Plugins/Runtime/MeshModelingToolset/Source/ModelingComponents/Public/ConversionUtils/SceneComponentToDynamicMesh.h
--- a/Engine/Plugins/Runtime/MeshModelingToolset/Source/ModelingComponents/Public/ConversionUtils/SceneComponentToDynamicMesh.h
+++ b/Engine/Plugins/Runtime/MeshModelingToolset/Source/ModelingComponents/Public/ConversionUtils/SceneComponentToDynamicMesh.h
@@ -56,6 +56,9 @@ namespace Conversion {
 
 		bool bWantNormals = true;
 		bool bWantTangents = true;
+		// Whether to request per-instance vertex colors rather than asset colors
+		// Note this is only supported for RenderData LODs of Static Mesh components
+		bool bWantInstanceColors = false;
 	};
 
 	/**
diff --git a/Engine/Source/Runtime/MeshConversionEngineTypes/Private/StaticMeshLODResourcesToDynamicMesh.cpp b/Engine/Source/Runtime/MeshConversionEngineTypes/Private/StaticMeshLODResourcesToDynamicMesh.cpp
--- a/Engine/Source/Runtime/MeshConversionEngineTypes/Private/StaticMeshLODResourcesToDynamicMesh.cpp
+++ b/Engine/Source/Runtime/MeshConversionEngineTypes/Private/StaticMeshLODResourcesToDynamicMesh.cpp
@@ -13,7 +13,9 @@ using namespace UE::Geometry;
 bool FStaticMeshLODResourcesToDynamicMesh::Convert(
 	const FStaticMeshLODResources* StaticMeshResources,
 	const ConversionOptions& Options,
-	FDynamicMesh3& OutputMesh)
+	FDynamicMesh3& OutputMesh,
+	bool bHasVertexColors,
+	TFunctionRef<FColor(int32)> GetVertexColorFromLODVertexIndex)
 {
 	if (!ensureMsgf(StaticMeshResources && StaticMeshResources->VertexBuffers.StaticMeshVertexBuffer.GetAllowCPUAccess(), TEXT("bAllowCPUAccess must be set to true for StaticMeshes before calling FStaticMeshLODResourcesToDynamicMesh::Convert(), otherwise the mesh geometry data isn't accessible!")))
 	{
@@ -238,25 +240,44 @@ bool FStaticMeshLODResourcesToDynamicMesh::Convert(
 	}
 
 	// copy overlay colors
-	if ( Adapter.HasColors() && Options.bWantVertexColors )
+	if (bHasVertexColors && Options.bWantVertexColors)
 	{
 		OutputMesh.Attributes()->EnablePrimaryColors();
 		FDynamicMeshColorOverlay* Colors = OutputMesh.Attributes()->PrimaryColors();
+
+		const int32 DstVertexCount = ToSrcVID.Num();
+		for (int32 DstVertID = 0; DstVertID < DstVertexCount; ++DstVertID)
+		{
+			const int32 SrcVID = ToSrcVID[DstVertID];
+			FColor C = GetVertexColorFromLODVertexIndex(SrcVID);
+			int32 ElemID = Colors->AppendElement(C.ReinterpretAsLinear());
+			check(ElemID == DstVertID);
+		}
+
 		for (int32 SrcTriID = 0; SrcTriID < SrcTriangleCount; ++SrcTriID)
 		{
 			const int32 DstTriID = ToDstTriID[SrcTriID];
 			if (DstTriID != FDynamicMesh3::InvalidID)
 			{
-				FColor C1, C2, C3;
-				Adapter.GetTriColors(SrcTriID, C1, C2, C3);
-				// [TODO] should we be doing RGBA conversion here?
-				int32 a = Colors->AppendElement( C1.ReinterpretAsLinear()  );
-				int32 b = Colors->AppendElement( C2.ReinterpretAsLinear() );
-				int32 c = Colors->AppendElement( C3.ReinterpretAsLinear() );
-				Colors->SetTriangle(DstTriID, FIndex3i(a, b, c));
+				FIndex3i Tri = OutputMesh.GetTriangle(DstTriID);
+				Colors->SetTriangle(DstTriID, FIndex3i(Tri.A, Tri.B, Tri.C));
 			}
 		}
 	}
 
 	return true;
 }
+
+
+bool FStaticMeshLODResourcesToDynamicMesh::Convert(
+	const FStaticMeshLODResources* StaticMeshResources,
+	const ConversionOptions& Options,
+	FDynamicMesh3& OutputMesh)
+{
+	bool bHasVertexColors = StaticMeshResources && StaticMeshResources->VertexBuffers.ColorVertexBuffer.GetAllowCPUAccess();
+	return Convert(StaticMeshResources, Options, OutputMesh, bHasVertexColors, 
+		[StaticMeshResources](int32 VID)
+		{
+			return StaticMeshResources->VertexBuffers.ColorVertexBuffer.VertexColor(VID);
+		});
+}
diff --git a/Engine/Source/Runtime/MeshConversionEngineTypes/Public/StaticMeshLODResourcesToDynamicMesh.h b/Engine/Source/Runtime/MeshConversionEngineTypes/Public/StaticMeshLODResourcesToDynamicMesh.h
--- a/Engine/Source/Runtime/MeshConversionEngineTypes/Public/StaticMeshLODResourcesToDynamicMesh.h
+++ b/Engine/Source/Runtime/MeshConversionEngineTypes/Public/StaticMeshLODResourcesToDynamicMesh.h
@@ -36,9 +36,18 @@ public:
 		FVector3d BuildScale = FVector3d::One();
 	};
 
+	// @param bHasVertexColors					Whether to add vertex colors to the output (if available)
+	// @param GetVertexColorFromLODVertexIndex	Function for getting the vertex color of a given vertex index
 	static MESHCONVERSIONENGINETYPES_API bool Convert(
 		const FStaticMeshLODResources* StaticMeshResources, 
 		const ConversionOptions& Options, 
+		FDynamicMesh3& OutputMesh,
+		bool bHasVertexColors,
+		TFunctionRef<FColor(int32)> GetVertexColorFromLODVertexIndex);
+
+	static MESHCONVERSIONENGINETYPES_API bool Convert(
+		const FStaticMeshLODResources* StaticMeshResources,
+		const ConversionOptions& Options,
 		FDynamicMesh3& OutputMesh);
 
 };
