ModelingTools: Minor comment/name fixes from extrude edge review.

diff --git a/Engine/Plugins/Runtime/MeshModelingToolset/Source/MeshModelingTools/Private/Operations/ExtrudeBoundaryEdges.cpp b/Engine/Plugins/Runtime/MeshModelingToolset/Source/MeshModelingTools/Private/Operations/ExtrudeBoundaryEdges.cpp
--- a/Engine/Plugins/Runtime/MeshModelingToolset/Source/MeshModelingTools/Private/Operations/ExtrudeBoundaryEdges.cpp
+++ b/Engine/Plugins/Runtime/MeshModelingToolset/Source/MeshModelingTools/Private/Operations/ExtrudeBoundaryEdges.cpp
@@ -178,10 +178,10 @@ namespace ExtrudeBoundaryEdgesLocals
 		if (ScalingLimit > 1)
 		{
 			FVector3d VectorToCompareAgainst = OutwardVector1.IsZero() ? OutwardVector2 : OutwardVector1;
-			double cosTheta = ExtrudeFrame.Frame.X().Dot(VectorToCompareAgainst);
-			if (FMath::Abs(cosTheta) > KINDA_SMALL_NUMBER)
+			double CosTheta = ExtrudeFrame.Frame.X().Dot(VectorToCompareAgainst);
+			if (FMath::Abs(CosTheta) > KINDA_SMALL_NUMBER)
 			{
-				ExtrudeFrame.Scaling.X = FMath::Min(1 / FMath::Abs(cosTheta), ScalingLimit);
+				ExtrudeFrame.Scaling.X = FMath::Min(1 / FMath::Abs(CosTheta), ScalingLimit);
 			}
 		}
 
diff --git a/Engine/Plugins/Runtime/MeshModelingToolset/Source/MeshModelingTools/Private/Operations/ExtrudeBoundaryEdges.h b/Engine/Plugins/Runtime/MeshModelingToolset/Source/MeshModelingTools/Private/Operations/ExtrudeBoundaryEdges.h
--- a/Engine/Plugins/Runtime/MeshModelingToolset/Source/MeshModelingTools/Private/Operations/ExtrudeBoundaryEdges.h
+++ b/Engine/Plugins/Runtime/MeshModelingToolset/Source/MeshModelingTools/Private/Operations/ExtrudeBoundaryEdges.h
@@ -57,7 +57,7 @@ public:
 	bool bUsePerVertexExtrudeFrames = true;
 
 	/**
-	 * Function queried for new verex positions. ExtrudeFrame origin is Position, unless it is not initialized
+	 * Function queried for new vertex positions. ExtrudeFrame origin is Position, unless it is not initialized
 	 *  due to bUsePerVertexExtrudeFrames being false.
 	 * The default given here assumes that bUsePerVertexExtrudeFrames is true, and extrudes along the X axis of the extrude frame.
 	 */
