clamp bevel distance ranges to avoid weird results for negative, zero, and very large bevel distances

#changelist validated
#virtualized

diff --git a/Engine/Plugins/Runtime/MeshModelingToolset/Source/MeshModelingTools/Private/ToolActivities/PolyEditBevelEdgeActivity.h b/Engine/Plugins/Runtime/MeshModelingToolset/Source/MeshModelingTools/Private/ToolActivities/PolyEditBevelEdgeActivity.h
--- a/Engine/Plugins/Runtime/MeshModelingToolset/Source/MeshModelingTools/Private/ToolActivities/PolyEditBevelEdgeActivity.h
+++ b/Engine/Plugins/Runtime/MeshModelingToolset/Source/MeshModelingTools/Private/ToolActivities/PolyEditBevelEdgeActivity.h
@@ -25,8 +25,8 @@ class MESHMODELINGTOOLS_API UPolyEditBevelEdgeProperties : public UInteractiveTo
 	GENERATED_BODY()
 
 public:
-	/** Distance that each beveled mesh edge is inset from it's initial position*/
-	UPROPERTY(EditAnywhere, Category = Bevel)
+	/** Distance that each beveled mesh edge is inset from its initial position */
+	UPROPERTY(EditAnywhere, Category = Bevel, meta = (UIMin = ".001", UIMax = "100", ClampMin = "0.0001", ClampMax = "10000", SliderExponent = "3"))
 	double BevelDistance = 4.0;
 
 	/** Number of edge loops added along the bevel faces */
