fix polyedit's edge loop selection modifying the edge selection while it iterates through it

#changelist validated
#virtualized

diff --git a/Engine/Plugins/Runtime/MeshModelingToolset/Source/ModelingComponents/Private/Selection/GroupTopologySelector.cpp b/Engine/Plugins/Runtime/MeshModelingToolset/Source/ModelingComponents/Private/Selection/GroupTopologySelector.cpp
--- a/Engine/Plugins/Runtime/MeshModelingToolset/Source/ModelingComponents/Private/Selection/GroupTopologySelector.cpp
+++ b/Engine/Plugins/Runtime/MeshModelingToolset/Source/ModelingComponents/Private/Selection/GroupTopologySelector.cpp
@@ -232,6 +232,7 @@ bool FGroupTopologySelector::ExpandSelectionByEdgeLoops(FGroupTopologySelection&
 {
 	TSet<int32>& EdgeSet = Selection.SelectedEdgeIDs;
 	int32 OriginalNumEdges = Selection.SelectedEdgeIDs.Num();
+	TSet<int32> ExpandedEdgeSet = EdgeSet; // make a copy of the edge set, to add to during iteration
 	for (int32 Eid : Selection.SelectedEdgeIDs)
 	{
 		const FIndex2i EndpointCorners = GroupTopologyUtils.GetEdgeEndpointCorners(Eid);
@@ -241,9 +242,10 @@ bool FGroupTopologySelector::ExpandSelectionByEdgeLoops(FGroupTopologySelection&
 		}
 
 		// Go forward and backward adding edges
-		GroupTopologyUtils.AddNewEdgeLoopEdgesFromCorner(Eid, EndpointCorners[0], EdgeSet);
-		GroupTopologyUtils.AddNewEdgeLoopEdgesFromCorner(Eid, EndpointCorners[1], EdgeSet);
+		GroupTopologyUtils.AddNewEdgeLoopEdgesFromCorner(Eid, EndpointCorners[0], ExpandedEdgeSet);
+		GroupTopologyUtils.AddNewEdgeLoopEdgesFromCorner(Eid, EndpointCorners[1], ExpandedEdgeSet);
 	}
+	EdgeSet = MoveTemp(ExpandedEdgeSet); // update the selection edges
 
 	return EdgeSet.Num() > OriginalNumEdges;
 }
@@ -252,7 +254,7 @@ bool FGroupTopologySelector::ExpandSelectionByBoundaryLoops(FGroupTopologySelect
 {
 	TSet<int32>& EdgeSet = Selection.SelectedEdgeIDs;
 	int32 OriginalNumEdges = Selection.SelectedEdgeIDs.Num();
-
+	TSet<int32> ExpandedEdgeSet = EdgeSet; // make a copy of the edge set, to add to during iteration
 	for (int32 Eid : Selection.SelectedEdgeIDs)
 	{
 		if (!GroupTopologyUtils.GroupTopology->IsBoundaryEdge(Eid))
@@ -267,9 +269,10 @@ bool FGroupTopologySelector::ExpandSelectionByBoundaryLoops(FGroupTopologySelect
 		}
 
 		// Go forward and backward adding edges
-		GroupTopologyUtils.AddNewBoundaryLoopEdges(Eid, EndpointCorners[0], EdgeSet);
-		GroupTopologyUtils.AddNewBoundaryLoopEdges(Eid, EndpointCorners[1], EdgeSet);
+		GroupTopologyUtils.AddNewBoundaryLoopEdges(Eid, EndpointCorners[0], ExpandedEdgeSet);
+		GroupTopologyUtils.AddNewBoundaryLoopEdges(Eid, EndpointCorners[1], ExpandedEdgeSet);
 	}
+	EdgeSet = MoveTemp(ExpandedEdgeSet); // update the selection edges
 
 	return EdgeSet.Num() > OriginalNumEdges;
 }
@@ -278,6 +281,7 @@ bool FGroupTopologySelector::ExpandSelectionByEdgeRings(FGroupTopologySelection&
 {
 	TSet<int32>& EdgeSet = Selection.SelectedEdgeIDs;
 	int32 OriginalNumEdges = Selection.SelectedEdgeIDs.Num();
+	TSet<int32> ExpandedEdgeSet = EdgeSet; // make a copy of the edge set, to add to during iteration
 	for (int32 Eid : Selection.SelectedEdgeIDs)
 	{
 		const FIndex2i EdgeGroups = GroupTopologyUtils.GetEdgeGroups(Eid);
@@ -285,13 +289,14 @@ bool FGroupTopologySelector::ExpandSelectionByEdgeRings(FGroupTopologySelection&
 		// Go forward and backward adding edges
 		if (EdgeGroups[0] != IndexConstants::InvalidID)
 		{
-			GroupTopologyUtils.AddNewEdgeRingEdges(Eid, EdgeGroups[0], EdgeSet);
+			GroupTopologyUtils.AddNewEdgeRingEdges(Eid, EdgeGroups[0], ExpandedEdgeSet);
 		}
 		if (EdgeGroups[0] != IndexConstants::InvalidID)
 		{
-			GroupTopologyUtils.AddNewEdgeRingEdges(Eid, EdgeGroups[1], EdgeSet);
+			GroupTopologyUtils.AddNewEdgeRingEdges(Eid, EdgeGroups[1], ExpandedEdgeSet);
 		}
 	}
+	EdgeSet = MoveTemp(ExpandedEdgeSet); // update the selection edges
 
 	return EdgeSet.Num() > OriginalNumEdges;
 }
