MeshModelingToolset: Add Bridge action to TriEdit

#changelist validated
#virtualized

diff --git a/Engine/Plugins/Runtime/MeshModelingToolset/Source/MeshModelingTools/Public/EditMeshPolygonsTool.h b/Engine/Plugins/Runtime/MeshModelingToolset/Source/MeshModelingTools/Public/EditMeshPolygonsTool.h
--- a/Engine/Plugins/Runtime/MeshModelingToolset/Source/MeshModelingTools/Public/EditMeshPolygonsTool.h
+++ b/Engine/Plugins/Runtime/MeshModelingToolset/Source/MeshModelingTools/Public/EditMeshPolygonsTool.h
@@ -446,23 +446,27 @@ public:
 	void Weld() { PostAction(EEditMeshPolygonsToolActions::WeldEdges); }
 
 	/** Fill the adjacent hole for any selected boundary edges */
-	UFUNCTION(CallInEditor, Category = EdgeEdits, meta = (DisplayName = "Fill Hole", DisplayPriority = 1))
+	UFUNCTION(CallInEditor, Category = EdgeEdits, meta = (DisplayName = "Fill Hole", DisplayPriority = 2))
 	void FillHole() { PostAction(EEditMeshPolygonsToolActions::FillHole); }
 
+	/** Create a new face that connects the selected edges */
+	UFUNCTION(CallInEditor, Category = EdgeEdits, meta = (DisplayName = "Bridge", DisplayPriority = 3))
+	void Bridge() { PostAction(EEditMeshPolygonsToolActions::BridgeEdges); }
+	
 	/** Duplicate and move boundary vertices outwards and connect them to the original boundary to create new faces. */
-	UFUNCTION(CallInEditor, Category = EdgeEdits, meta = (DisplayName = "Extrude", DisplayPriority = 1))
+	UFUNCTION(CallInEditor, Category = EdgeEdits, meta = (DisplayName = "Extrude", DisplayPriority = 4))
 	void Extrude() { PostAction(EEditMeshPolygonsToolActions::ExtrudeEdges); }
 
 	/** Collapse the selected edges, deleting the attached triangles and merging its two vertices into one */
-	UFUNCTION(CallInEditor, Category = EdgeEdits, meta = (DisplayName = "Collapse", DisplayPriority = 1))
+	UFUNCTION(CallInEditor, Category = EdgeEdits, meta = (DisplayName = "Collapse", DisplayPriority = 5))
 	void Collapse() { PostAction(EEditMeshPolygonsToolActions::CollapseSingleEdge); }
 
 	/** Flip the selected (non-border, non-seam) edges, replacing them with new edges in the crossing direction */
-	UFUNCTION(CallInEditor, Category = EdgeEdits, meta = (DisplayName = "Flip", DisplayPriority = 1))
+	UFUNCTION(CallInEditor, Category = EdgeEdits, meta = (DisplayName = "Flip", DisplayPriority = 6))
 	void Flip() { PostAction(EEditMeshPolygonsToolActions::FlipSingleEdge); }
 
 	/** Split the selected edges, inserting a new vertex at each edge midpoint */
-	UFUNCTION(CallInEditor, Category = EdgeEdits, meta = (DisplayName = "Split", DisplayPriority = 1))
+	UFUNCTION(CallInEditor, Category = EdgeEdits, meta = (DisplayName = "Split", DisplayPriority = 7))
 	void Split() { PostAction(EEditMeshPolygonsToolActions::SplitSingleEdge); }
 };
 
