Skin Bind Tool: Fix binding computation when there's only one influence per bone

#rnx

diff --git a/Engine/Plugins/Runtime/MeshModelingToolset/Source/ModelingOperators/Private/SkinningOps/SkinBindingOp.cpp b/Engine/Plugins/Runtime/MeshModelingToolset/Source/ModelingOperators/Private/SkinningOps/SkinBindingOp.cpp
--- a/Engine/Plugins/Runtime/MeshModelingToolset/Source/ModelingOperators/Private/SkinningOps/SkinBindingOp.cpp
+++ b/Engine/Plugins/Runtime/MeshModelingToolset/Source/ModelingOperators/Private/SkinningOps/SkinBindingOp.cpp
@@ -217,6 +217,11 @@ struct FTransformHierarchyQuery
 		}		
 	}
 
+	bool IsEndEffector(const int32 InBoneIndex) const
+	{
+		return BoneFans[InBoneIndex].TipsPos.IsEmpty();
+	}
+
 	float GetDistanceToBoneFan(const int32 InBoneIndex, const FVector& InPoint) const
 	{
 		return BoneFans[InBoneIndex].GetDistance(InPoint);
@@ -298,13 +303,35 @@ namespace
 	{
 		TArray<TPair<FBoneIndexType, float>> RawBoneWeights;
 		TArray<AnimationCore::FBoneWeight> BoneWeights;
+
+		void NormalizeWeightsAndLimitCount(const float InTotalWeight, const int32 InMaxWeights)
+		{
+			// Normalize
+			for (TPair<FBoneIndexType, float> &BoneWeight: RawBoneWeights)
+			{
+				BoneWeight.Value /= InTotalWeight;
+			}
+		
+			RawBoneWeights.Sort([](const TPair<FBoneIndexType, float> &A, const TPair<FBoneIndexType, float> &B)
+			{
+				return A.Value > B.Value;
+			});
+
+			BoneWeights.Reset(InMaxWeights);
+			for (int32 BoneIndex = 0; BoneIndex < FMath::Min(InMaxWeights, RawBoneWeights.Num()); BoneIndex++)
+			{
+				const TPair<FBoneIndexType, float>& BoneWeight = RawBoneWeights[BoneIndex];
+				BoneWeights.Add(UE::AnimationCore::FBoneWeight(BoneWeight.Key, BoneWeight.Value));
+			}
+			
+		}
 	};
 
 	float ComputeWeightStiffness(const float InWeight, const float InStiffness)
 	{
 		return (1.0f - InStiffness) * InWeight + InStiffness * InWeight * InWeight;
 	}
-	
+
 }
 
 
@@ -394,7 +421,7 @@ void FSkinBindingOp::CreateSkinWeights_DirectDistance(
 	FDynamicMesh3& InMesh,
 	float InStiffness,
 	const AnimationCore::FBoneWeightsSettings& InSettings
-	)
+	) const
 {
 	using namespace AnimationCore;
 
@@ -413,11 +440,25 @@ void FSkinBindingOp::CreateSkinWeights_DirectDistance(
 
 		FCreateSkinWeights_Closest_WorkData &WorkData = FCreateSkinWeights_Closest_WorkData::Get();
 		
-		WorkData.RawBoneWeights.Reset(TransformHierarchy.Num());
-
+		if (MaxInfluences > 1)
+		{
+			WorkData.RawBoneWeights.Reset(TransformHierarchy.Num());
+		}
+		else
+		{
+			WorkData.RawBoneWeights.Init(MakeTuple(0, 1.0), 1);
+		}
+		
 		float TotalWeight = 0.0f;
 		for (int32 BoneIndex = 0; BoneIndex < TransformHierarchy.Num(); BoneIndex++)
 		{
+			// For single influences, avoid end effectors to avoid closest-distance fighting with the tips of their
+			// parent bones.
+			if (MaxInfluences == 1 && Skeleton.IsEndEffector(BoneIndex))
+			{
+				continue;
+			}
+			
 			// Normalize the distance by the diagonal size of the bbox to maintain scale invariance.
 			float Weight = Skeleton.GetDistanceToBoneFan(BoneIndex, Pos) / DiagBounds;
 
@@ -428,27 +469,21 @@ void FSkinBindingOp::CreateSkinWeights_DirectDistance(
 			// Compute the actual weight, factoring in the stiffness value. W = (1/S(D))^2
 			// Where S(x) is the stiffness function.
 			Weight = FMath::Square(1.0f / ComputeWeightStiffness(Weight, InStiffness));
-			TotalWeight += Weight;
-			WorkData.RawBoneWeights.Add(MakeTuple(static_cast<FBoneIndexType>(BoneIndex), Weight));
-		}
 
-		// Normalize
-		for (TPair<FBoneIndexType, float> &BoneWeight: WorkData.RawBoneWeights)
-		{
-			BoneWeight.Value /= TotalWeight;
+			if (MaxInfluences > 1)
+			{
+				TotalWeight += Weight;
+				WorkData.RawBoneWeights.Add(MakeTuple(static_cast<FBoneIndexType>(BoneIndex), Weight));
+			}
+			else if (Weight > TotalWeight)
+			{
+				// For single influences, we only care about the strongest influence.
+				TotalWeight = Weight;
+				WorkData.RawBoneWeights[0] = MakeTuple(static_cast<FBoneIndexType>(BoneIndex), Weight);
+			}
 		}
-			
-		WorkData.RawBoneWeights.Sort([](const TPair<FBoneIndexType, float> &A, const TPair<FBoneIndexType, float> &B)
-		{
-			return A.Value > B.Value;
-		});
 
-		WorkData.BoneWeights.Reset(InSettings.GetMaxWeightCount());
-		for (int32 BoneIndex = 0; BoneIndex < FMath::Min(InSettings.GetMaxWeightCount(), WorkData.RawBoneWeights.Num()); BoneIndex++)
-		{
-			const TPair<FBoneIndexType, float>& BoneWeight = WorkData.RawBoneWeights[BoneIndex];
-			WorkData.BoneWeights.Add(FBoneWeight(BoneWeight.Key, BoneWeight.Value));
-		}
+		WorkData.NormalizeWeightsAndLimitCount(TotalWeight, InSettings.GetMaxWeightCount());
 
 		SkinWeights->SetValue(VertexIdx, FBoneWeights::Create(WorkData.BoneWeights, InSettings));
 	});
@@ -459,7 +494,7 @@ void FSkinBindingOp::CreateSkinWeights_GeodesicVoxel(
 	FDynamicMesh3& InMesh,
 	float InStiffness,
 	const AnimationCore::FBoneWeightsSettings& InSettings
-	)
+	) const
 {
 	using namespace AnimationCore;
 
@@ -483,6 +518,13 @@ void FSkinBindingOp::CreateSkinWeights_GeodesicVoxel(
 	const FVector3i Dimensions = Occupancy.GetOccupancyStateGrid().GetDimensions();
 		
 	ParallelFor(TransformHierarchy.Num(), [&](int32 BoneIndex) {
+		// For single influences, avoid end effectors to avoid closest-distance fighting with the tips of their
+		// parent bones.
+		if (MaxInfluences == 1 && Skeleton.IsEndEffector(BoneIndex))
+		{
+			return;
+		}
+		
 		TFIFOQueue<FVector3i> WorkingSet;  
 		FDenseGrid3f BoneDistance(Dimensions.X, Dimensions.Y, Dimensions.Z, DiagonalBounds);
 
@@ -584,34 +626,41 @@ void FSkinBindingOp::CreateSkinWeights_GeodesicVoxel(
 	{
 		FCreateSkinWeights_Closest_WorkData &WorkData = FCreateSkinWeights_Closest_WorkData::Get();
 			
-		WorkData.RawBoneWeights.Reset(TransformHierarchy.Num());
-
-		float TotalWeight = 0.0f;
-		for (int32 BoneIndex = 0; BoneIndex < TransformHierarchy.Num(); BoneIndex++)
+		if (MaxInfluences > 1)
 		{
-			const float Weight = Weights[VertexIdx * TransformHierarchy.Num() + BoneIndex];
-			TotalWeight += Weight;
-			WorkData.RawBoneWeights.Add(MakeTuple(static_cast<FBoneIndexType>(BoneIndex), Weight));
+			WorkData.RawBoneWeights.Reset(TransformHierarchy.Num());
 		}
-
-		// Normalize
-		for (TPair<FBoneIndexType, float> &BoneWeight: WorkData.RawBoneWeights)
+		else
 		{
-			BoneWeight.Value /= TotalWeight;
+			WorkData.RawBoneWeights.Init(MakeTuple(0, 1.0), 1);
 		}
-				
-		WorkData.RawBoneWeights.Sort([](const TPair<FBoneIndexType, float> &A, const TPair<FBoneIndexType, float> &B)
-		{
-			return A.Value > B.Value;
-		});
 
-		WorkData.BoneWeights.Reset(InSettings.GetMaxWeightCount());
-		for (int32 BoneIndex = 0; BoneIndex < FMath::Min(InSettings.GetMaxWeightCount(), WorkData.RawBoneWeights.Num()); BoneIndex++)
+		float TotalWeight = 0.0f;
+		for (int32 BoneIndex = 0; BoneIndex < TransformHierarchy.Num(); BoneIndex++)
 		{
-			const TPair<FBoneIndexType, float>& BoneWeight = WorkData.RawBoneWeights[BoneIndex];
-			WorkData.BoneWeights.Add(FBoneWeight(BoneWeight.Key, BoneWeight.Value));
+			// Ignore end effectors for single influence computation, since those we not
+			// calculated.
+			if (MaxInfluences == 1 && Skeleton.IsEndEffector(BoneIndex))
+			{
+				continue;
+			}
+			
+			const float Weight = Weights[VertexIdx * TransformHierarchy.Num() + BoneIndex];
+			
+			if (MaxInfluences > 1)
+			{
+				TotalWeight += Weight;
+				WorkData.RawBoneWeights.Add(MakeTuple(static_cast<FBoneIndexType>(BoneIndex), Weight));
+			}
+			else if (Weight > TotalWeight)
+			{
+				TotalWeight = Weight;
+				WorkData.RawBoneWeights[0] = MakeTuple(static_cast<FBoneIndexType>(BoneIndex), Weight);
+			}
 		}
 
+		WorkData.NormalizeWeightsAndLimitCount(TotalWeight, InSettings.GetMaxWeightCount());
+
 		SkinWeights->SetValue(VertexIdx, FBoneWeights::Create(WorkData.BoneWeights, InSettings));
 	});
 }
diff --git a/Engine/Plugins/Runtime/MeshModelingToolset/Source/ModelingOperators/Public/SkinningOps/SkinBindingOp.h b/Engine/Plugins/Runtime/MeshModelingToolset/Source/ModelingOperators/Public/SkinningOps/SkinBindingOp.h
--- a/Engine/Plugins/Runtime/MeshModelingToolset/Source/ModelingOperators/Public/SkinningOps/SkinBindingOp.h
+++ b/Engine/Plugins/Runtime/MeshModelingToolset/Source/ModelingOperators/Public/SkinningOps/SkinBindingOp.h
@@ -62,14 +62,14 @@ private:
 		FDynamicMesh3& InMesh,
 		float InStiffness,
 		const AnimationCore::FBoneWeightsSettings& InSettings
-		);
+		) const;
 
 
 	void CreateSkinWeights_GeodesicVoxel(
 		FDynamicMesh3& InMesh,
 		float InStiffness,
 		const AnimationCore::FBoneWeightsSettings& InSettings
-		);
+		) const; 
 };
 
 
