MeshModelingToolset: Fix crash when hiding actors deleted by tools with Delete Inputs option.

Our approach to deleting actors was apparently not sufficient to notify various editor systems of the deletion. This CL replaces AActor::Destroy with calls to UUnrealEdEngine::DeleteActors and UWorld::DeleteActor (depending on the context)

diff --git a/Engine/Plugins/Runtime/MeshModelingToolset/Source/ModelingComponents/ModelingComponents.Build.cs b/Engine/Plugins/Runtime/MeshModelingToolset/Source/ModelingComponents/ModelingComponents.Build.cs
--- a/Engine/Plugins/Runtime/MeshModelingToolset/Source/ModelingComponents/ModelingComponents.Build.cs
+++ b/Engine/Plugins/Runtime/MeshModelingToolset/Source/ModelingComponents/ModelingComponents.Build.cs
@@ -43,6 +43,15 @@ public class ModelingComponents : ModuleRules
 				// ... add private dependencies that you statically link with here ...
 			}
 		);
+		
+		if (Target.bCompileAgainstEditor) // #if WITH_EDITOR
+		{
+			PrivateDependencyModuleNames.AddRange(
+				new string[]
+				{
+					"UnrealEd"
+				});
+		}
 
 		DynamicallyLoadedModuleNames.AddRange(
 			new string[]
diff --git a/Engine/Plugins/Runtime/MeshModelingToolset/Source/ModelingComponents/Private/PropertySets/OnAcceptProperties.cpp b/Engine/Plugins/Runtime/MeshModelingToolset/Source/ModelingComponents/Private/PropertySets/OnAcceptProperties.cpp
--- a/Engine/Plugins/Runtime/MeshModelingToolset/Source/ModelingComponents/Private/PropertySets/OnAcceptProperties.cpp
+++ b/Engine/Plugins/Runtime/MeshModelingToolset/Source/ModelingComponents/Private/PropertySets/OnAcceptProperties.cpp
@@ -2,10 +2,17 @@
 
 #include "PropertySets/OnAcceptProperties.h"
 
+#include "Engine/World.h"
 #include "GameFramework/Actor.h"
 #include "UObject/UObjectGlobals.h"
 #include "InteractiveToolManager.h"
 
+#if WITH_EDITOR
+#include "Editor/UnrealEdEngine.h"
+#include "Selection.h"
+#include "UnrealEdGlobals.h"
+#endif
+
 #include UE_INLINE_GENERATED_CPP_BY_NAME(OnAcceptProperties)
 
 #define LOCTEXT_NAMESPACE "UOnAcceptHandleSourcesProperties"
@@ -51,7 +58,19 @@ void UOnAcceptHandleSourcesPropertiesBase::ApplyMethod(const TArray<AActor*>& Ac
 
 			if (bDelete)
 			{
-				Actor->Destroy();
+				if (UWorld* ActorWorld = Actor->GetWorld())
+				{
+#if WITH_EDITOR
+					if (GIsEditor && GUnrealEd)
+					{
+						GUnrealEd->DeleteActors(TArray{Actor}, ActorWorld, GUnrealEd->GetSelectedActors()->GetElementSelectionSet());
+					}
+					else
+#endif
+					{
+						ActorWorld->DestroyActor(Actor);
+					}
+				}
 			}
 			else
 			{
