disable workarounds for now-fixed collision update bug, so bake transform, edit pivot, and the various voxel/boolean tools will not create two transactions per operation for static meshes that have convex collision

diff --git a/Engine/Plugins/Experimental/MeshModelingToolsetExp/Source/MeshModelingToolsExp/Private/BakeTransformTool.cpp b/Engine/Plugins/Experimental/MeshModelingToolsetExp/Source/MeshModelingToolsExp/Private/BakeTransformTool.cpp
--- a/Engine/Plugins/Experimental/MeshModelingToolsetExp/Source/MeshModelingToolsExp/Private/BakeTransformTool.cpp
+++ b/Engine/Plugins/Experimental/MeshModelingToolsetExp/Source/MeshModelingToolsExp/Private/BakeTransformTool.cpp
@@ -157,7 +157,7 @@ void UBakeTransformTool::UpdateAssets()
 		}
 	}
 
-	constexpr bool bWorkaroundForCrashIfConvexAndMeshModifiedInSameTransaction = true;
+	constexpr bool bWorkaroundForCrashIfConvexAndMeshModifiedInSameTransaction = false;
 	bool bNeedSeparateTransactionForSimpleCollision = false;
 	if constexpr (bWorkaroundForCrashIfConvexAndMeshModifiedInSameTransaction)
 	{
diff --git a/Engine/Plugins/Experimental/MeshModelingToolsetExp/Source/MeshModelingToolsExp/Private/EditPivotTool.cpp b/Engine/Plugins/Experimental/MeshModelingToolsetExp/Source/MeshModelingToolsExp/Private/EditPivotTool.cpp
--- a/Engine/Plugins/Experimental/MeshModelingToolsetExp/Source/MeshModelingToolsExp/Private/EditPivotTool.cpp
+++ b/Engine/Plugins/Experimental/MeshModelingToolsetExp/Source/MeshModelingToolsExp/Private/EditPivotTool.cpp
@@ -417,7 +417,7 @@ void UEditPivotTool::UpdateAssets(const FFrame3d& NewPivotWorldFrame)
 
 	// If the targets have any convex hulls, we need to update those first in a separate transaction to work around a crash
 	// in the handling of convex DDC data
-	constexpr bool bWorkaroundForCrashIfConvexAndMeshModifiedInSameTransaction = true;
+	constexpr bool bWorkaroundForCrashIfConvexAndMeshModifiedInSameTransaction = false;
 	bool bNeedSeparateTransactionForSimpleCollision = false;
 	if constexpr (bWorkaroundForCrashIfConvexAndMeshModifiedInSameTransaction)
 	{
diff --git a/Engine/Plugins/Runtime/MeshModelingToolset/Source/ModelingComponents/Private/BaseTools/BaseCreateFromSelectedTool.cpp b/Engine/Plugins/Runtime/MeshModelingToolset/Source/ModelingComponents/Private/BaseTools/BaseCreateFromSelectedTool.cpp
--- a/Engine/Plugins/Runtime/MeshModelingToolset/Source/ModelingComponents/Private/BaseTools/BaseCreateFromSelectedTool.cpp
+++ b/Engine/Plugins/Runtime/MeshModelingToolset/Source/ModelingComponents/Private/BaseTools/BaseCreateFromSelectedTool.cpp
@@ -358,6 +358,7 @@ void UBaseCreateFromSelectedTool::OnShutdown(EToolShutdownType ShutdownType)
 		// Note: For static meshes, this must be done as a separate transaction due to a long-standing bug where undo/redo 
 		// will crash if an asset updates both its collision and geometry in the same transaction
 		// TODO: If/when that static mesh bug is fixed, this collision updating code can be separated out into the UpdateAsset() method
+		constexpr bool bWorkaroundForCrashIfConvexAndMeshModifiedInSameTransaction = false;
 		bool bAlreadyOpenedMainTransaction = false;
 		if (SupportsCollisionTransfer() && HandleSourcesProperties->OutputWriteTo != EBaseCreateFromSelectedTargetType::NewObject && Targets.Num() > 1)
 		{
@@ -397,7 +398,8 @@ void UBaseCreateFromSelectedTool::OnShutdown(EToolShutdownType ShutdownType)
 
 				if (bHasAddedShapes)
 				{
-					if (!Cast<UStaticMeshComponent>(TargetComponent))
+					bool bWorkaroundNotNeeded = !bWorkaroundForCrashIfConvexAndMeshModifiedInSameTransaction | !Cast<UStaticMeshComponent>(TargetComponent); // slightly odd conditional to workaround compiler warning
+					if (bWorkaroundNotNeeded)
 					{
 						// If we're not operating on a static mesh, we can open the main transaction here so the collision transaction becomes part of it ...
 						// for static mesh, we need the collision update to be a separate transaction to avoid crashing on undo/redo :(
