MeshModelingToolset: Update preview when toggling Shells to Solids in PolyEdit extrude activity.

diff --git a/Engine/Plugins/Runtime/MeshModelingToolset/Source/MeshModelingTools/Private/ToolActivities/PolyEditExtrudeActivity.cpp b/Engine/Plugins/Runtime/MeshModelingToolset/Source/MeshModelingTools/Private/ToolActivities/PolyEditExtrudeActivity.cpp
--- a/Engine/Plugins/Runtime/MeshModelingToolset/Source/MeshModelingTools/Private/ToolActivities/PolyEditExtrudeActivity.cpp
+++ b/Engine/Plugins/Runtime/MeshModelingToolset/Source/MeshModelingTools/Private/ToolActivities/PolyEditExtrudeActivity.cpp
@@ -174,6 +174,14 @@ void UPolyEditExtrudeActivity::Setup(UInteractiveTool* ParentToolIn)
 			ActivityContext->Preview->InvalidateResult();
 		}
 	});
+	ExtrudeProperties->WatchProperty(ExtrudeProperties->bShellsToSolids,
+	[this](bool) {
+		if (bIsRunning)
+		{
+			ActivityContext->Preview->InvalidateResult();
+		}
+	});
+
 
 	OffsetProperties = NewObject<UPolyEditOffsetProperties>();
 	OffsetProperties->RestoreProperties(ParentTool.Get());
