stop assuming that the group topology edge IDs are the same as the corresponding mesh edge IDs in the tri edit extrude activity

diff --git a/Engine/Plugins/Runtime/MeshModelingToolset/Source/MeshModelingTools/Private/ToolActivities/PolyEditExtrudeEdgeActivity.cpp b/Engine/Plugins/Runtime/MeshModelingToolset/Source/MeshModelingTools/Private/ToolActivities/PolyEditExtrudeEdgeActivity.cpp
--- a/Engine/Plugins/Runtime/MeshModelingToolset/Source/MeshModelingTools/Private/ToolActivities/PolyEditExtrudeEdgeActivity.cpp
+++ b/Engine/Plugins/Runtime/MeshModelingToolset/Source/MeshModelingTools/Private/ToolActivities/PolyEditExtrudeEdgeActivity.cpp
@@ -427,17 +427,11 @@ void UPolyEditExtrudeEdgeActivity::ApplyExtrude()
 
 	StableEdgeIDs.GetEdgeIDs(*ActivityContext->CurrentMesh, NewSelectionEids);
 
-	if (ActivityContext->bTriangleMode)
+	for (int32 Eid : NewSelectionEids)
 	{
-		NewSelection.SelectedEdgeIDs.Append(NewSelectionEids);
-	}
-	else
-	{
-		for (int32 Eid : NewSelectionEids)
-		{
-			NewSelection.SelectedEdgeIDs.Add(ActivityContext->CurrentTopology->FindGroupEdgeID(Eid));
-		}
+		NewSelection.SelectedEdgeIDs.Add(ActivityContext->CurrentTopology->FindGroupEdgeID(Eid));
 	}
+
 	ActivityContext->SelectionMechanic->SetSelection(NewSelection);
 	ActivityContext->SelectionMechanic->EndChangeAndEmitIfModified();
 	ParentTool->GetToolManager()->EndUndoTransaction();
@@ -543,9 +537,10 @@ void UPolyEditExtrudeEdgeActivity::GatherSelectedEids()
 
 		if (ActivityContext->bTriangleMode)
 		{
-			SelectedEids.Add(GroupEdgeID);
+			int32 MeshEdgeID = ActivityContext->CurrentTopology->Edges[GroupEdgeID].Span.Edges[0];
+			SelectedEids.Add(MeshEdgeID);
 
-			int32 Tid = ActivityContext->CurrentMesh->GetEdgeT(GroupEdgeID).A;
+			int32 Tid = ActivityContext->CurrentMesh->GetEdgeT(MeshEdgeID).A;
 			int32 NeighborGroup = ActivityContext->CurrentMesh->GetTriangleGroup(Tid);
 
 			int32* NewGroupID = NeighborGroupToNewGroup.Find(NeighborGroup);
