MeshModelingToolset: Allow snapping to objects in scene when using DrawPolygonTool even when the camera is positioned such that the raycast does not hit the tool plane.

We now early out when the raycast misses the plane if snapping is disabled. If the tool plane is missed, then we use a point along the camera ray to find scene snap point (passed to FindSceneSnapPoint). A related bug is also fixed in this change in the snapping manager which caused only small sections of snappable edges to be considered. This was because the point passed to FindSceneSnapPoint was used to find a point on the 3 edges of a hit triangle, which is sensitive to camera position. Instead, we now use the raycast's impact point for the purposes of finding snap points on edges.

diff --git a/Engine/Plugins/Runtime/MeshModelingToolset/Source/MeshModelingTools/Private/DrawPolygonTool.cpp b/Engine/Plugins/Runtime/MeshModelingToolset/Source/MeshModelingTools/Private/DrawPolygonTool.cpp
--- a/Engine/Plugins/Runtime/MeshModelingToolset/Source/MeshModelingTools/Private/DrawPolygonTool.cpp
+++ b/Engine/Plugins/Runtime/MeshModelingToolset/Source/MeshModelingTools/Private/DrawPolygonTool.cpp
@@ -524,8 +524,9 @@ bool UDrawPolygonTool::FindDrawPlaneHitPoint(const FInputDeviceRay& ClickPos, FV
 
 	const FFrame3d& Frame = PlaneMechanic->Plane;
 	FVector3d HitPos;
-	bool bHit = Frame.RayPlaneIntersection((FVector3d)ClickPos.WorldRay.Origin, (FVector3d)ClickPos.WorldRay.Direction, 2, HitPos);
-	if (bHit == false)
+	FVector3d SceneSnapPos;
+	bool bHitPlane = Frame.RayPlaneIntersection((FVector3d)ClickPos.WorldRay.Origin, (FVector3d)ClickPos.WorldRay.Direction, 2, HitPos);
+	if (!bHitPlane && (!SnapProperties->bEnableSnapping || bIgnoreSnappingToggle))
 	{
 		return false;
 	}
@@ -542,8 +543,10 @@ bool UDrawPolygonTool::FindDrawPlaneHitPoint(const FInputDeviceRay& ClickPos, FV
 	}
 	else 
 	{
+		// Since we do not early out if we do not hit the tool plane when snapping is enabled,
+		// HitPos is not guaranteed to be meaningful here unless bHitPlane is true
 		FVector3d WorldGridSnapPos;
-		if (ToolSceneQueriesUtil::FindWorldGridSnapPoint(this, HitPos, WorldGridSnapPos))
+		if (bHitPlane && ToolSceneQueriesUtil::FindWorldGridSnapPoint(this, HitPos, WorldGridSnapPos))
 		{
 			WorldGridSnapPos = Frame.ToPlane(WorldGridSnapPos, 2);
 			SnapEngine.AddPointTarget(WorldGridSnapPos, CurrentGridSnapID, 
@@ -553,8 +556,8 @@ bool UDrawPolygonTool::FindDrawPlaneHitPoint(const FInputDeviceRay& ClickPos, FV
 
 		if (SnapProperties->bSnapToVertices || SnapProperties->bSnapToEdges)
 		{
-			FVector3d SceneSnapPos;
-			if (ToolSceneQueriesUtil::FindSceneSnapPoint(this, HitPos, SceneSnapPos, SnapProperties->bSnapToVertices, SnapProperties->bSnapToEdges, 0, &LastSnapGeometry))
+			const FVector3d PointOnWorldRay = ClickPos.WorldRay.PointAt(1);
+			if (ToolSceneQueriesUtil::FindSceneSnapPoint(this, bHitPlane ? HitPos : PointOnWorldRay, SceneSnapPos, SnapProperties->bSnapToVertices, SnapProperties->bSnapToEdges, 0, &LastSnapGeometry))
 			{
 				SnapEngine.AddPointTarget(SceneSnapPos, CurrentSceneSnapID, SnapEngine.MinInternalPriority() - 10);
 			}
@@ -568,6 +571,7 @@ bool UDrawPolygonTool::FindDrawPlaneHitPoint(const FInputDeviceRay& ClickPos, FV
 		}
 		SnapEngine.bEnableSnapToKnownLengths = SnapProperties->bSnapToLengths;
 	}
+	
 	// ignore snapping to start point unless we have at least 3 vertices
 	if (bInFixedPolygonMode == false && PolygonVertices.Num() > 0)
 	{
@@ -580,25 +584,27 @@ bool UDrawPolygonTool::FindDrawPlaneHitPoint(const FInputDeviceRay& ClickPos, FV
 			SnapEngine.RemoveIgnoreTarget(StartPointSnapID);
 		}
 	}
-
-	SnapEngine.UpdateSnappedPoint(HitPos);
+	
+	SnapEngine.UpdateSnappedPoint(bHitPlane ? HitPos : SceneSnapPos);
 
 	// remove scene snap point
 	SnapEngine.RemovePointTargetsByID(CurrentSceneSnapID);
 	SnapEngine.RemovePointTargetsByID(CurrentGridSnapID);
 
+	// Success case 1: HitPosOut is set to a snap point.
 	if (SnapEngine.HaveActiveSnap())
 	{
 		HitPosOut = SnapEngine.GetActiveSnapToPoint();
 		return true;
 	}
-
-
+	
 	// if not yet snapped and we want to hit objects, do that
-	if (SnapProperties->bSnapToSurfaces && !bIgnoreSnappingToggle)
+	if (SnapProperties->bEnableSnapping && SnapProperties->bSnapToSurfaces && !bIgnoreSnappingToggle)
 	{
 		FHitResult Result;
 		bool bWorldHit = ToolSceneQueriesUtil::FindNearestVisibleObjectHit(this, Result, ClickPos.WorldRay);
+
+		// Success case 2: HitPosOut is set to a point on the tool plane by projecting a found point in the world onto the plane.
 		if (bWorldHit)
 		{
 			bHaveSurfaceHit = true;
@@ -606,11 +612,20 @@ bool UDrawPolygonTool::FindDrawPlaneHitPoint(const FInputDeviceRay& ClickPos, FV
 			const FVector3d UseHitPos = Result.ImpactPoint + static_cast<double>(SnapProperties->SnapToSurfacesOffset) * Result.Normal;
 			HitPos = Frame.ToPlane(UseHitPos, 2);
 			SurfaceOffsetPoint = UseHitPos;
+
+			HitPosOut = HitPos;
+			return true;
 		}
 	}
 
-	HitPosOut = HitPos;
-	return true;
+	// Success case 3: HitPosOut is set to a point on the tool plane based on raycast intersection with the plane.
+	if (bHitPlane)
+	{
+		HitPosOut = HitPos;
+		return true;
+	}
+	
+	return false;
 }
 
 void UDrawPolygonTool::OnBeginSequencePreview(const FInputDeviceRay& DevicePos)
diff --git a/Engine/Plugins/Runtime/MeshModelingToolset/Source/ModelingComponents/Private/Snapping/ModelingSceneSnappingManager.cpp b/Engine/Plugins/Runtime/MeshModelingToolset/Source/ModelingComponents/Private/Snapping/ModelingSceneSnappingManager.cpp
--- a/Engine/Plugins/Runtime/MeshModelingToolset/Source/ModelingComponents/Private/Snapping/ModelingSceneSnappingManager.cpp
+++ b/Engine/Plugins/Runtime/MeshModelingToolset/Source/ModelingComponents/Private/Snapping/ModelingSceneSnappingManager.cpp
@@ -539,7 +539,7 @@ bool UModelingSceneSnappingManager::ExecuteSceneSnapQueryPosition(const FSceneSn
 	};
 
 
-	auto TrySnapToEdge = [](const FSceneSnapQueryRequest& RequestIn, const FRay3d& WorldRay, FSceneSnapQueryResult& SnapResult, double& SmallestSnapAngle)
+	auto TrySnapToEdge = [](const FSceneSnapQueryRequest& RequestIn, const FVector& SurfacePoint, const FRay3d& WorldRay, FSceneSnapQueryResult& SnapResult, double& SmallestSnapAngle)
 	{
 		if ( ((RequestIn.TargetTypes & ESceneSnapQueryTargetType::MeshEdge) != ESceneSnapQueryTargetType::None) &&
 			  (SnapResult.TargetType != ESceneSnapQueryTargetType::MeshVertex) )
@@ -547,7 +547,7 @@ bool UModelingSceneSnappingManager::ExecuteSceneSnapQueryPosition(const FSceneSn
 			for (int j = 0; j < 3; ++j)
 			{
 				UE::Geometry::FSegment3d Segment(SnapResult.TriVertices[j], SnapResult.TriVertices[(j+1)%3]);
-				FVector3d EdgeNearestPt = Segment.NearestPoint(RequestIn.Position);
+				FVector3d EdgeNearestPt = Segment.NearestPoint(SurfacePoint);
 				double VisualAngle = UE::Geometry::VectorUtil::OpeningAngleD(RequestIn.Position, EdgeNearestPt, WorldRay.Origin);
 				if (VisualAngle < SmallestSnapAngle )
 				{
@@ -617,7 +617,7 @@ bool UModelingSceneSnappingManager::ExecuteSceneSnapQueryPosition(const FSceneSn
 			TrySnapToVertex(Request, WorldRay, SnapResult, SmallestAngle);
 
 			// try snapping to nearest points on edges
-			TrySnapToEdge(Request, WorldRay, SnapResult, SmallestAngle);
+			TrySnapToEdge(Request, HitPoint.WorldPoint, WorldRay, SnapResult, SmallestAngle);
 
 			// if we found a valid snap, return it
 			if (SmallestAngle < (double)Request.VisualAngleThresholdDegrees)
@@ -647,7 +647,7 @@ bool UModelingSceneSnappingManager::ExecuteSceneSnapQueryPosition(const FSceneSn
 				TrySnapToVertex(Request, WorldRay, SnapResult, SmallestAngle);
 
 				// try snapping to nearest points on edges
-				TrySnapToEdge(Request, WorldRay, SnapResult, SmallestAngle);
+				TrySnapToEdge(Request, HitResult.ImpactPoint, WorldRay, SnapResult, SmallestAngle);
 
 				// if we found a valid snap, return it
 				if (SmallestAngle < (double)Request.VisualAngleThresholdDegrees)
