remove deprecated persistent mesh selection code from modelingcomponents

diff --git a/Engine/Plugins/Runtime/MeshModelingToolset/Source/ModelingComponents/Private/BaseTools/SingleSelectionMeshEditingTool.cpp b/Engine/Plugins/Runtime/MeshModelingToolset/Source/ModelingComponents/Private/BaseTools/SingleSelectionMeshEditingTool.cpp
--- a/Engine/Plugins/Runtime/MeshModelingToolset/Source/ModelingComponents/Private/BaseTools/SingleSelectionMeshEditingTool.cpp
+++ b/Engine/Plugins/Runtime/MeshModelingToolset/Source/ModelingComponents/Private/BaseTools/SingleSelectionMeshEditingTool.cpp
@@ -45,17 +45,6 @@ void USingleSelectionMeshEditingToolBuilder::InitializeNewTool(USingleSelectionM
 	check(Target);
 	NewTool->SetTarget(Target);
 	NewTool->SetWorld(SceneState.World);
-
-PRAGMA_DISABLE_DEPRECATION_WARNINGS
-	if (WantsInputSelectionIfAvailable())
-	{
-		const UDEPRECATED_PersistentMeshSelection* InputSelection = UE::Geometry::GetCurrentToolInputSelection(SceneState, Target);
-		if (InputSelection != nullptr)
-		{
-			NewTool->SetInputSelection(InputSelection);
-		}
-	}
-PRAGMA_ENABLE_DEPRECATION_WARNINGS
 }
 
 
diff --git a/Engine/Plugins/Runtime/MeshModelingToolset/Source/ModelingComponents/Private/Selection/PersistentMeshSelection.cpp b/Engine/Plugins/Runtime/MeshModelingToolset/Source/ModelingComponents/Private/Selection/PersistentMeshSelection.cpp
deleted file mode 100644
diff --git a/Engine/Plugins/Runtime/MeshModelingToolset/Source/ModelingComponents/Private/Selection/PersistentMeshSelectionManager.cpp b/Engine/Plugins/Runtime/MeshModelingToolset/Source/ModelingComponents/Private/Selection/PersistentMeshSelectionManager.cpp
deleted file mode 100644
diff --git a/Engine/Plugins/Runtime/MeshModelingToolset/Source/ModelingComponents/Private/Selection/StoredMeshSelectionUtil.cpp b/Engine/Plugins/Runtime/MeshModelingToolset/Source/ModelingComponents/Private/Selection/StoredMeshSelectionUtil.cpp
--- a/Engine/Plugins/Runtime/MeshModelingToolset/Source/ModelingComponents/Private/Selection/StoredMeshSelectionUtil.cpp
+++ b/Engine/Plugins/Runtime/MeshModelingToolset/Source/ModelingComponents/Private/Selection/StoredMeshSelectionUtil.cpp
@@ -2,7 +2,6 @@
 
 #include "Selection/StoredMeshSelectionUtil.h"
 #include "Selection/PersistentMeshSelection.h"
-#include "Selection/PersistentMeshSelectionManager.h"
 #include "TargetInterfaces/PrimitiveComponentBackedTarget.h"
 #include "InteractiveTool.h"
 #include "InteractiveToolManager.h"
@@ -65,39 +64,6 @@ bool UE::Geometry::SetToolOutputGeometrySelectionForTarget(UInteractiveTool* Too
 
 
 
-const UDEPRECATED_PersistentMeshSelection* UE::Geometry::GetCurrentToolInputSelection(const FToolBuilderState& SceneState, UToolTarget* Target)
-{
-	ensureMsgf(false, TEXT("This path no longer functions and will be removed"));
-	return nullptr;
-}
-
-
-
-bool UE::Geometry::GetStoredSelectionAsTriangles(
-	const UDEPRECATED_PersistentMeshSelection* Selection,
-	const FDynamicMesh3& Mesh,
-	TArray<int32>& TrianglesOut)
-{
-	return false;
-}
-
-
-void UE::Geometry::ClearActiveToolSelection(UInteractiveToolManager* ToolManager)
-{
-}
-
-
-bool UE::Geometry::SetToolOutputSelection(UInteractiveTool* Tool, UDEPRECATED_PersistentMeshSelection* Selection)
-{
-	return false;
-}
-
-
-
-
-
-
-
 
 namespace UELocal
 {
diff --git a/Engine/Plugins/Runtime/MeshModelingToolset/Source/ModelingComponents/Public/BaseTools/SingleSelectionMeshEditingTool.h b/Engine/Plugins/Runtime/MeshModelingToolset/Source/ModelingComponents/Public/BaseTools/SingleSelectionMeshEditingTool.h
--- a/Engine/Plugins/Runtime/MeshModelingToolset/Source/ModelingComponents/Public/BaseTools/SingleSelectionMeshEditingTool.h
+++ b/Engine/Plugins/Runtime/MeshModelingToolset/Source/ModelingComponents/Public/BaseTools/SingleSelectionMeshEditingTool.h
@@ -7,7 +7,6 @@
 #include "SingleSelectionMeshEditingTool.generated.h"
 
 class USingleSelectionMeshEditingTool;
-class UDEPRECATED_PersistentMeshSelection;
 
 /**
  * USingleSelectionMeshEditingToolBuilder is a base tool builder for single
@@ -32,10 +31,6 @@ public:
 	/** Called by BuildTool to configure the Tool with the input MeshSource based on the SceneState */
 	virtual void InitializeNewTool(USingleSelectionMeshEditingTool* Tool, const FToolBuilderState& SceneState) const;
 
-	/** @return true if this Tool would like access to an available Input Selection object */
-	UE_DEPRECATED(5.2, "UPersistentMeshSelection and related functions are deprecated")
-	virtual bool WantsInputSelectionIfAvailable() const { return false; }
-
 protected:
 	virtual const FToolTargetTypeRequirements& GetTargetRequirements() const override;
 };
@@ -58,47 +53,4 @@ public:
 protected:
 	UPROPERTY()
 	TWeakObjectPtr<UWorld> TargetWorld = nullptr;
-
-PRAGMA_DISABLE_DEPRECATION_WARNINGS
-
-	//
-	// Mesh Selection support
-	//
-
-public:
-	/**
-	 * Set a Selection for the Tool to use. This should be called before tool Setup() (ie in the ToolBuilder) 
-	 * to allow the Tool to configure it's behavior based on the Selection (which may or may-not exist).
-	 * If the Tool requires a Selection, this needs to be handled at the Builder level.
-	 */
-	UE_DEPRECATED(5.2, "UPersistentMeshSelection and related functions are deprecated")
-	virtual void SetInputSelection(const UDEPRECATED_PersistentMeshSelection* SelectionIn)
-	{
-		InputSelection_DEPRECATED = SelectionIn;
-	}
-
-	/** @return true if an InputSelection is available */
-	UE_DEPRECATED(5.2, "UPersistentMeshSelection and related functions are deprecated")
-	virtual bool HasInputSelection() const
-	{
-		return InputSelection_DEPRECATED != nullptr;
-	}
-
-	/** @return the input Selection, or nullptr if one was not configured */
-	UE_DEPRECATED(5.2, "UPersistentMeshSelection and related functions are deprecated")
-	virtual const UDEPRECATED_PersistentMeshSelection* GetInputSelection() const
-	{
-		return InputSelection_DEPRECATED;
-	}
-
-protected:
-
-	/**
-	 * (optional) Mesh Selection provided on Tool Input. This should never be modified after the Tool's Setup() is called
-	 */
-	UE_DEPRECATED(5.2, "UPersistentMeshSelection and related functions are deprecated")
-	UPROPERTY()
-	TObjectPtr<const UDEPRECATED_PersistentMeshSelection> InputSelection_DEPRECATED = nullptr;
-
-PRAGMA_ENABLE_DEPRECATION_WARNINGS
 };
diff --git a/Engine/Plugins/Runtime/MeshModelingToolset/Source/ModelingComponents/Public/Selection/PersistentMeshSelection.h b/Engine/Plugins/Runtime/MeshModelingToolset/Source/ModelingComponents/Public/Selection/PersistentMeshSelection.h
--- a/Engine/Plugins/Runtime/MeshModelingToolset/Source/ModelingComponents/Public/Selection/PersistentMeshSelection.h
+++ b/Engine/Plugins/Runtime/MeshModelingToolset/Source/ModelingComponents/Public/Selection/PersistentMeshSelection.h
@@ -6,7 +6,6 @@
 #include "IndexTypes.h"
 #include "SegmentTypes.h"
 
-#include "PersistentMeshSelection.generated.h"
 
 class UPrimitiveComponent;
 class UInteractiveTool;
@@ -26,7 +25,7 @@ PREDECLARE_USE_GEOMETRY_CLASS(FCompactMaps);
  * 
  * @warning this class is likely to change in the future
  */
-struct MODELINGCOMPONENTS_API FGenericMeshSelection
+struct FGenericMeshSelection
 {
 	// selection type
 	enum class ETopologyType
@@ -95,51 +94,3 @@ struct MODELINGCOMPONENTS_API FGenericMeshSelection
 
 };
 
-
-
-
-
-/**
- * UPersistentMeshSelection is a UObject wrapper for a FGenericMeshSelection
- */
-//UE_DEPRECATED(5.2, "UPersistentMeshSelection and related functions are deprecated")
-UCLASS(Deprecated)
-class MODELINGCOMPONENTS_API UDEPRECATED_PersistentMeshSelection : public UObject
-{
-	GENERATED_BODY()
-
-public:
-
-	/**
-	 * Resets the contents of the object using the given selection.
-	 *
-	 * @param CompactMaps If the mesh was compacted without updating the passed in topology object, these
-	 *  maps will be used to give an object that will work with the new mesh vids.
-	 */
-	void SetSelection(const FGroupTopology& TopologyIn, const FGroupTopologySelection& SelectionIn, const FCompactMaps* CompactMaps = nullptr);
-
-	/**
-	 * Initializes a FGroupTopologySelection using the current contents of the object. The topology
-	 * must already be initialized.
-	 */
-	void ExtractIntoSelectionObject(const FGroupTopology& TopologyIn, FGroupTopologySelection& SelectionOut) const;
-
-	/** @return true if the selection is empty */
-	bool IsEmpty() const
-	{
-		return Selection.IsEmpty();
-	}
-
-	UPrimitiveComponent* GetTargetComponent() const { return Selection.SourceComponent; }
-	FGenericMeshSelection::ETopologyType GetSelectionType() const { return Selection.TopologyType; }
-
-	/** replace the internal Selection data */
-	void SetSelection(const FGenericMeshSelection& SelectionIn) { Selection = SelectionIn; }
-
-	/** @return the internal Selection data */
-	const FGenericMeshSelection& GetSelection() const { return Selection; }
-
-protected:
-	FGenericMeshSelection Selection;
-};
-
diff --git a/Engine/Plugins/Runtime/MeshModelingToolset/Source/ModelingComponents/Public/Selection/PersistentMeshSelectionManager.h b/Engine/Plugins/Runtime/MeshModelingToolset/Source/ModelingComponents/Public/Selection/PersistentMeshSelectionManager.h
deleted file mode 100644
diff --git a/Engine/Plugins/Runtime/MeshModelingToolset/Source/ModelingComponents/Public/Selection/StoredMeshSelectionUtil.h b/Engine/Plugins/Runtime/MeshModelingToolset/Source/ModelingComponents/Public/Selection/StoredMeshSelectionUtil.h
--- a/Engine/Plugins/Runtime/MeshModelingToolset/Source/ModelingComponents/Public/Selection/StoredMeshSelectionUtil.h
+++ b/Engine/Plugins/Runtime/MeshModelingToolset/Source/ModelingComponents/Public/Selection/StoredMeshSelectionUtil.h
@@ -7,7 +7,6 @@
 #include "Selections/GeometrySelection.h"
 
 struct FToolBuilderState;
-class UDEPRECATED_PersistentMeshSelection;
 class UToolTarget;
 class UInteractiveTool;
 class UInteractiveToolManager;
@@ -51,45 +50,6 @@ namespace Geometry
 	MODELINGCOMPONENTS_API bool SetToolOutputGeometrySelectionForTarget(UInteractiveTool* Tool, UToolTarget* Target, const FGeometrySelection& OutputSelection);
 
 
-
-
-	/**
-	 * Given a FToolBuilderState, determine if an Input Selection exists that is applicable to the ToolTarget
-	 * @return Selection if found, or nullptr if not
-	 */
-	UE_DEPRECATED(5.2, "PersistentMeshSelection and related functions are deprecated. GetCurrentToolInputSelection now returns nullptr.")
-	MODELINGCOMPONENTS_API const UDEPRECATED_PersistentMeshSelection* GetCurrentToolInputSelection(const FToolBuilderState& SceneState, UToolTarget* Target);
-
-	/**
-	 * Convert the given Seletion to a list of Triangles of the specified Mesh. 
-	 * This will attempt to convert the different selection types represented by PersistentMeshSelection
-	 * (eg group topology, UV island, etc) to triangle indices in TrianglesOut
-	 * @return true if a list of triangles could be created.
-	 */
-	UE_DEPRECATED(5.2, "PersistentMeshSelection and related functions are deprecated. GetStoredSelectionAsTriangles now always returns an empty TrianglesOut list.")
-	MODELINGCOMPONENTS_API bool GetStoredSelectionAsTriangles(
-		const UDEPRECATED_PersistentMeshSelection* Selection,
-		const FDynamicMesh3& Mesh,
-		TArray<int32>& TrianglesOut );
-
-
-	/**
-	 * Set the active output selection of a Tool, which (if supported) will be persist between Tools
-	 * via a PersistentMeshSelectionManager context object. Passing a null Selection will clear any active selection.
-	 * @return true if the selection was saved
-	 */
-	UE_DEPRECATED(5.2, "PersistentMeshSelection and related functions are deprecated. SetToolOutputSelection is now ignored.")
-	MODELINGCOMPONENTS_API bool SetToolOutputSelection(
-		UInteractiveTool* Tool,
-		UDEPRECATED_PersistentMeshSelection* Selection);
-
-	/**
-	 * Clear any currently-active stored Tool selection
-	 */
-	UE_DEPRECATED(5.2, "PersistentMeshSelection and related functions are deprecated. ClearActiveToolSelection no longer has any effect.")
-	MODELINGCOMPONENTS_API void ClearActiveToolSelection(UInteractiveToolManager* ToolManager);
-
-
 	/*
 	 * Helper class that can store a list of edges as pairs of triangle IDs and {0,1,2} indices into 
 	 * the triangle edge triplet, because regular FDynamicMesh edge IDs may not stay the same across 
