remove various deprecated methods from modeling components code

diff --git a/Engine/Plugins/Runtime/MeshModelingToolset/Source/ModelingComponents/Private/Drawing/PolyEditPreviewMesh.cpp b/Engine/Plugins/Runtime/MeshModelingToolset/Source/ModelingComponents/Private/Drawing/PolyEditPreviewMesh.cpp
--- a/Engine/Plugins/Runtime/MeshModelingToolset/Source/ModelingComponents/Private/Drawing/PolyEditPreviewMesh.cpp
+++ b/Engine/Plugins/Runtime/MeshModelingToolset/Source/ModelingComponents/Private/Drawing/PolyEditPreviewMesh.cpp
@@ -36,9 +36,6 @@ void UPolyEditPreviewMesh::InitializeExtrudeType(
 
 	// save copy of initial patch
 	InitialEditPatch = EditPatch;
-	PRAGMA_DISABLE_DEPRECATION_WARNINGS
-	InitialEditPatchBVTree.SetMesh(&InitialEditPatch);
-	PRAGMA_ENABLE_DEPRECATION_WARNINGS
 
 	// save initial tri normals
 	InitialTriNormals.SetNum(InitialEditPatch.MaxTriangleID());
@@ -106,11 +103,6 @@ void UPolyEditPreviewMesh::InitializeExtrudeType(FDynamicMesh3&& BaseMesh,
 		MeshTransforms::ApplyTransform(InitialEditPatch, MeshTransform, true);
 	}
 
-	// save copy of initial patch
-	PRAGMA_DISABLE_DEPRECATION_WARNINGS
-	InitialEditPatchBVTree.SetMesh(&InitialEditPatch);
-	PRAGMA_ENABLE_DEPRECATION_WARNINGS
-
 	// save initial tri normals
 	InitialTriNormals.SetNum(InitialEditPatch.MaxTriangleID());
 	for ( int32 tid : InitialEditPatch.TriangleIndicesItr() )
@@ -265,9 +257,6 @@ void UPolyEditPreviewMesh::InitializeInsetType(const FDynamicMesh3* SourceMesh,
 
 	// save copy of initial patch
 	InitialEditPatch = EditPatch;
-	PRAGMA_DISABLE_DEPRECATION_WARNINGS
-	InitialEditPatchBVTree.SetMesh(&InitialEditPatch);
-	PRAGMA_ENABLE_DEPRECATION_WARNINGS
 
 	// initialize the preview mesh
 	UpdatePreview(&EditPatch);
diff --git a/Engine/Plugins/Runtime/MeshModelingToolset/Source/ModelingComponents/Private/Scene/SceneCapturePhotoSet.cpp b/Engine/Plugins/Runtime/MeshModelingToolset/Source/ModelingComponents/Private/Scene/SceneCapturePhotoSet.cpp
--- a/Engine/Plugins/Runtime/MeshModelingToolset/Source/ModelingComponents/Private/Scene/SceneCapturePhotoSet.cpp
+++ b/Engine/Plugins/Runtime/MeshModelingToolset/Source/ModelingComponents/Private/Scene/SceneCapturePhotoSet.cpp
@@ -307,229 +307,6 @@ void FSceneCapturePhotoSet::Compute()
 }
 
 
-void FSceneCapturePhotoSet::AddStandardExteriorCapturesFromBoundingBox(
-	FImageDimensions PhotoDimensions,
-	double HorizontalFOVDegrees,
-	double NearPlaneDist,
-	bool bFaces,
-	bool bUpperCorners,
-	bool bLowerCorners,
-	bool bUpperEdges,
-	bool bSideEdges)
-{
-	TArray<FVector3d> Directions;
-
-	if (bFaces)
-	{
-		Directions.Add(FVector3d::UnitX());
-		Directions.Add(-FVector3d::UnitX());
-		Directions.Add(FVector3d::UnitY());
-		Directions.Add(-FVector3d::UnitY());
-		Directions.Add(FVector3d::UnitZ());
-		Directions.Add(-FVector3d::UnitZ());
-	}
-	if (bUpperCorners)
-	{
-		Directions.Add(Normalized(FVector3d(1, 1, -1)));
-		Directions.Add(Normalized(FVector3d(-1, 1, -1)));
-		Directions.Add(Normalized(FVector3d(1, -1, -1)));
-		Directions.Add(Normalized(FVector3d(-1, -1, -1)));
-	}
-	if (bLowerCorners)
-	{
-		Directions.Add(Normalized(FVector3d(1, 1, 1)));
-		Directions.Add(Normalized(FVector3d(-1, 1, 1)));
-		Directions.Add(Normalized(FVector3d(1, -1, 1)));
-		Directions.Add(Normalized(FVector3d(-1, -1, 1)));
-	}
-	if (bUpperEdges)
-	{
-		Directions.Add(Normalized(FVector3d(-1, 0, -1)));
-		Directions.Add(Normalized(FVector3d(1, 0, -1)));
-		Directions.Add(Normalized(FVector3d(0, -1, -1)));
-		Directions.Add(Normalized(FVector3d(0, 1, -1)));
-	}
-	if (bSideEdges)
-	{
-		Directions.Add(Normalized(FVector3d(1, 1, 0)));
-		Directions.Add(Normalized(FVector3d(-1, 1, 0)));
-		Directions.Add(Normalized(FVector3d(1, -1, 0)));
-		Directions.Add(Normalized(FVector3d(-1, -1, 0)));
-	}
-	PRAGMA_DISABLE_DEPRECATION_WARNINGS
-	AddExteriorCaptures(PhotoDimensions, HorizontalFOVDegrees, NearPlaneDist, Directions);
-	PRAGMA_ENABLE_DEPRECATION_WARNINGS
-}
-
-void FSceneCapturePhotoSet::AddExteriorCaptures(
-	FImageDimensions PhotoDimensions,
-	double HorizontalFOVDegrees,
-	double NearPlaneDist,
-	const TArray<FVector3d>& Directions)
-{
-	check(this->TargetWorld != nullptr);
-
-	bWasCancelled = false;
-
-	FScopedSlowTask Progress(Directions.Num(), LOCTEXT("ComputingViewpoints", "Computing Viewpoints..."));
-	Progress.MakeDialog(bAllowCancel);
-
-	// Unregister all components to remove unwanted proxies from the scene. This was previously the only way to "hide" nanite meshes, now optional.
-	TSet<AActor*> VisibleActorsSet(VisibleActors);
-	TArray<AActor*> ActorsToRegister;
-	if (bEnforceVisibilityViaUnregister)
-	{
-		for (TActorIterator<AActor> Actor(TargetWorld); Actor; ++Actor)
-		{
-			if (!VisibleActorsSet.Contains(*Actor))
-			{
-				Actor->UnregisterAllComponents();
-				ActorsToRegister.Add(*Actor);
-			}
-		}
-	}
-
-	ON_SCOPE_EXIT
-	{
-		// Workaround for Nanite scene proxies visibility
-		// Reregister all components we previously unregistered
-		for (AActor* Actor : ActorsToRegister)
-		{
-			Actor->RegisterAllComponents();
-		}
-	};
-
-	FWorldRenderCapture RenderCapture;
-	RenderCapture.SetWorld(TargetWorld);
-	RenderCapture.SetVisibleActorsAndComponents(VisibleActors, VisibleComponents);
-	RenderCapture.SetDimensions(PhotoDimensions);
-	if (bWriteDebugImages)
-	{
-		RenderCapture.SetEnableWriteDebugImage(true, 0, DebugImagesFolderName);
-	}
-
-	// this tells us origin and radius - could be view-dependent...
-	FSphere RenderSphere = RenderCapture.ComputeContainingRenderSphere(HorizontalFOVDegrees);
-
-	int32 NumDirections = Directions.Num();
-	for (int32 di = 0; di < NumDirections; ++di)
-	{
-		Progress.EnterProgressFrame(1.f);
-		if (Progress.ShouldCancel())
-		{
-			bWasCancelled = true;
-			return;
-		}
-
-		FVector3d ViewDirection = Directions[di];
-		ViewDirection.Normalize();
-
-		FSpatialPhotoParams Params;
-		Params.NearPlaneDist = NearPlaneDist;
-		Params.HorzFOVDegrees = HorizontalFOVDegrees;
-		Params.Dimensions = PhotoDimensions;
-		// TODO Align the frame with the renderer coordinate system then remove the axis swapping in WorldRenderCapture.cpp
-		Params.Frame.AlignAxis(0, ViewDirection);
-		Params.Frame.ConstrainedAlignAxis(2, FVector3d::UnitZ(), Params.Frame.X());
-		Params.Frame.Origin = RenderSphere.Center;
-		Params.Frame.Origin -= RenderSphere.W * Params.Frame.X();
-
-		auto CaptureImageTypeFunc_3f = [this, &Progress, &RenderCapture, &Params, &NumDirections]
-			(ERenderCaptureType CaptureType, FSpatialPhotoSet3f& PhotoSet)
-		{
-			// Test NumDirections to avoid recomputing photo sets. Search :SceneCaptureWithExistingCaptures
-			if (PhotoSet.Num() < NumDirections)
-			{
-				FSpatialPhoto3f NewPhoto;
-				NewPhoto.Frame = Params.Frame;
-				NewPhoto.NearPlaneDist = Params.NearPlaneDist;
-				NewPhoto.HorzFOVDegrees = Params.HorzFOVDegrees;
-				NewPhoto.Dimensions = Params.Dimensions;
-
-				// TODO Do something with the success boolean returned by RenderCapture.CaptureFromPosition
-				FImageAdapter Image(&NewPhoto.Image);
-				FRenderCaptureConfig Config = GetCaptureConfig(CaptureType);
-				RenderCapture.CaptureFromPosition(CaptureType, NewPhoto.Frame, NewPhoto.HorzFOVDegrees, NewPhoto.NearPlaneDist, Image, Config);
-				PhotoSet.Add(MoveTemp(NewPhoto));
-			}
-
-			Progress.TickProgress();
-		};
-
-		auto CaptureImageTypeFunc_1f = [this, &Progress, &RenderCapture, &Params, &NumDirections]
-			(ERenderCaptureType CaptureType, FSpatialPhotoSet1f& PhotoSet)
-		{
-			// Test NumDirections to avoid recomputing photo sets. Search :SceneCaptureWithExistingCaptures
-			if (PhotoSet.Num() < NumDirections)
-			{
-				FSpatialPhoto1f NewPhoto;
-				NewPhoto.Frame = Params.Frame;
-				NewPhoto.NearPlaneDist = Params.NearPlaneDist;
-				NewPhoto.HorzFOVDegrees = Params.HorzFOVDegrees;
-				NewPhoto.Dimensions = Params.Dimensions;
-
-				// TODO Do something with the success boolean returned by RenderCapture.CaptureFromPosition
-				FImageAdapter Image(&NewPhoto.Image);
-				FRenderCaptureConfig Config = GetCaptureConfig(CaptureType);
-				RenderCapture.CaptureFromPosition(CaptureType, NewPhoto.Frame, NewPhoto.HorzFOVDegrees, NewPhoto.NearPlaneDist, Image, Config);
-				PhotoSet.Add(MoveTemp(NewPhoto));
-			}
-
-			Progress.TickProgress();
-		};
-
-		if (PhotoSetStatus[ERenderCaptureType::DeviceDepth] == ECaptureTypeStatus::Pending)
-		{
-			CaptureImageTypeFunc_1f(ERenderCaptureType::DeviceDepth, DeviceDepthPhotoSet);
-		}
-		if (PhotoSetStatus[ERenderCaptureType::BaseColor] == ECaptureTypeStatus::Pending)
-		{
-			CaptureImageTypeFunc_3f(ERenderCaptureType::BaseColor, BaseColorPhotoSet);
-		}
-		if (PhotoSetStatus[ERenderCaptureType::Roughness] == ECaptureTypeStatus::Pending)
-		{
-			CaptureImageTypeFunc_1f(ERenderCaptureType::Roughness, RoughnessPhotoSet);
-		}
-		if (PhotoSetStatus[ERenderCaptureType::Specular] == ECaptureTypeStatus::Pending)
-		{
-			CaptureImageTypeFunc_1f(ERenderCaptureType::Specular, SpecularPhotoSet);
-		}
-		if (PhotoSetStatus[ERenderCaptureType::Metallic] == ECaptureTypeStatus::Pending)
-		{
-			CaptureImageTypeFunc_1f(ERenderCaptureType::Metallic, MetallicPhotoSet);
-		}
-		if (PhotoSetStatus[ERenderCaptureType::CombinedMRS] == ECaptureTypeStatus::Pending)
-		{
-			CaptureImageTypeFunc_3f(ERenderCaptureType::CombinedMRS, PackedMRSPhotoSet);
-		}
-		if (PhotoSetStatus[ERenderCaptureType::WorldNormal] == ECaptureTypeStatus::Pending)
-		{
-			CaptureImageTypeFunc_3f(ERenderCaptureType::WorldNormal, WorldNormalPhotoSet);
-		}
-		if (PhotoSetStatus[ERenderCaptureType::Emissive] == ECaptureTypeStatus::Pending)
-		{
-			CaptureImageTypeFunc_3f(ERenderCaptureType::Emissive, EmissivePhotoSet);
-		}
-		if (PhotoSetStatus[ERenderCaptureType::Opacity] == ECaptureTypeStatus::Pending)
-		{
-			CaptureImageTypeFunc_1f(ERenderCaptureType::Opacity, OpacityPhotoSet);
-		}
-		if (PhotoSetStatus[ERenderCaptureType::SubsurfaceColor] == ECaptureTypeStatus::Pending)
-		{
-			CaptureImageTypeFunc_3f(ERenderCaptureType::SubsurfaceColor, SubsurfaceColorPhotoSet);
-		}
-
-		// AddExteriorCaptures can be called on an FSceneCapturePhotoSet which already has some capture types computed
-		// and in this case we should not modify the existing photo sets/cached photo set parameters.
-		// See :SceneCaptureWithExistingCaptures 
-		if (PhotoSetParams.Num() < NumDirections)
-		{
-			PhotoSetParams.Add(Params);
-		}
-	} // end directions loop
-}
-
-
 
 void FSceneCapturePhotoSet::OptimizePhotoSets()
 {
diff --git a/Engine/Plugins/Runtime/MeshModelingToolset/Source/ModelingComponents/Private/Selection/GroupTopologySelector.cpp b/Engine/Plugins/Runtime/MeshModelingToolset/Source/ModelingComponents/Private/Selection/GroupTopologySelector.cpp
--- a/Engine/Plugins/Runtime/MeshModelingToolset/Source/ModelingComponents/Private/Selection/GroupTopologySelector.cpp
+++ b/Engine/Plugins/Runtime/MeshModelingToolset/Source/ModelingComponents/Private/Selection/GroupTopologySelector.cpp
@@ -314,10 +314,6 @@ void FGroupTopologySelector::Initialize(const FDynamicMesh3* MeshIn, const FGrou
 	GroupTopologyUtils.GroupTopology = TopologyIn;
 	bGeometryInitialized = false;
 	bGeometryUpToDate = false;
-
-	PRAGMA_DISABLE_DEPRECATION_WARNINGS
-	Topology = TopologyIn;
-	PRAGMA_ENABLE_DEPRECATION_WARNINGS
 }
 
 
diff --git a/Engine/Plugins/Runtime/MeshModelingToolset/Source/ModelingComponents/Private/Selection/PolygonSelectionMechanic.cpp b/Engine/Plugins/Runtime/MeshModelingToolset/Source/ModelingComponents/Private/Selection/PolygonSelectionMechanic.cpp
--- a/Engine/Plugins/Runtime/MeshModelingToolset/Source/ModelingComponents/Private/Selection/PolygonSelectionMechanic.cpp
+++ b/Engine/Plugins/Runtime/MeshModelingToolset/Source/ModelingComponents/Private/Selection/PolygonSelectionMechanic.cpp
@@ -37,29 +37,6 @@ void UPolygonSelectionMechanic::Initialize(
 		GetSpatialSourceFuncIn);
 }
 
-// TODO: Remove this function when Properties is removed after deprecation
-void UPolygonSelectionMechanic::Setup(UInteractiveTool* ParentToolIn)
-{
-	UMeshTopologySelectionMechanic::Setup(ParentToolIn);
-
-	PRAGMA_DISABLE_DEPRECATION_WARNINGS
-	Properties_DEPRECATED = NewObject<UDEPRECATED_PolygonSelectionMechanicProperties>(this);
-	Properties_DEPRECATED->Initialize(this);
-	Properties_DEPRECATED->WatchProperty(Properties_DEPRECATED->bSelectVertices, [this](bool bSelectVertices) {
-		UpdateMarqueeEnabled();
-	});
-	Properties_DEPRECATED->WatchProperty(Properties_DEPRECATED->bSelectEdges, [this](bool bSelectVertices) {
-		UpdateMarqueeEnabled();
-	});
-	Properties_DEPRECATED->WatchProperty(Properties_DEPRECATED->bSelectFaces, [this](bool bSelectFaces) {
-		UpdateMarqueeEnabled();
-	});
-	Properties_DEPRECATED->WatchProperty(Properties_DEPRECATED->bEnableMarquee, [this](bool bEnableMarquee) {
-		UpdateMarqueeEnabled();
-	});
-	PRAGMA_ENABLE_DEPRECATION_WARNINGS
-}
-
 
 void UPolygonSelectionMechanic::GetSelection_AsGroupTopology(UE::Geometry::FGeometrySelection& SelectionOut, const FCompactMaps* CompactMapsToApply) const
 {
diff --git a/Engine/Plugins/Runtime/MeshModelingToolset/Source/ModelingComponents/Public/Drawing/PolyEditPreviewMesh.h b/Engine/Plugins/Runtime/MeshModelingToolset/Source/ModelingComponents/Public/Drawing/PolyEditPreviewMesh.h
--- a/Engine/Plugins/Runtime/MeshModelingToolset/Source/ModelingComponents/Public/Drawing/PolyEditPreviewMesh.h
+++ b/Engine/Plugins/Runtime/MeshModelingToolset/Source/ModelingComponents/Public/Drawing/PolyEditPreviewMesh.h
@@ -60,9 +60,6 @@ public:
 	
 	/// @param TargetMeshOut Will be populated with the initial edit patch mesh transformed to world space via the preview mesh transform
 	void MakeInsetTypeTargetMesh(FDynamicMesh3& TargetMeshOut);
-
-	UE_DEPRECATED(5.1, "We are removing this accessor because it was unused.")
-	const FDynamicMesh3& GetInitialPatchMesh() const { return InitialEditPatch; }
 	
 	void ApplyTranslationToPreview(const FTransform& TransformIn, FTransform& TranslationOut, FTransform& RotateScaleOut)
 	{
@@ -72,20 +69,9 @@ public:
 		RotateScaleOut.SetTranslation(FVector::ZeroVector);
 	}
 
-	// must be non-const because query functions are non-const
-	UE_DEPRECATED(5.1, "We are removing this accessor and the underlying tree, because they were unused.")
-	UE::Geometry::FDynamicMeshAABBTree3& GetInitialPatchMeshSpatial()
-	{
-		PRAGMA_DISABLE_DEPRECATION_WARNINGS
-		return InitialEditPatchBVTree; 
-		PRAGMA_ENABLE_DEPRECATION_WARNINGS
-	}
-
 protected:
 	TUniquePtr<UE::Geometry::FDynamicSubmesh3> ActiveSubmesh;
 	FDynamicMesh3 InitialEditPatch;
-	UE_DEPRECATED(5.1, "We are removing this tree, because it was unused.")
-	UE::Geometry::FDynamicMeshAABBTree3 InitialEditPatchBVTree;
 
 	TArray<int32> EditVertices;
 	TArray<FVector3d> InitialPositions;
diff --git a/Engine/Plugins/Runtime/MeshModelingToolset/Source/ModelingComponents/Public/Scene/SceneCapturePhotoSet.h b/Engine/Plugins/Runtime/MeshModelingToolset/Source/ModelingComponents/Public/Scene/SceneCapturePhotoSet.h
--- a/Engine/Plugins/Runtime/MeshModelingToolset/Source/ModelingComponents/Public/Scene/SceneCapturePhotoSet.h
+++ b/Engine/Plugins/Runtime/MeshModelingToolset/Source/ModelingComponents/Public/Scene/SceneCapturePhotoSet.h
@@ -116,33 +116,6 @@ public:
 	 */
 	void Compute();
 
-	/**
-	 * Add captures at the corners and face centers of the "view box",
-	 * ie the bounding box that contains the view sphere (see AddExteriorCaptures)
-	 */
-	UE_DEPRECATED(5.2, "AddStandardExteriorCapturesFromBoundingBox is deprecated, please use SetSpatialPhotoParams and Compute instead.")
-	void AddStandardExteriorCapturesFromBoundingBox(
-		FImageDimensions PhotoDimensions,
-		double HorizontalFOVDegrees,
-		double NearPlaneDist,
-		bool bFaces,
-		bool bUpperCorners,
-		bool bLowerCorners,
-		bool bUpperEdges,
-		bool bSideEdges);
-
-	/**
-	 * Add captures on the "view sphere", ie a sphere centered/sized such that the target actors
-	 * will be fully contained inside a square image rendered from locations on the sphere, where
-	 * the view direction is towards the sphere center. The Directions array defines the directions.
-	 */
-	UE_DEPRECATED(5.2, "AddExteriorCaptures is deprecated, please use SetSpatialPhotoParams and Compute instead.")
-	void AddExteriorCaptures(
-		FImageDimensions PhotoDimensions,
-		double HorizontalFOVDegrees,
-		double NearPlaneDist,
-		const TArray<FVector3d>& Directions);
-
 	/**
 	 * Post-process the various PhotoSets after capture, to reduce memory usage and sampling cost.
 	 */
diff --git a/Engine/Plugins/Runtime/MeshModelingToolset/Source/ModelingComponents/Public/Selection/GroupTopologySelector.h b/Engine/Plugins/Runtime/MeshModelingToolset/Source/ModelingComponents/Public/Selection/GroupTopologySelector.h
--- a/Engine/Plugins/Runtime/MeshModelingToolset/Source/ModelingComponents/Public/Selection/GroupTopologySelector.h
+++ b/Engine/Plugins/Runtime/MeshModelingToolset/Source/ModelingComponents/Public/Selection/GroupTopologySelector.h
@@ -172,11 +172,6 @@ public:
 	 */
 	bool ExpandSelectionByEdgeRings(FGroupTopologySelection& Selection);
 
-protected:
-
-	UE_DEPRECATED(5.2, "Use TopologyProvider in the base class")
-	const FGroupTopology* Topology = nullptr;
-
 private:
 
 	FGroupTopologyUtils GroupTopologyUtils;
diff --git a/Engine/Plugins/Runtime/MeshModelingToolset/Source/ModelingComponents/Public/Selection/PolygonSelectionMechanic.h b/Engine/Plugins/Runtime/MeshModelingToolset/Source/ModelingComponents/Public/Selection/PolygonSelectionMechanic.h
--- a/Engine/Plugins/Runtime/MeshModelingToolset/Source/ModelingComponents/Public/Selection/PolygonSelectionMechanic.h
+++ b/Engine/Plugins/Runtime/MeshModelingToolset/Source/ModelingComponents/Public/Selection/PolygonSelectionMechanic.h
@@ -45,7 +45,6 @@ public:
 	);
 
 	// UMeshTopologySelectionMechanic
-	virtual void Setup(UInteractiveTool* ParentTool) override;
 	virtual bool UpdateHighlight(const FRay& WorldRay) override;
 	virtual bool UpdateSelection(const FRay& WorldRay, FVector3d& LocalHitPositionOut, FVector3d& LocalHitNormalOut) override;
 
@@ -71,16 +70,6 @@ public:
 	void SetSelection_AsTriangleTopology(const UE::Geometry::FGeometrySelection& Selection);
 
 
-	UE_DEPRECATED(5.2, "Use GetTopologySelector in base class")
-	TSharedPtr<FGroupTopologySelector, ESPMode::ThreadSafe> GetTopologySelector() 
-	{ 
-		return StaticCastSharedPtr<FGroupTopologySelector>(TopoSelector); 
-	}
-
-	UE_DEPRECATED(5.2, "Use MechanicProperties in base class")
-	UPROPERTY()
-	TObjectPtr<UDEPRECATED_PolygonSelectionMechanicProperties> Properties_DEPRECATED;
-
 private:
 
 	// TODO: Would be nice to get rid of this and write everything in terms of TopologySelector and TopologyProvider
