ModelingTools: Minor inset, extrude, and push/pull option renames, reordering, and tooltip cleanup.

diff --git a/Engine/Plugins/Runtime/MeshModelingToolset/Source/MeshModelingTools/Private/ToolActivities/PolyEditExtrudeActivity.cpp b/Engine/Plugins/Runtime/MeshModelingToolset/Source/MeshModelingTools/Private/ToolActivities/PolyEditExtrudeActivity.cpp
--- a/Engine/Plugins/Runtime/MeshModelingToolset/Source/MeshModelingTools/Private/ToolActivities/PolyEditExtrudeActivity.cpp
+++ b/Engine/Plugins/Runtime/MeshModelingToolset/Source/MeshModelingTools/Private/ToolActivities/PolyEditExtrudeActivity.cpp
@@ -225,7 +225,15 @@ void UPolyEditExtrudeActivity::Setup(UInteractiveTool* ParentToolIn)
 	SetToolPropertySourceEnabled(PushPullProperties, false);
 	PushPullProperties->WatchProperty(PushPullProperties->MeasureDirection,
 	[this](EPolyEditExtrudeDirection) {
-		if (bIsRunning)
+		if (bIsRunning && PushPullProperties->DirectionMode != EPolyEditPushPullModeOptions::SingleDirection)
+		{
+			ReinitializeExtrudeHeightMechanic();
+			ActivityContext->Preview->InvalidateResult();
+		}
+	});
+	PushPullProperties->WatchProperty(PushPullProperties->SingleDirection,
+	[this](EPolyEditExtrudeDirection) {
+		if (bIsRunning && PushPullProperties->DirectionMode == EPolyEditPushPullModeOptions::SingleDirection)
 		{
 			ReinitializeExtrudeHeightMechanic();
 			ActivityContext->Preview->InvalidateResult();
@@ -578,7 +586,8 @@ FVector3d UPolyEditExtrudeActivity::GetExtrudeDirection() const
 		DirectionToUse = OffsetProperties->MeasureDirection;
 		break;
 	case EPropertySetToUse::PushPull:
-		DirectionToUse = PushPullProperties->MeasureDirection;
+		DirectionToUse = PushPullProperties->DirectionMode == EPolyEditPushPullModeOptions::SingleDirection ?
+			PushPullProperties->SingleDirection : PushPullProperties->MeasureDirection;
 		break;
 	}
 
diff --git a/Engine/Plugins/Runtime/MeshModelingToolset/Source/MeshModelingTools/Public/ToolActivities/PolyEditExtrudeActivity.h b/Engine/Plugins/Runtime/MeshModelingToolset/Source/MeshModelingTools/Public/ToolActivities/PolyEditExtrudeActivity.h
--- a/Engine/Plugins/Runtime/MeshModelingToolset/Source/MeshModelingTools/Public/ToolActivities/PolyEditExtrudeActivity.h
+++ b/Engine/Plugins/Runtime/MeshModelingToolset/Source/MeshModelingTools/Public/ToolActivities/PolyEditExtrudeActivity.h
@@ -120,29 +120,29 @@ public:
 		EditConditionHides, EditCondition = "DistanceMode == EPolyEditExtrudeDistanceMode::Fixed"))
 	double Distance = 100;
 
+	/** How to move the vertices during the extrude */
+	UPROPERTY(EditAnywhere, Category = Extrude)
+	EPolyEditExtrudeModeOptions DirectionMode = EPolyEditExtrudeModeOptions::SingleDirection;
+
 	/** Direction in which to extrude. */
 	UPROPERTY(EditAnywhere, Category = Extrude,
 		meta = (EditConditionHides, EditCondition = "DirectionMode == EPolyEditExtrudeModeOptions::SingleDirection"))
 	EPolyEditExtrudeDirection Direction = EPolyEditExtrudeDirection::SelectionNormal;
 
-	/** What axis to measure the extrusion distance along. */
-	UPROPERTY(EditAnywhere, Category = Extrude, AdvancedDisplay,
-		meta = (EditConditionHides, EditCondition = "DirectionMode != EPolyEditExtrudeModeOptions::SingleDirection && DistanceMode == EPolyEditExtrudeDistanceMode::ClickInViewport"))
-	EPolyEditExtrudeDirection MeasureDirection = EPolyEditExtrudeDirection::SelectionNormal;
-
-	/** Controls whether extruding an entire open-border patch should create a solid or an open shell */
-	UPROPERTY(EditAnywhere, Category = Extrude)
-	bool bShellsToSolids = true;
-
-	/** How to move the vertices during the extrude */
-	UPROPERTY(EditAnywhere, Category = Extrude)
-	EPolyEditExtrudeModeOptions DirectionMode = EPolyEditExtrudeModeOptions::SingleDirection;
-
 	/** Controls the maximum distance vertices can move from the target distance in order to stay parallel with their source triangles. */
 	UPROPERTY(EditAnywhere, Category = Extrude,
 		meta = (ClampMin = "1", EditConditionHides, EditCondition = "DirectionMode == EPolyEditExtrudeModeOptions::SelectedTriangleNormalsEven"))
 	double MaxDistanceScaleFactor = 4.0;
 
+	/** Controls whether extruding an entire open-border patch should create a solid or an open shell */
+	UPROPERTY(EditAnywhere, Category = Extrude)
+	bool bShellsToSolids = true;
+	
+	/** What axis to measure the extrusion distance along. */
+	UPROPERTY(EditAnywhere, Category = Extrude, AdvancedDisplay, meta = (EditConditionHides, 
+		EditCondition = "DirectionMode != EPolyEditExtrudeModeOptions::SingleDirection && DistanceMode == EPolyEditExtrudeDistanceMode::ClickInViewport"))
+	EPolyEditExtrudeDirection MeasureDirection = EPolyEditExtrudeDirection::SelectionNormal;
+
 	/** 
 	 * When extruding regions that touch the mesh border, assign the side groups (groups on the 
 	 * stitched side of the extrude) in a way that considers edge colinearity. For instance, when
@@ -160,11 +160,11 @@ class MESHMODELINGTOOLS_API UPolyEditOffsetProperties : public UInteractiveToolP
 
 public:
 	/** How the offset distance is set. */
-	UPROPERTY(EditAnywhere, Category = Extrude)
+	UPROPERTY(EditAnywhere, Category = Offset)
 	EPolyEditExtrudeDistanceMode DistanceMode = EPolyEditExtrudeDistanceMode::ClickInViewport;
 
 	/** Offset distance. */
-	UPROPERTY(EditAnywhere, Category = Extrude,
+	UPROPERTY(EditAnywhere, Category = Offset,
 		meta = (EditConditionHides, EditCondition = "DistanceMode == EPolyEditExtrudeDistanceMode::Fixed"))
 	double Distance = 100;
 
@@ -181,14 +181,14 @@ public:
 	UPROPERTY(EditAnywhere, Category = Offset)
 	bool bShellsToSolids = true;
 
-	/** What axis to measure the extrusion distance along. */
+	/** What axis to measure the offset distance along. */
 	UPROPERTY(EditAnywhere, Category = Offset, AdvancedDisplay, meta = (EditConditionHides,
 		EditCondition = "DistanceMode == EPolyEditExtrudeDistanceMode::ClickInViewport"))
 	EPolyEditExtrudeDirection MeasureDirection = EPolyEditExtrudeDirection::SelectionNormal;
 
 	/**
 	 * When offsetting regions that touch the mesh border, assign the side groups (groups on the
-	 * stitched side of the extrude) in a way that considers edge colinearity. For instance, when
+	 * stitched side of the offset) in a way that considers edge colinearity. For instance, when
 	 * true, extruding a flat rectangle will give four different groups on its sides rather than
 	 * one connected group.
 	 */
@@ -202,21 +202,26 @@ class MESHMODELINGTOOLS_API UPolyEditPushPullProperties : public UInteractiveToo
 	GENERATED_BODY()
 
 public:
-	/** How the offset distance is set. */
-	UPROPERTY(EditAnywhere, Category = Extrude)
+	/** How the extrusion distance is set. */
+	UPROPERTY(EditAnywhere, Category = ExtrusionOptions)
 	EPolyEditExtrudeDistanceMode DistanceMode = EPolyEditExtrudeDistanceMode::ClickInViewport;
 
-	/** Offset distance. */
-	UPROPERTY(EditAnywhere, Category = Extrude,
+	/** Extrusion distance. */
+	UPROPERTY(EditAnywhere, Category = ExtrusionOptions,
 		meta = (EditConditionHides, EditCondition = "DistanceMode == EPolyEditExtrudeDistanceMode::Fixed"))
 	double Distance = 100;
 
-	/** Which way to move vertices during the offset */
+	/** Which way to move vertices while extruding. */
 	UPROPERTY(EditAnywhere, Category = ExtrusionOptions)
 	EPolyEditPushPullModeOptions DirectionMode = EPolyEditPushPullModeOptions::SelectedTriangleNormals;
 
+	/** Direction in which to extrude. */
+	UPROPERTY(EditAnywhere, Category = ExtrusionOptions,
+		meta = (DisplayName="Direction", EditConditionHides, EditCondition = "DirectionMode == EPolyEditPushPullModeOptions::SingleDirection"))
+	EPolyEditExtrudeDirection SingleDirection = EPolyEditExtrudeDirection::SelectionNormal;
+
 	/** Controls the maximum distance vertices can move from the target distance in order to stay parallel with their source triangles. */
-	UPROPERTY(EditAnywhere, Category = Offset,
+	UPROPERTY(EditAnywhere, Category = ExtrusionOptions,
 		meta = (ClampMin = "1", EditConditionHides, EditCondition = "DirectionMode == EPolyEditPushPullModeOptions::SelectedTriangleNormalsEven"))
 	double MaxDistanceScaleFactor = 4.0;
 
@@ -226,7 +231,7 @@ public:
 
 	/** What axis to measure the extrusion distance along. */
 	UPROPERTY(EditAnywhere, Category = ExtrusionOptions, AdvancedDisplay, meta = (EditConditionHides, 
-		EditCondition = "DistanceMode == EPolyEditExtrudeDistanceMode::ClickInViewport"))
+		EditCondition = "DirectionMode != EPolyEditPushPullModeOptions::SingleDirection && DistanceMode == EPolyEditExtrudeDistanceMode::ClickInViewport"))
 	EPolyEditExtrudeDirection MeasureDirection = EPolyEditExtrudeDirection::SelectionNormal;
 
 	/**
diff --git a/Engine/Plugins/Runtime/MeshModelingToolset/Source/MeshModelingTools/Public/ToolActivities/PolyEditInsetOutsetActivity.h b/Engine/Plugins/Runtime/MeshModelingToolset/Source/MeshModelingTools/Public/ToolActivities/PolyEditInsetOutsetActivity.h
--- a/Engine/Plugins/Runtime/MeshModelingToolset/Source/MeshModelingTools/Public/ToolActivities/PolyEditInsetOutsetActivity.h
+++ b/Engine/Plugins/Runtime/MeshModelingToolset/Source/MeshModelingTools/Public/ToolActivities/PolyEditInsetOutsetActivity.h
@@ -20,21 +20,21 @@ class MESHMODELINGTOOLS_API UPolyEditInsetOutsetProperties : public UInteractive
 
 public:
 	/** Amount of smoothing applied to outset boundary */
-	UPROPERTY(EditAnywhere, Category = "Inset/Outset", 
+	UPROPERTY(EditAnywhere, Category = "Region", 
 		meta = (UIMin = "0.0", UIMax = "1.0", EditCondition = "bBoundaryOnly == false"))
 	float Softness = 0.5;
 
 	/** Controls whether outset operation will move interior vertices as well as border vertices */
-	UPROPERTY(EditAnywhere, Category = "Inset/Outset", AdvancedDisplay)
+	UPROPERTY(EditAnywhere, Category = "Region", AdvancedDisplay)
 	bool bBoundaryOnly = false;
 
 	/** Tweak area scaling when solving for interior vertices */
-	UPROPERTY(EditAnywhere, Category = "Inset/Outset", AdvancedDisplay, 
+	UPROPERTY(EditAnywhere, Category = "Region", AdvancedDisplay, 
 		meta = (UIMin = "0.0", UIMax = "1.0", EditCondition = "bBoundaryOnly == false"))
 	float AreaScale = true;
 
 	/** When insetting, determines whether vertices in inset region should be projected back onto input surface */
-	UPROPERTY(EditAnywhere, Category = "Inset/Outset", Meta = (EditCondition = "!bOutset", HideEditConditionToggle, EditConditionHides))
+	UPROPERTY(EditAnywhere, Category = "Region", Meta = (EditCondition = "!bOutset", HideEditConditionToggle, EditConditionHides))
 	bool bReproject = true;
 
 	//~ This is not user editable- it gets set by PolyEdit depending on whether the user clicks
