deprecate now-unused FGenericMeshSelection struct

diff --git a/Engine/Plugins/Runtime/MeshModelingToolset/Source/ModelingComponents/Public/Selection/PersistentMeshSelection.h b/Engine/Plugins/Runtime/MeshModelingToolset/Source/ModelingComponents/Public/Selection/PersistentMeshSelection.h
--- a/Engine/Plugins/Runtime/MeshModelingToolset/Source/ModelingComponents/Public/Selection/PersistentMeshSelection.h
+++ b/Engine/Plugins/Runtime/MeshModelingToolset/Source/ModelingComponents/Public/Selection/PersistentMeshSelection.h
@@ -23,9 +23,9 @@ PREDECLARE_USE_GEOMETRY_CLASS(FCompactMaps);
  * In addition, "Render" geometry is stored, which can be used by higher-level
  * code to draw the selection in some way (eg a selection highlight)
  * 
- * @warning this class is likely to change in the future
+ * @warning this class will be removed in the future
  */
-struct FGenericMeshSelection
+struct UE_DEPRECATED(5.4, "Use FGeometrySelection instead") FGenericMeshSelection
 {
 	// selection type
 	enum class ETopologyType
