stop silent-updating the uv layout preview wireframe when its visibility is turned off, so it actually sees the visibility change

diff --git a/Engine/Plugins/Runtime/MeshModelingToolset/Source/ModelingComponents/Private/Drawing/UVLayoutPreview.cpp b/Engine/Plugins/Runtime/MeshModelingToolset/Source/ModelingComponents/Private/Drawing/UVLayoutPreview.cpp
--- a/Engine/Plugins/Runtime/MeshModelingToolset/Source/ModelingComponents/Private/Drawing/UVLayoutPreview.cpp
+++ b/Engine/Plugins/Runtime/MeshModelingToolset/Source/ModelingComponents/Private/Drawing/UVLayoutPreview.cpp
@@ -100,7 +100,6 @@ void UUVLayoutPreview::SetVisible(bool bVisible)
 {
 	PreviewMesh->SetVisible(bVisible);
 	MeshElementsVisualizer->Settings->bVisible = bVisible;
-	MeshElementsVisualizer->Settings->SilentUpdateWatched();
 }
 
 
