fix convert tool losing materials

diff --git a/Engine/Plugins/Runtime/MeshModelingToolset/Source/ModelingComponents/Private/ConversionUtils/SceneComponentToDynamicMesh.cpp b/Engine/Plugins/Runtime/MeshModelingToolset/Source/ModelingComponents/Private/ConversionUtils/SceneComponentToDynamicMesh.cpp
--- a/Engine/Plugins/Runtime/MeshModelingToolset/Source/ModelingComponents/Private/ConversionUtils/SceneComponentToDynamicMesh.cpp
+++ b/Engine/Plugins/Runtime/MeshModelingToolset/Source/ModelingComponents/Private/ConversionUtils/SceneComponentToDynamicMesh.cpp
@@ -370,7 +370,7 @@ bool SceneComponentToDynamicMesh(USceneComponent* Component, const FToMeshOption
 	}
 	OutLocalToWorld = Component->GetComponentTransform();
 
-	auto GetPrimitiveComponentMaterials = [](UPrimitiveComponent* PrimComp, TArray<UMaterialInterface*> Materials)
+	auto GetPrimitiveComponentMaterials = [](UPrimitiveComponent* PrimComp, TArray<UMaterialInterface*>& Materials)
 	{
 		int32 NumMaterials = PrimComp->GetNumMaterials();
 		Materials.SetNum(NumMaterials);
