Make create-from-selected tools apply updated materials to both the component as well as the asset, so that the material reliably shows in the result while still applying to the asset

diff --git a/Engine/Plugins/Runtime/MeshModelingToolset/Source/ModelingComponents/Private/BaseTools/BaseCreateFromSelectedTool.cpp b/Engine/Plugins/Runtime/MeshModelingToolset/Source/ModelingComponents/Private/BaseTools/BaseCreateFromSelectedTool.cpp
--- a/Engine/Plugins/Runtime/MeshModelingToolset/Source/ModelingComponents/Private/BaseTools/BaseCreateFromSelectedTool.cpp
+++ b/Engine/Plugins/Runtime/MeshModelingToolset/Source/ModelingComponents/Private/BaseTools/BaseCreateFromSelectedTool.cpp
@@ -297,7 +297,10 @@ void UBaseCreateFromSelectedTool::UpdateAsset(const FDynamicMeshOpResult& Result
 
 	FComponentMaterialSet MaterialSet;
 	MaterialSet.Materials = GetOutputMaterials();
+	// apply updated materials to both asset and component, to ensure that they appear in the result
+	// TODO: consider adding a method to the material provider interface that more-directly handles this use case
 	Cast<IMaterialProvider>(UpdateTarget)->CommitMaterialSetUpdate(MaterialSet, true);
+	Cast<IMaterialProvider>(UpdateTarget)->CommitMaterialSetUpdate(MaterialSet, false);
 }
 
 
