Modeling Tools: Places Mesh Element Selection Gizmo in Object Frame Mode at selected component(s) pivot rather than at object's pivot

diff --git a/Engine/Plugins/Runtime/MeshModelingToolset/Source/ModelingComponents/Private/Selection/DynamicMeshSelector.cpp b/Engine/Plugins/Runtime/MeshModelingToolset/Source/ModelingComponents/Private/Selection/DynamicMeshSelector.cpp
--- a/Engine/Plugins/Runtime/MeshModelingToolset/Source/ModelingComponents/Private/Selection/DynamicMeshSelector.cpp
+++ b/Engine/Plugins/Runtime/MeshModelingToolset/Source/ModelingComponents/Private/Selection/DynamicMeshSelector.cpp
@@ -569,13 +569,20 @@ void FBaseDynamicMeshSelector::GetSelectionFrame(const FGeometrySelection& Selec
 
 	if (bTransformToWorld)
 	{
-		GetTargetFrame(SelectionFrame);
+		SelectionFrame.Transform(GetLocalToWorldTransform());
 	}
 }
 
-void FBaseDynamicMeshSelector::GetTargetFrame(FFrame3d& SelectionFrame)
+void FBaseDynamicMeshSelector::GetTargetFrame(const FGeometrySelection& Selection, FFrame3d& SelectionFrame)
 {
 	SelectionFrame.Transform(GetLocalToWorldTransform());
+	TQuaternion<double> TargetRotation = SelectionFrame.Rotation;
+
+	// Places gizmo at the selection's accumulated origin
+	GetSelectionFrame(Selection, SelectionFrame, true);
+
+	// Uses object rotation instead of accumulated normals from selection
+	SelectionFrame.Rotation = TargetRotation;
 }
 
 void FBaseDynamicMeshSelector::AccumulateSelectionBounds(const FGeometrySelection& Selection, FGeometrySelectionBounds& BoundsInOut, bool bTransformToWorld)
diff --git a/Engine/Plugins/Runtime/MeshModelingToolset/Source/ModelingComponents/Private/Selection/GeometrySelectionManager.cpp b/Engine/Plugins/Runtime/MeshModelingToolset/Source/ModelingComponents/Private/Selection/GeometrySelectionManager.cpp
--- a/Engine/Plugins/Runtime/MeshModelingToolset/Source/ModelingComponents/Private/Selection/GeometrySelectionManager.cpp
+++ b/Engine/Plugins/Runtime/MeshModelingToolset/Source/ModelingComponents/Private/Selection/GeometrySelectionManager.cpp
@@ -1027,7 +1027,7 @@ void UGeometrySelectionManager::GetTargetWorldFrame(FFrame3d& SelectionFrame) co
 		// only handling one target for now
 		//if (ActiveTargetReferences.Num() == 1)
 		TSharedPtr<FGeometrySelectionTarget> Target = ActiveTargetReferences[0];
-		Target->Selector->GetTargetFrame(SelectionFrame);
+		Target->Selector->GetTargetFrame(Target->Selection, SelectionFrame);
 	}
 }
 
diff --git a/Engine/Plugins/Runtime/MeshModelingToolset/Source/ModelingComponents/Public/Selection/DynamicMeshSelector.h b/Engine/Plugins/Runtime/MeshModelingToolset/Source/ModelingComponents/Public/Selection/DynamicMeshSelector.h
--- a/Engine/Plugins/Runtime/MeshModelingToolset/Source/ModelingComponents/Public/Selection/DynamicMeshSelector.h
+++ b/Engine/Plugins/Runtime/MeshModelingToolset/Source/ModelingComponents/Public/Selection/DynamicMeshSelector.h
@@ -116,7 +116,7 @@ public:
 	}
 
 	virtual void GetSelectionFrame(const FGeometrySelection& Selection, UE::Geometry::FFrame3d& SelectionFrame, bool bTransformToWorld) override;
-	virtual void GetTargetFrame(UE::Geometry::FFrame3d& SelectionFrame) override;
+	virtual void GetTargetFrame(const FGeometrySelection& Selection, UE::Geometry::FFrame3d& SelectionFrame) override;
 	virtual void AccumulateSelectionBounds(const FGeometrySelection& Selection, FGeometrySelectionBounds& BoundsInOut, bool bTransformToWorld) override;
 	virtual void AccumulateSelectionElements(const FGeometrySelection& Selection, FGeometrySelectionElements& Elements, bool bTransformToWorld, bool bIsForPreview) override;
 	
diff --git a/Engine/Plugins/Runtime/MeshModelingToolset/Source/ModelingComponents/Public/Selection/GeometrySelector.h b/Engine/Plugins/Runtime/MeshModelingToolset/Source/ModelingComponents/Public/Selection/GeometrySelector.h
--- a/Engine/Plugins/Runtime/MeshModelingToolset/Source/ModelingComponents/Public/Selection/GeometrySelector.h
+++ b/Engine/Plugins/Runtime/MeshModelingToolset/Source/ModelingComponents/Public/Selection/GeometrySelector.h
@@ -400,7 +400,7 @@ public:
 	/**
 	 * Compute a 3D frame for the Selector's target object
 	 */
-	virtual void GetTargetFrame(UE::Geometry::FFrame3d& SelectionFrameOut) = 0;
+	virtual void GetTargetFrame(const FGeometrySelection& Selection, UE::Geometry::FFrame3d& SelectionFrameOut) = 0;
 
 	/**
 	 * Accumulate the bounds of the provided Selection in the provided BoundsInOut. BoundsInOut is not cleared.
