Make modeling tools set SectionInfoMaps on newly-created static meshes. This fixes downstream issues with code that expects the maps to be populated, including the Lumen card build.

diff --git a/Engine/Plugins/Runtime/MeshModelingToolset/Source/ModelingComponentsEditorOnly/Private/AssetUtils/CreateStaticMeshUtil.cpp b/Engine/Plugins/Runtime/MeshModelingToolset/Source/ModelingComponentsEditorOnly/Private/AssetUtils/CreateStaticMeshUtil.cpp
--- a/Engine/Plugins/Runtime/MeshModelingToolset/Source/ModelingComponentsEditorOnly/Private/AssetUtils/CreateStaticMeshUtil.cpp
+++ b/Engine/Plugins/Runtime/MeshModelingToolset/Source/ModelingComponentsEditorOnly/Private/AssetUtils/CreateStaticMeshUtil.cpp
@@ -17,6 +17,7 @@
 #include "Materials/Material.h"
 #include "Materials/MaterialInterface.h"
 #include "AssetUtils/StaticMeshMaterialUtil.h"
+#include "StaticMeshAttributes.h"
 
 using namespace UE::AssetUtils;
 
@@ -175,6 +176,33 @@ UE::AssetUtils::ECreateStaticMeshResult UE::AssetUtils::CreateStaticMeshAsset(
 		CurNumValidSections++;
 	}
 
+	// populate section info map
+	FMeshSectionInfoMap SectionInfoMap;
+	for (int32 LODIndex = 0; LODIndex < UseNumSourceModels; ++LODIndex)
+	{
+		if (FMeshDescription* Mesh = NewStaticMesh->GetMeshDescription(LODIndex))
+		{
+			FStaticMeshConstAttributes MeshDescriptionAttributes(*Mesh);
+			TPolygonGroupAttributesConstRef<FName> MaterialSlotNames = MeshDescriptionAttributes.GetPolygonGroupMaterialSlotNames();
+			int32 SectionIndex = 0;
+			for (FPolygonGroupID PolygonGroupID : Mesh->PolygonGroups().GetElementIDs())
+			{
+				// Material index is either from the matching material slot name or the section index if that name is not found
+				int32 MaterialIndex = NewStaticMesh->GetStaticMaterials().IndexOfByPredicate(
+					[&MaterialSlotName = MaterialSlotNames[PolygonGroupID]](const FStaticMaterial& StaticMaterial) { return StaticMaterial.MaterialSlotName == MaterialSlotName; }
+				);
+				if (MaterialIndex == INDEX_NONE)
+				{
+					MaterialIndex = SectionIndex;
+				}
+				SectionInfoMap.Set(LODIndex, SectionIndex, FMeshSectionInfo(MaterialIndex));
+				SectionIndex++;
+			}
+		}
+	}
+	NewStaticMesh->GetSectionInfoMap().CopyFrom(SectionInfoMap);
+	NewStaticMesh->GetOriginalSectionInfoMap().CopyFrom(SectionInfoMap);
+
 	// Nanite options
 	NewStaticMesh->NaniteSettings.bEnabled = false;
 	if (Options.bGenerateNaniteEnabledMesh)
