make static mesh <-> dynamic mesh conversion be aware of section (polygon group ID) to material index mappings

diff --git a/Engine/Plugins/Runtime/MeshModelingToolset/Source/ModelingComponents/Private/ConversionUtils/DynamicMeshViaMeshDescriptionUtil.cpp b/Engine/Plugins/Runtime/MeshModelingToolset/Source/ModelingComponents/Private/ConversionUtils/DynamicMeshViaMeshDescriptionUtil.cpp
--- a/Engine/Plugins/Runtime/MeshModelingToolset/Source/ModelingComponents/Private/ConversionUtils/DynamicMeshViaMeshDescriptionUtil.cpp
+++ b/Engine/Plugins/Runtime/MeshModelingToolset/Source/ModelingComponents/Private/ConversionUtils/DynamicMeshViaMeshDescriptionUtil.cpp
@@ -14,6 +14,7 @@ FDynamicMesh3 UE::Geometry::GetDynamicMeshViaMeshDescription(
 	FDynamicMesh3 DynamicMesh;
 	FMeshDescriptionToDynamicMesh Converter;
 	Converter.bVIDsFromNonManifoldMeshDescriptionAttr = true;
+	Converter.SetPolygonGroupToMaterialIndexMap(MeshDescriptionProvider.GetPolygonGroupToMaterialIndexMap());
 	if (bRequestTangents)
 	{
 		FGetMeshParameters GetMeshParams;
@@ -43,6 +44,7 @@ void UE::Geometry::CommitDynamicMeshViaMeshDescription(
 	ConversionOptions.bTransformVtxColorsSRGBToLinear = CommitInfo.bTransformVertexColorsSRGBToLinear;
 
 	FDynamicMeshToMeshDescription Converter(ConversionOptions);
+	Converter.SetMaterialIDMapFromInverseMap(MeshDescriptionCommitter.GetPolygonGroupToMaterialIndexMap());
 	if (!CommitInfo.bTopologyChanged)
 	{
 		Converter.UpdateUsingConversionOptions(&Mesh, CurrentMeshDescription);
diff --git a/Engine/Plugins/Runtime/MeshModelingToolset/Source/ModelingComponents/Private/ConversionUtils/SceneComponentToDynamicMesh.cpp b/Engine/Plugins/Runtime/MeshModelingToolset/Source/ModelingComponents/Private/ConversionUtils/SceneComponentToDynamicMesh.cpp
--- a/Engine/Plugins/Runtime/MeshModelingToolset/Source/ModelingComponents/Private/ConversionUtils/SceneComponentToDynamicMesh.cpp
+++ b/Engine/Plugins/Runtime/MeshModelingToolset/Source/ModelingComponents/Private/ConversionUtils/SceneComponentToDynamicMesh.cpp
@@ -96,6 +96,76 @@ namespace Private::ConversionHelper
 		bool bRequestInstanceVertexColors = true;
 	};
 
+	// helper for the material ID remapping used for source LODs
+	// note: returns empty array if no remapping needed (or if not WITH_EDITOR)
+	TArray<int32> MapSectionToMaterialID(const UStaticMesh* Mesh, int32 SourceLOD, bool bHighResLOD)
+	{
+#if WITH_EDITOR
+		check(Mesh);
+		TMap<int32, int32> SectionToMaterial;
+		const int32 NumMaterials = Mesh->GetStaticMaterials().Num();
+		int32 MaxMaterialIndex = -1;
+		if (bHighResLOD)
+		{
+			// custom path for HiResSource, where the section info map isn't available so we use mesh description slot names
+			// (note that in practice this info seems to be incorrect for some meshes; prefer the section info map where available)
+			const FMeshDescription* MeshDescription = Mesh->GetHiResMeshDescription();
+			if (!MeshDescription)
+			{
+				// fall back to empty array (treated as identity map)
+				return TArray<int32>();
+			}
+			const FStaticMeshConstAttributes MeshDescriptionAttributes(*MeshDescription);
+			TPolygonGroupAttributesConstRef<FName> MaterialSlotNames = MeshDescriptionAttributes.GetPolygonGroupMaterialSlotNames();
+			int32 SectionIndex = 0;
+			for (FPolygonGroupID PolygonGroupID : MeshDescription->PolygonGroups().GetElementIDs())
+			{
+				int32 MaterialIndex = Mesh->GetStaticMaterials().IndexOfByPredicate(
+					[&MaterialSlotName = MaterialSlotNames[PolygonGroupID]](const FStaticMaterial& StaticMaterial) { return StaticMaterial.MaterialSlotName == MaterialSlotName; }
+				);
+				if (MaterialIndex != INDEX_NONE)
+				{
+					SectionToMaterial.Add(SectionIndex, MaterialIndex);
+					MaxMaterialIndex = FMath::Max(MaterialIndex, MaxMaterialIndex);
+				}
+				++SectionIndex;
+			}
+		}
+		else
+		{
+			int32 UseLOD = SourceLOD;
+			const FMeshSectionInfoMap& SectionMap = Mesh->GetSectionInfoMap();
+			int32 LODSectionNum = SectionMap.GetSectionNumber(UseLOD);
+			TArray<int32> Result;
+			for (int32 SectionIndex = 0; SectionIndex < LODSectionNum; ++SectionIndex)
+			{
+				if (SectionMap.IsValidSection(UseLOD, SectionIndex))
+				{
+					int32 MaterialIndex = SectionMap.Get(UseLOD, SectionIndex).MaterialIndex;
+					SectionToMaterial.Add(SectionIndex, MaterialIndex);
+					MaxMaterialIndex = FMath::Max(MaterialIndex, MaxMaterialIndex);
+				}
+			}
+		}
+
+		MaxMaterialIndex = FMath::Min(MaxMaterialIndex, NumMaterials - 1);
+		TArray<int32> Result;
+		Result.SetNumUninitialized(MaxMaterialIndex + 1);
+		// Fill in identity mapping first to cover any unmapped indices
+		for (int32 Idx = 0; Idx < Result.Num(); ++Idx)
+		{
+			Result[Idx] = Idx;
+		}
+		for (TPair<int32, int32> SectionMaterial : SectionToMaterial)
+		{
+			Result[SectionMaterial.Key] = FMath::Clamp(SectionMaterial.Value, 0, NumMaterials - 1);
+		}
+		return Result;
+#else
+		return TArray<int32>();
+#endif
+	}
+
 	static bool CopyMeshFromStaticMesh_SourceData(
 		UStaticMesh* FromStaticMeshAsset,
 		FStaticMeshConversionOptions AssetOptions,
@@ -132,6 +202,8 @@ namespace Private::ConversionHelper
 		const FMeshDescription* SourceMesh = nullptr;
 		const FMeshBuildSettings* BuildSettings = nullptr;
 
+		TArray<int32> PolygonGroupToMaterialMap = GetPolygonGroupToMaterialIndexMap(FromStaticMeshAsset, LODType, LODIndex);
+
 		if ((LODType == EMeshLODType::HiResSourceModel) ||
 			(LODType == EMeshLODType::MaxAvailable && FromStaticMeshAsset->IsHiResMeshDescriptionValid()))
 		{
@@ -200,6 +272,7 @@ namespace Private::ConversionHelper
 		}
 
 		FMeshDescriptionToDynamicMesh Converter;
+		Converter.SetPolygonGroupToMaterialIndexMap(PolygonGroupToMaterialMap);
 		Converter.Convert(SourceMesh, OutMesh, AssetOptions.bRequestTangents);
 
 		bSuccess = true;
@@ -353,6 +426,25 @@ namespace Private::ConversionHelper
 	}
 }
 
+TArray<int32> GetPolygonGroupToMaterialIndexMap(const UStaticMesh* StaticMesh, EMeshLODType LODType, int32 LODIndex)
+{
+#if WITH_EDITOR
+	if (LODType == EMeshLODType::RenderData)
+	{
+		// don't need to remap material indices for render LODs
+		return TArray<int32>();
+	}
+	// map the 'max available' lod type
+	if (LODType == EMeshLODType::MaxAvailable)
+	{
+		LODType = StaticMesh->IsHiResMeshDescriptionValid() ? EMeshLODType::HiResSourceModel : EMeshLODType::SourceModel;
+		LODIndex = 0;
+	}
+	return Private::ConversionHelper::MapSectionToMaterialID(StaticMesh, LODIndex, LODType == EMeshLODType::HiResSourceModel);
+#else
+	return TArray<int32>();
+#endif
+}
 
 bool SceneComponentToDynamicMesh(USceneComponent* Component, const FToMeshOptions& Options, bool bTransformToWorld, 
 	Geometry::FDynamicMesh3& OutMesh, FTransform& OutLocalToWorld, FText& OutErrorMessage,
diff --git a/Engine/Plugins/Runtime/MeshModelingToolset/Source/ModelingComponents/Private/ModelingToolTargetUtil.cpp b/Engine/Plugins/Runtime/MeshModelingToolset/Source/ModelingComponents/Private/ModelingToolTargetUtil.cpp
--- a/Engine/Plugins/Runtime/MeshModelingToolset/Source/ModelingComponents/Private/ModelingToolTargetUtil.cpp
+++ b/Engine/Plugins/Runtime/MeshModelingToolset/Source/ModelingComponents/Private/ModelingToolTargetUtil.cpp
@@ -236,7 +236,8 @@ FDynamicMesh3 UE::ToolTarget::GetDynamicMeshCopy(UToolTarget* Target, bool bWant
 	if (MeshDescriptionProvider)
 	{
 		FMeshDescriptionToDynamicMesh Converter;
-		Converter.bVIDsFromNonManifoldMeshDescriptionAttr= true;
+		Converter.bVIDsFromNonManifoldMeshDescriptionAttr = true;
+		Converter.SetPolygonGroupToMaterialIndexMap(MeshDescriptionProvider->GetPolygonGroupToMaterialIndexMap());
 		if (bWantMeshTangents)
 		{
 			FGetMeshParameters GetMeshParams;
@@ -304,6 +305,7 @@ UE::ToolTarget::EDynamicMeshUpdateResult UE::ToolTarget::CommitMeshDescriptionUp
 	}
 
 	FDynamicMeshToMeshDescription Converter;
+	Converter.SetMaterialIDMapFromInverseMap(MeshDescriptionCommitter->GetPolygonGroupToMaterialIndexMap());
 	FMeshDescription ConvertedMesh;
 	if (bHaveModifiedTopology)
 	{
@@ -395,6 +397,7 @@ UE::ToolTarget::EDynamicMeshUpdateResult UE::ToolTarget::CommitDynamicMeshUpdate
 	{
 		FMeshDescription ConvertedMesh;
 		FDynamicMeshToMeshDescription Converter(ConversionOptions);
+		Converter.SetMaterialIDMapFromInverseMap(MeshDescriptionCommitter->GetPolygonGroupToMaterialIndexMap());
 		if (!bHaveModifiedTopology)
 		{
 			Converter.UpdateUsingConversionOptions(&UpdatedMesh, ConvertedMesh);
@@ -446,15 +449,15 @@ UE::ToolTarget::EDynamicMeshUpdateResult UE::ToolTarget::CommitDynamicMeshUVUpda
 	FMeshDescription NewMeshDescription = UE::ToolTarget::GetMeshDescriptionCopy(Target);
 	bool bVerticesOnly = false;
 	bool bAttributesOnly = true;
+	FDynamicMeshToMeshDescription Converter;
+	Converter.SetMaterialIDMapFromInverseMap(MeshDescriptionCommitter->GetPolygonGroupToMaterialIndexMap());
 	if (FDynamicMeshToMeshDescription::HaveMatchingElementCounts(UpdatedMesh, &NewMeshDescription, bVerticesOnly, bAttributesOnly))
 	{
-		FDynamicMeshToMeshDescription Converter;
 		Converter.UpdateAttributes(UpdatedMesh, NewMeshDescription, false, false, true/*update uvs*/);
 	}
 	else
 	{
 		// must have been duplicate tris in the mesh description; we can't count on 1-to-1 mapping of TriangleIDs.  Just convert 
-		FDynamicMeshToMeshDescription Converter;
 		Converter.Convert(UpdatedMesh, NewMeshDescription);
 	}
 
@@ -497,15 +500,15 @@ UE::ToolTarget::EDynamicMeshUpdateResult UE::ToolTarget::CommitDynamicMeshNormal
 	FMeshDescription NewMeshDescription = UE::ToolTarget::GetMeshDescriptionCopy(Target);
 	bool bVerticesOnly = false;
 	bool bAttributesOnly = true;
+	FDynamicMeshToMeshDescription Converter;
+	Converter.SetMaterialIDMapFromInverseMap(MeshDescriptionCommitter->GetPolygonGroupToMaterialIndexMap());
 	if (FDynamicMeshToMeshDescription::HaveMatchingElementCounts(UpdatedMesh, &NewMeshDescription, bVerticesOnly, bAttributesOnly))
 	{
-		FDynamicMeshToMeshDescription Converter;
 		Converter.UpdateAttributes(UpdatedMesh, NewMeshDescription, true, bUpdateTangents, false);
 	}
 	else
 	{
 		// must have been duplicate tris in the mesh description; we can't count on 1-to-1 mapping of TriangleIDs.  Just convert 
-		FDynamicMeshToMeshDescription Converter;
 		Converter.Convert(UpdatedMesh, NewMeshDescription);
 	}
 
diff --git a/Engine/Plugins/Runtime/MeshModelingToolset/Source/ModelingComponents/Public/ConversionUtils/SceneComponentToDynamicMesh.h b/Engine/Plugins/Runtime/MeshModelingToolset/Source/ModelingComponents/Public/ConversionUtils/SceneComponentToDynamicMesh.h
--- a/Engine/Plugins/Runtime/MeshModelingToolset/Source/ModelingComponents/Public/ConversionUtils/SceneComponentToDynamicMesh.h
+++ b/Engine/Plugins/Runtime/MeshModelingToolset/Source/ModelingComponents/Public/ConversionUtils/SceneComponentToDynamicMesh.h
@@ -61,6 +61,14 @@ namespace Conversion {
 		bool bWantInstanceColors = false;
 	};
 
+	/**
+	 * Helper to find the Polygon Group ID to Material Index remapping for source LOD MeshDescription data.
+	 * Returns an empty array for runtime uses, render LODs, and other cases where an identity mapping can be used.
+	 * 
+	 * @return a mapping from polygon groups to material indices for the given LODType and LODIndex, or an empty array if the identity map should be used
+	 */
+	TArray<int32> MODELINGCOMPONENTS_API GetPolygonGroupToMaterialIndexMap(const UStaticMesh* StaticMesh, EMeshLODType LODType, int32 LODIndex);
+
 	/**
 	* Converts geometry from various types of Scene Component to Dynamic Mesh, or reports that it cannot.
 	* 
@@ -85,5 +93,5 @@ namespace Conversion {
 	 */
 	bool MODELINGCOMPONENTS_API CanConvertSceneComponentToDynamicMesh(USceneComponent* SceneComponent);
 
-} // end namespace Geometry
+} // end namespace Conversion
 } // end namespace UE
\ No newline at end of file
diff --git a/Engine/Plugins/Runtime/MeshModelingToolset/Source/ModelingComponentsEditorOnly/Private/ToolTargets/StaticMeshComponentToolTarget.cpp b/Engine/Plugins/Runtime/MeshModelingToolset/Source/ModelingComponentsEditorOnly/Private/ToolTargets/StaticMeshComponentToolTarget.cpp
--- a/Engine/Plugins/Runtime/MeshModelingToolset/Source/ModelingComponentsEditorOnly/Private/ToolTargets/StaticMeshComponentToolTarget.cpp
+++ b/Engine/Plugins/Runtime/MeshModelingToolset/Source/ModelingComponentsEditorOnly/Private/ToolTargets/StaticMeshComponentToolTarget.cpp
@@ -149,6 +149,16 @@ const FMeshDescription* UStaticMeshComponentToolTarget::GetMeshDescription(const
 	return nullptr;
 }
 
+TArray<int32> UStaticMeshComponentToolTarget::GetPolygonGroupToMaterialIndexMap() const
+{
+	if (IsValid())
+	{
+		UStaticMesh* StaticMesh = Cast<UStaticMeshComponent>(Component)->GetStaticMesh();
+		return UStaticMeshToolTarget::MapSectionToMaterialID(StaticMesh, EditingLOD);
+	}
+	return TArray<int32>();
+}
+
 FMeshDescription UStaticMeshComponentToolTarget::GetEmptyMeshDescription()
 {
 	FMeshDescription EmptyMeshDescription;
diff --git a/Engine/Plugins/Runtime/MeshModelingToolset/Source/ModelingComponentsEditorOnly/Private/ToolTargets/StaticMeshToolTarget.cpp b/Engine/Plugins/Runtime/MeshModelingToolset/Source/ModelingComponentsEditorOnly/Private/ToolTargets/StaticMeshToolTarget.cpp
--- a/Engine/Plugins/Runtime/MeshModelingToolset/Source/ModelingComponentsEditorOnly/Private/ToolTargets/StaticMeshToolTarget.cpp
+++ b/Engine/Plugins/Runtime/MeshModelingToolset/Source/ModelingComponentsEditorOnly/Private/ToolTargets/StaticMeshToolTarget.cpp
@@ -4,6 +4,7 @@
 
 #include "AssetUtils/MeshDescriptionUtil.h"
 #include "ConversionUtils/DynamicMeshViaMeshDescriptionUtil.h"
+#include "ConversionUtils/SceneComponentToDynamicMesh.h"
 #include "DynamicMesh/DynamicMesh3.h"
 #include "Engine/StaticMesh.h"
 #include "MaterialDomain.h"
@@ -275,6 +276,36 @@ const FMeshDescription* UStaticMeshToolTarget::GetMeshDescription(const FGetMesh
 	return nullptr;
 }
 
+TArray<int32> UStaticMeshToolTarget::GetPolygonGroupToMaterialIndexMap() const
+{
+	if (IsValid())
+	{
+		return UStaticMeshToolTarget::MapSectionToMaterialID(StaticMesh.Get(), EditingLOD);
+	}
+	return TArray<int32>();
+}
+
+TArray<int32> UStaticMeshToolTarget::MapSectionToMaterialID(const UStaticMesh* Mesh, EMeshLODIdentifier EditingLOD)
+{
+	UE::Conversion::EMeshLODType LODType = [EditingLOD]() {
+		switch (EditingLOD)
+		{
+		case EMeshLODIdentifier::MaxQuality:
+			return UE::Conversion::EMeshLODType::MaxAvailable;
+		case EMeshLODIdentifier::HiResSource:
+			return UE::Conversion::EMeshLODType::HiResSourceModel;
+		default:
+			return UE::Conversion::EMeshLODType::SourceModel;
+		}
+	}();
+	int32 UseLOD = (int32)EditingLOD;
+	if (UseLOD > 7) // map the special LODs (default, hi res, max) to index 0
+	{
+		UseLOD = 0;
+	}
+	return UE::Conversion::GetPolygonGroupToMaterialIndexMap(Mesh, LODType, UseLOD);
+}
+
 const FMeshDescription* UStaticMeshToolTarget::GetMeshDescriptionWithScaleApplied(UStaticMesh* StaticMesh, int32 UseLOD, FMeshDescriptionCache& CachedMeshDescriptions)
 {
 	static FMeshDescription EmptyMeshDescription;
diff --git a/Engine/Plugins/Runtime/MeshModelingToolset/Source/ModelingComponentsEditorOnly/Public/ToolTargets/StaticMeshComponentToolTarget.h b/Engine/Plugins/Runtime/MeshModelingToolset/Source/ModelingComponentsEditorOnly/Public/ToolTargets/StaticMeshComponentToolTarget.h
--- a/Engine/Plugins/Runtime/MeshModelingToolset/Source/ModelingComponentsEditorOnly/Public/ToolTargets/StaticMeshComponentToolTarget.h
+++ b/Engine/Plugins/Runtime/MeshModelingToolset/Source/ModelingComponentsEditorOnly/Public/ToolTargets/StaticMeshComponentToolTarget.h
@@ -52,8 +52,9 @@ public:
 	virtual bool SupportsLODs() const override { return true; }
 	virtual TArray<EMeshLODIdentifier> GetAvailableLODs(bool bSkipAutoGenerated = true) const override;
 	virtual EMeshLODIdentifier GetMeshDescriptionLOD() const override { return EditingLOD; }
+	virtual TArray<int32> GetPolygonGroupToMaterialIndexMap() const override;
 
-	// IMeshDescritpionCommitter implementation
+	// IMeshDescriptionCommitter implementation
 	virtual void CommitMeshDescription(const FCommitter& Committer, const FCommitMeshParameters& CommitParams = FCommitMeshParameters()) override;
 	using IMeshDescriptionCommitter::CommitMeshDescription; // unhide the other overload
 
diff --git a/Engine/Plugins/Runtime/MeshModelingToolset/Source/ModelingComponentsEditorOnly/Public/ToolTargets/StaticMeshToolTarget.h b/Engine/Plugins/Runtime/MeshModelingToolset/Source/ModelingComponentsEditorOnly/Public/ToolTargets/StaticMeshToolTarget.h
--- a/Engine/Plugins/Runtime/MeshModelingToolset/Source/ModelingComponentsEditorOnly/Public/ToolTargets/StaticMeshToolTarget.h
+++ b/Engine/Plugins/Runtime/MeshModelingToolset/Source/ModelingComponentsEditorOnly/Public/ToolTargets/StaticMeshToolTarget.h
@@ -55,6 +55,10 @@ public:
 	virtual bool SupportsLODs() const override { return true; }
 	virtual TArray<EMeshLODIdentifier> GetAvailableLODs(bool bSkipAutoGenerated = true) const override;
 	virtual EMeshLODIdentifier GetMeshDescriptionLOD() const override { return EditingLOD; }
+	virtual TArray<int32> GetPolygonGroupToMaterialIndexMap() const override;
+
+	// helper to get the material mapping returned by GetPolygonGroupToMaterialIndexMap for any static mesh
+	static TArray<int32> MapSectionToMaterialID(const UStaticMesh* Mesh, EMeshLODIdentifier EditingLOD);
 
 	// IMeshDescritpionCommitter implementation
 	virtual void CommitMeshDescription(const FCommitter& Committer, const FCommitMeshParameters& CommitParams = FCommitMeshParameters()) override;
diff --git a/Engine/Source/Runtime/InteractiveToolsFramework/Public/TargetInterfaces/MeshDescriptionCommitter.h b/Engine/Source/Runtime/InteractiveToolsFramework/Public/TargetInterfaces/MeshDescriptionCommitter.h
--- a/Engine/Source/Runtime/InteractiveToolsFramework/Public/TargetInterfaces/MeshDescriptionCommitter.h
+++ b/Engine/Source/Runtime/InteractiveToolsFramework/Public/TargetInterfaces/MeshDescriptionCommitter.h
@@ -61,7 +61,13 @@ public:
 		return bSuccess;
 	}
 
-
+	/**
+	 * @return a mapping from material ID to polygon group ID. If empty, an identity mapping is assumed.
+	 */
+	virtual TArray<int32> GetPolygonGroupToMaterialIndexMap() const
+	{
+		return TArray<int32>();
+	}
 
 protected:
 
diff --git a/Engine/Source/Runtime/InteractiveToolsFramework/Public/TargetInterfaces/MeshDescriptionProvider.h b/Engine/Source/Runtime/InteractiveToolsFramework/Public/TargetInterfaces/MeshDescriptionProvider.h
--- a/Engine/Source/Runtime/InteractiveToolsFramework/Public/TargetInterfaces/MeshDescriptionProvider.h
+++ b/Engine/Source/Runtime/InteractiveToolsFramework/Public/TargetInterfaces/MeshDescriptionProvider.h
@@ -73,4 +73,12 @@ public:
 		return false;
 	}
 
+	/**
+	 * @return a mapping from section ID to material ID. If empty, an identity mapping is assumed.
+	 */
+	virtual TArray<int32> GetPolygonGroupToMaterialIndexMap() const
+	{
+		return TArray<int32>();
+	}
+
 };
diff --git a/Engine/Source/Runtime/MeshConversion/Private/DynamicMeshToMeshDescription.cpp b/Engine/Source/Runtime/MeshConversion/Private/DynamicMeshToMeshDescription.cpp
--- a/Engine/Source/Runtime/MeshConversion/Private/DynamicMeshToMeshDescription.cpp
+++ b/Engine/Source/Runtime/MeshConversion/Private/DynamicMeshToMeshDescription.cpp
@@ -371,6 +371,29 @@ void FDynamicMeshToMeshDescription::Convert(const FDynamicMesh3* MeshIn, FMeshDe
 }
 
 
+void FDynamicMeshToMeshDescription::SetMaterialIDMapFromInverseMap(TArrayView<const int32> PolygroupIDToMaterialIDMap)
+{
+	int32 MaxMatID = -1;
+	for (int32 MatID : PolygroupIDToMaterialIDMap)
+	{
+		MaxMatID = FMath::Max(MaxMatID, MatID);
+	}
+	MaterialIDToPolygroupIDMap.SetNumUninitialized(MaxMatID + 1);
+	for (int32 Idx = 0; Idx < MaterialIDToPolygroupIDMap.Num(); ++Idx)
+	{
+		MaterialIDToPolygroupIDMap[Idx] = Idx;
+	}
+	for (int32 PolyGroupIdx = 0; PolyGroupIdx < PolygroupIDToMaterialIDMap.Num(); ++PolyGroupIdx)
+	{
+		int32 MatID = PolygroupIDToMaterialIDMap[PolyGroupIdx];
+		if (MaterialIDToPolygroupIDMap.IsValidIndex(MatID))
+		{
+			MaterialIDToPolygroupIDMap[MatID] = PolyGroupIdx;
+		}
+	}
+}
+
+
 bool FDynamicMeshToMeshDescription::HaveMatchingElementCounts(const FDynamicMesh3* DynamicMesh, const FMeshDescription* MeshDescription, bool bVerticesOnly, bool bAttributesOnly)
 {
 	bool bVerticesMatch = DynamicMesh->IsCompactV() && DynamicMesh->VertexCount() == MeshDescription->Vertices().Num();
@@ -542,6 +565,10 @@ void FDynamicMeshToMeshDescription::Convert_SharedInstances(const FDynamicMesh3*
 		{
 			int32 MaterialID;
 			MaterialIDAttrib->GetValue(TriID, &MaterialID);
+			if (MaterialIDToPolygroupIDMap.IsValidIndex(MaterialID))
+			{
+				MaterialID = MaterialIDToPolygroupIDMap[MaterialID];
+			}
 			MaxPolygonGroupID = FMath::Max(MaterialID, MaxPolygonGroupID);
 		}
 		if (MaxPolygonGroupID == 0)
@@ -648,6 +675,10 @@ void FDynamicMeshToMeshDescription::Convert_SharedInstances(const FDynamicMesh3*
 		{
 			int32 MaterialID;
 			MaterialIDAttrib->GetValue(TriID, &MaterialID);
+			if (MaterialIDToPolygroupIDMap.IsValidIndex(MaterialID))
+			{
+				MaterialID = MaterialIDToPolygroupIDMap[MaterialID];
+			}
 			UsePolygonGroupID = FPolygonGroupID(MaterialID);
 		}
 
@@ -874,6 +905,10 @@ void FDynamicMeshToMeshDescription::Convert_NoSharedInstances(const FDynamicMesh
 		{
 			int32 MaterialID;
 			MaterialIDAttrib->GetValue(TriID, &MaterialID);
+			if (MaterialIDToPolygroupIDMap.IsValidIndex(MaterialID))
+			{
+				MaterialID = MaterialIDToPolygroupIDMap[MaterialID];
+			}
 			MaxPolygonGroupID = FMath::Max(MaterialID, MaxPolygonGroupID);
 		}
 		if (MaxPolygonGroupID == 0)
@@ -916,6 +951,10 @@ void FDynamicMeshToMeshDescription::Convert_NoSharedInstances(const FDynamicMesh
 		{
 			int32 MaterialID;
 			MaterialIDAttrib->GetValue(TriID, &MaterialID);
+			if (MaterialIDToPolygroupIDMap.IsValidIndex(MaterialID))
+			{
+				MaterialID = MaterialIDToPolygroupIDMap[MaterialID];
+			}
 			UsePolygonGroupID = FPolygonGroupID(MaterialID);
 		}
 
diff --git a/Engine/Source/Runtime/MeshConversion/Private/MeshDescriptionToDynamicMesh.cpp b/Engine/Source/Runtime/MeshConversion/Private/MeshDescriptionToDynamicMesh.cpp
--- a/Engine/Source/Runtime/MeshConversion/Private/MeshDescriptionToDynamicMesh.cpp
+++ b/Engine/Source/Runtime/MeshConversion/Private/MeshDescriptionToDynamicMesh.cpp
@@ -820,7 +820,12 @@ void FMeshDescriptionToDynamicMesh::Convert(const FMeshDescription* MeshIn, FDyn
 				for (int32 TriangleID : MeshOut.TriangleIndicesItr())
 				{
 					const FTriData& TriData = AddedTriangles[TriangleID];
-					MaterialIDAttrib->SetValue(TriangleID, &TriData.PolygonGroupID);
+					int32 MaterialIndex = TriData.PolygonGroupID;
+					if (PolygonGroupToMaterialIndexMap.IsValidIndex(TriData.PolygonGroupID))
+					{
+						MaterialIndex = PolygonGroupToMaterialIndexMap[TriData.PolygonGroupID];
+					}
+					MaterialIDAttrib->SetValue(TriangleID, &MaterialIndex);
 				}
 			});
 			Pending.Add(MoveTemp(MaterialTask));
diff --git a/Engine/Source/Runtime/MeshConversion/Public/DynamicMeshToMeshDescription.h b/Engine/Source/Runtime/MeshConversion/Public/DynamicMeshToMeshDescription.h
--- a/Engine/Source/Runtime/MeshConversion/Public/DynamicMeshToMeshDescription.h
+++ b/Engine/Source/Runtime/MeshConversion/Public/DynamicMeshToMeshDescription.h
@@ -30,6 +30,9 @@ public:
 	/** General settings for conversions to mesh description */
 	FConversionToMeshDescriptionOptions ConversionOptions;
 
+	// Set MaterialID to Polygroup mapping via the reverse mapping
+	void MESHCONVERSION_API SetMaterialIDMapFromInverseMap(TArrayView<const int32> PolygroupIDToMaterialIDMap);
+
 	FDynamicMeshToMeshDescription()
 	{
 	}
@@ -182,4 +185,8 @@ protected:
 	 * @param Color color to transform
 	 */
 	MESHCONVERSION_API void ApplyVertexColorTransform(FVector4f& Color) const;
+	
+private:
+	/** Optional remapping from dynamic mesh material ID to desired mesh description polygroup IDs. IDs not found in mapping will be mapped to themselves. */
+	TArray<int32> MaterialIDToPolygroupIDMap;
 };
diff --git a/Engine/Source/Runtime/MeshConversion/Public/MeshDescriptionToDynamicMesh.h b/Engine/Source/Runtime/MeshConversion/Public/MeshDescriptionToDynamicMesh.h
--- a/Engine/Source/Runtime/MeshConversion/Public/MeshDescriptionToDynamicMesh.h
+++ b/Engine/Source/Runtime/MeshConversion/Public/MeshDescriptionToDynamicMesh.h
@@ -42,6 +42,12 @@ public:
 	/** Should Vertex Colors of MeshDescription be transformed from Linear to SRGB */
 	bool bTransformVertexColorsLinearToSRGB = true;
 
+	/** Set an optional mapping from mesh description polygon group to dynamic mesh material indices; if unset, an identity map is assumed */
+	void SetPolygonGroupToMaterialIndexMap(const TArray<int32>& PolygonGroupToMaterialIndexMapIn)
+	{
+		PolygonGroupToMaterialIndexMap = PolygonGroupToMaterialIndexMapIn;
+	}
+
 	/** map from DynamicMesh triangle ID to MeshDescription FTriangleID*/
 	TArray<FTriangleID> TriIDMap;
 
@@ -109,4 +115,9 @@ protected:
 	 * @param Color color to transform
 	 */
 	MESHCONVERSION_API void ApplyVertexColorTransform(FVector4f& Color) const;
+
+private:
+
+	/** Mapping from mesh description polygon group element IDs to material indices */
+	TArray<int32> PolygonGroupToMaterialIndexMap;
 };
