Modeling Tools: Adds Geometry Selection Support for the Split Tool

diff --git a/Engine/Plugins/Experimental/MeshModelingToolsetExp/Source/MeshModelingToolsExp/Private/SplitMeshesTool.cpp b/Engine/Plugins/Experimental/MeshModelingToolsetExp/Source/MeshModelingToolsExp/Private/SplitMeshesTool.cpp
--- a/Engine/Plugins/Experimental/MeshModelingToolsetExp/Source/MeshModelingToolsExp/Private/SplitMeshesTool.cpp
+++ b/Engine/Plugins/Experimental/MeshModelingToolsetExp/Source/MeshModelingToolsExp/Private/SplitMeshesTool.cpp
@@ -12,6 +12,7 @@
 #include "Selections/MeshConnectedComponents.h"
 #include "DynamicMesh/MeshTransforms.h"
 #include "DynamicSubmesh3.h"
+#include "Selections/GeometrySelectionUtil.h"
 
 #include "TargetInterfaces/MeshDescriptionProvider.h"
 #include "TargetInterfaces/PrimitiveComponentBackedTarget.h"
@@ -35,7 +36,7 @@ const FToolTargetTypeRequirements& USplitMeshesToolBuilder::GetTargetRequirement
 	return TypeRequirements;
 }
 
-UMultiSelectionMeshEditingTool* USplitMeshesToolBuilder::CreateNewTool(const FToolBuilderState& SceneState) const
+UMultiTargetWithSelectionTool* USplitMeshesToolBuilder::CreateNewTool(const FToolBuilderState& SceneState) const
 {
 	return NewObject<USplitMeshesTool>(SceneState.ToolManager);
 }
@@ -169,9 +170,21 @@ void USplitMeshesTool::UpdateSplitMeshes()
 		const TArray<UMaterialInterface*> SourceMaterials = SourceMeshes[si].Materials;
 
 		FMeshConnectedComponents MeshComponents(SourceMesh);
-		MeshComponents.FindConnectedTriangles();
-		int32 NumComponents = MeshComponents.Num();
-
+		
+		int32 NumComponents;
+		const bool bMeshHasGeometrySelection = HasGeometrySelection(si);
+		if (bMeshHasGeometrySelection)
+		{
+			// when there is a geometry selection, ignore any computations the Split Tool would normally have done to
+			// decide where to split the mesh; instead split into 2 meshes regardless: the selected geometry, and everything else
+			NumComponents = 2;
+		}
+		else
+		{
+			MeshComponents.FindConnectedTriangles();
+			NumComponents = MeshComponents.Num();
+		}
+		
 		if (NumComponents < 2)
 		{
 			SplitInfo.bNoComponents = true;
@@ -185,7 +198,45 @@ void USplitMeshesTool::UpdateSplitMeshes()
 		SplitInfo.Origins.SetNum(NumComponents);
 		for (int32 k = 0; k < NumComponents; ++k)
 		{
-			FDynamicSubmesh3 SubmeshCalc(SourceMesh, MeshComponents[k].Indices);
+			FDynamicSubmesh3 SubmeshCalc;
+			if (bMeshHasGeometrySelection)
+			{
+				// retrieve the triangles in the current selection
+				// when using Edge or Vertex selection mode, any triangle touching the selected edges/vertices will be included
+				TSet<int> SelectionTriangles;
+				const FGeometrySelection& InputSelection = GetGeometrySelection(si);
+				
+				UE::Geometry::EnumerateSelectionTriangles(InputSelection, *SourceMesh,
+			[&](int32 TriangleID) { SelectionTriangles.Add(TriangleID); });
+
+				TArray<int> SelectionTrianglesArray = SelectionTriangles.Array();
+				
+				if (k == 0) // component made of the selected triangles
+				{
+					SubmeshCalc = FDynamicSubmesh3(SourceMesh, SelectionTrianglesArray);
+				}
+				else if (k == 1) // component made of the rest of the mesh (unselected triangles)
+				{
+					TArray<int> NonSelectedTriangles;
+					for (int TID = 0; TID < SourceMesh->MaxTriangleID(); TID++)
+					{
+						if (SourceMesh->IsTriangle(TID) && !SelectionTriangles.Contains(TID))
+						{
+							NonSelectedTriangles.Add(TID);
+						}
+					}
+					SubmeshCalc = FDynamicSubmesh3(SourceMesh, NonSelectedTriangles);
+				}
+			}
+			else
+			{
+				// if statement should always be true- components should always have been calculated & populated when there's no geometry selection
+				if (!MeshComponents.Components.IsEmpty())
+				{
+					SubmeshCalc = FDynamicSubmesh3(SourceMesh, MeshComponents[k].Indices);
+				}
+			}
+
 			FDynamicMesh3& Submesh = SubmeshCalc.GetSubmesh();
 			TArray<UMaterialInterface*> NewMaterials;
 
diff --git a/Engine/Plugins/Experimental/MeshModelingToolsetExp/Source/MeshModelingToolsExp/Public/SplitMeshesTool.h b/Engine/Plugins/Experimental/MeshModelingToolsetExp/Source/MeshModelingToolsExp/Public/SplitMeshesTool.h
--- a/Engine/Plugins/Experimental/MeshModelingToolsetExp/Source/MeshModelingToolsExp/Public/SplitMeshesTool.h
+++ b/Engine/Plugins/Experimental/MeshModelingToolsetExp/Source/MeshModelingToolsExp/Public/SplitMeshesTool.h
@@ -3,7 +3,7 @@
 #pragma once
 
 #include "CoreMinimal.h"
-#include "BaseTools/MultiSelectionMeshEditingTool.h"
+#include "BaseTools/MultiTargetWithSelectionTool.h"
 #include "InteractiveToolBuilder.h"
 #include "PropertySets/CreateMeshObjectTypeProperties.h"
 #include "SplitMeshesTool.generated.h"
@@ -12,11 +12,13 @@ class UMaterialInterface;
 
 
 UCLASS()
-class MESHMODELINGTOOLSEXP_API USplitMeshesToolBuilder : public UMultiSelectionMeshEditingToolBuilder
+class MESHMODELINGTOOLSEXP_API USplitMeshesToolBuilder : public UMultiTargetWithSelectionToolBuilder
 {
 	GENERATED_BODY()
 public:
-	virtual UMultiSelectionMeshEditingTool* CreateNewTool(const FToolBuilderState& SceneState) const override;
+	virtual UMultiTargetWithSelectionTool* CreateNewTool(const FToolBuilderState& SceneState) const override;
+
+	virtual bool RequiresInputSelection() const override { return false; }
 
 protected:
 	virtual const FToolTargetTypeRequirements& GetTargetRequirements() const override;
@@ -36,7 +38,7 @@ public:
 
 
 UCLASS()
-class MESHMODELINGTOOLSEXP_API USplitMeshesTool : public UMultiSelectionMeshEditingTool
+class MESHMODELINGTOOLSEXP_API USplitMeshesTool : public UMultiTargetWithSelectionTool
 {
 	GENERATED_BODY()
 
diff --git a/Engine/Plugins/Runtime/MeshModelingToolset/Source/ModelingComponents/Private/BaseTools/MultiTargetWithSelectionTool.cpp b/Engine/Plugins/Runtime/MeshModelingToolset/Source/ModelingComponents/Private/BaseTools/MultiTargetWithSelectionTool.cpp
new file mode 100644
--- /dev/null
+++ b/Engine/Plugins/Runtime/MeshModelingToolset/Source/ModelingComponents/Private/BaseTools/MultiTargetWithSelectionTool.cpp
@@ -0,0 +1,163 @@
+// Copyright Epic Games, Inc. All Rights Reserved.
+
+#include "BaseTools/MultiTargetWithSelectionTool.h"
+
+#include "ModelingToolTargetUtil.h"
+#include "Engine/World.h"
+#include "UDynamicMesh.h"
+
+#include "Drawing/PreviewGeometryActor.h"
+#include "TargetInterfaces/DynamicMeshSource.h"
+#include "TargetInterfaces/MaterialProvider.h"
+#include "TargetInterfaces/MeshDescriptionCommitter.h"
+#include "TargetInterfaces/MeshDescriptionProvider.h"
+#include "TargetInterfaces/PrimitiveComponentBackedTarget.h"
+#include "ToolTargetManager.h"
+#include "Selection/StoredMeshSelectionUtil.h"
+#include "Selection/GeometrySelectionVisualization.h"
+#include "Selections/GeometrySelection.h"
+#include "PropertySets/GeometrySelectionVisualizationProperties.h"
+
+#include UE_INLINE_GENERATED_CPP_BY_NAME(MultiTargetWithSelectionTool)
+
+
+/*
+ * ToolBuilder
+ */
+const FToolTargetTypeRequirements& UMultiTargetWithSelectionToolBuilder::GetTargetRequirements() const
+{
+	static FToolTargetTypeRequirements TypeRequirements({
+		UMaterialProvider::StaticClass(),
+		UMeshDescriptionCommitter::StaticClass(),
+		UMeshDescriptionProvider::StaticClass(),
+		UPrimitiveComponentBackedTarget::StaticClass()
+		});
+	return TypeRequirements;
+}
+
+bool UMultiTargetWithSelectionToolBuilder::CanBuildTool(const FToolBuilderState& SceneState) const
+{
+	if (RequiresInputSelection() && UE::Geometry::HaveAvailableGeometrySelection(SceneState) == false )
+	{
+		return false;
+	}
+
+	return SceneState.TargetManager->CountSelectedAndTargetable(SceneState, GetTargetRequirements()) > 0;
+}
+
+UInteractiveTool* UMultiTargetWithSelectionToolBuilder::BuildTool(const FToolBuilderState& SceneState) const
+{
+	UMultiTargetWithSelectionTool* NewTool = CreateNewTool(SceneState);
+	InitializeNewTool(NewTool, SceneState);
+	return NewTool;
+}
+
+void UMultiTargetWithSelectionToolBuilder::InitializeNewTool(UMultiTargetWithSelectionTool* NewTool, const FToolBuilderState& SceneState) const
+{
+
+	const TArray<TObjectPtr<UToolTarget>> Targets = SceneState.TargetManager->BuildAllSelectedTargetable(SceneState, GetTargetRequirements());
+	const int NumTargets = Targets.Num();
+	
+	NewTool->InitializeGeometrySelectionArrays(NumTargets);
+	NewTool->SetTargets(Targets);
+	NewTool->SetTargetWorld(SceneState.World);
+
+	TArray<UE::Geometry::FGeometrySelection> Selections;
+	Selections.SetNum(NumTargets);
+	for (int TargetIndex = 0; TargetIndex < NumTargets; TargetIndex++)
+	{
+		bool bHaveSelection = UE::Geometry::GetCurrentGeometrySelectionForTarget(SceneState, Targets[TargetIndex], Selections[TargetIndex]);
+
+		if (bHaveSelection)
+		{
+			NewTool->SetGeometrySelection(MoveTemp(Selections[TargetIndex]), TargetIndex);
+		}
+	}
+
+}
+
+void UMultiTargetWithSelectionTool::OnTick(float DeltaTime)
+{
+	Super::OnTick(DeltaTime);
+
+	if (GeometrySelectionViz)
+	{
+		UE::Geometry::UpdateGeometrySelectionVisualization(GeometrySelectionViz, GeometrySelectionVizProperties);
+	}
+}
+
+
+void UMultiTargetWithSelectionTool::Shutdown(EToolShutdownType ShutdownType)
+{
+	OnShutdown(ShutdownType);
+	TargetWorld = nullptr;
+
+	Super::Shutdown(ShutdownType);
+}
+
+void UMultiTargetWithSelectionTool::OnShutdown(EToolShutdownType ShutdownType)
+{
+	if (GeometrySelectionViz)
+	{
+		GeometrySelectionViz->Disconnect();
+	}
+
+	if (GeometrySelectionVizProperties)
+	{
+		GeometrySelectionVizProperties->SaveProperties(this);
+	}
+}
+
+void UMultiTargetWithSelectionTool::SetTargetWorld(UWorld* World)
+{
+	TargetWorld = World;
+}
+
+UWorld* UMultiTargetWithSelectionTool::GetTargetWorld()
+{
+	return TargetWorld.Get();
+}
+
+void UMultiTargetWithSelectionTool::InitializeGeometrySelectionArrays(const int NumTargets)
+{
+	GeometrySelectionArray.SetNum(NumTargets);
+	GeometrySelectionBoolArray.SetNum(NumTargets);
+	for (int BoolArrIndex = 0; BoolArrIndex < NumTargets; BoolArrIndex++)
+	{
+		GeometrySelectionBoolArray[BoolArrIndex] = false;
+	}
+}
+
+
+
+void UMultiTargetWithSelectionTool::SetGeometrySelection(const UE::Geometry::FGeometrySelection& SelectionIn, const int TargetIndex)
+{
+	GeometrySelectionArray[TargetIndex] = SelectionIn;
+	GeometrySelectionBoolArray[TargetIndex] = true;
+}
+
+void UMultiTargetWithSelectionTool::SetGeometrySelection(UE::Geometry::FGeometrySelection&& SelectionIn, const int TargetIndex)
+{
+	GeometrySelectionArray[TargetIndex] = MoveTemp(SelectionIn);
+	GeometrySelectionBoolArray[TargetIndex] = true;
+}
+
+bool UMultiTargetWithSelectionTool::HasGeometrySelection(const int TargetIndex) const
+{
+	return GeometrySelectionBoolArray[TargetIndex];
+}
+
+const UE::Geometry::FGeometrySelection& UMultiTargetWithSelectionTool::GetGeometrySelection(const int TargetIndex) const
+{
+	return GeometrySelectionArray[TargetIndex];
+}
+
+bool UMultiTargetWithSelectionTool::HasAnyGeometrySelection() const
+{
+	bool bHasGeometrySelectedAcrossAllTargets = false;
+	for (const bool bGeoSelectedPerTarget : GeometrySelectionBoolArray)
+	{
+		bHasGeometrySelectedAcrossAllTargets = bHasGeometrySelectedAcrossAllTargets || bGeoSelectedPerTarget;
+	}
+	return bHasGeometrySelectedAcrossAllTargets;
+}
diff --git a/Engine/Plugins/Runtime/MeshModelingToolset/Source/ModelingComponents/Public/BaseTools/MultiTargetWithSelectionTool.h b/Engine/Plugins/Runtime/MeshModelingToolset/Source/ModelingComponents/Public/BaseTools/MultiTargetWithSelectionTool.h
new file mode 100644
--- /dev/null
+++ b/Engine/Plugins/Runtime/MeshModelingToolset/Source/ModelingComponents/Public/BaseTools/MultiTargetWithSelectionTool.h
@@ -0,0 +1,96 @@
+// Copyright Epic Games, Inc. All Rights Reserved.
+
+#pragma once
+
+
+#include "InteractiveToolBuilder.h"
+#include "MultiSelectionMeshEditingTool.h"
+#include "Selections/GeometrySelection.h"
+#include "MultiTargetWithSelectionTool.generated.h"
+
+class UMultiTargetWithSelectionTool;
+class UPreviewGeometry;
+class UGeometrySelectionVisualizationProperties;
+
+PREDECLARE_GEOMETRY(class FDynamicMesh3)
+PREDECLARE_GEOMETRY(class FGroupTopology)
+
+
+/**
+ * UMultiTargetWithSelectionToolBuilder is a base tool builder for multi
+ * selection tools with selections.
+ * Currently, geometry selection across multiple meshes is not supported, restricting the effectiveness
+ * of this class. If that support is built in the future, this will become more useful, and likely need to be expanded
+ */
+UCLASS(Transient, Abstract)
+class MODELINGCOMPONENTS_API UMultiTargetWithSelectionToolBuilder : public UInteractiveToolWithToolTargetsBuilder
+{
+	GENERATED_BODY()
+public:
+	/** @return true if mesh sources can be found in the active selection */
+	virtual bool CanBuildTool(const FToolBuilderState& SceneState) const override;
+
+	/** @return new Tool instance initialized with selected mesh source(s) */
+	virtual UInteractiveTool* BuildTool(const FToolBuilderState& SceneState) const override;
+
+	/** @return new Tool instance. Override this in subclasses to build a different Tool class type */
+	virtual UMultiTargetWithSelectionTool* CreateNewTool(const FToolBuilderState& SceneState) const PURE_VIRTUAL(UMultiTargetWithSelectionToolBuilder::CreateNewTool, return nullptr; );
+
+	/** Called by BuildTool to configure the Tool with the input mesh source(s) based on the SceneState */
+	virtual void InitializeNewTool(UMultiTargetWithSelectionTool* NewTool, const FToolBuilderState& SceneState) const;
+
+	/** @return true if this Tool requires an input selection */
+	virtual bool RequiresInputSelection() const { return false; }
+
+protected:
+	virtual const FToolTargetTypeRequirements& GetTargetRequirements() const override;
+};
+
+/**
+ * Multi Target with Selection tool base class.
+ */
+
+UCLASS()
+class MODELINGCOMPONENTS_API UMultiTargetWithSelectionTool : public UMultiSelectionTool
+{
+	GENERATED_BODY()
+public:
+	virtual void Shutdown(EToolShutdownType ShutdownType) override;
+	virtual void OnShutdown(EToolShutdownType ShutdownType);
+
+	virtual void OnTick(float DeltaTime) override;
+
+	virtual void SetTargetWorld(UWorld* World);
+	virtual UWorld* GetTargetWorld();
+
+protected:
+	UPROPERTY()
+	TWeakObjectPtr<UWorld> TargetWorld = nullptr;
+
+
+public:
+	virtual void SetGeometrySelection(const UE::Geometry::FGeometrySelection& SelectionIn, const int TargetIndex);
+	virtual void SetGeometrySelection(UE::Geometry::FGeometrySelection&& SelectionIn, const int TargetIndex);
+
+	/** @return true if a Selection is available for the Target at the given index*/
+	virtual bool HasGeometrySelection(const int TargetIndex) const;
+
+	/** @return the input Selection for the Target at the given index*/
+	virtual const UE::Geometry::FGeometrySelection& GetGeometrySelection(const int TargetIndex) const;
+
+	/** @return if a Selection is available for ANY of the Targets */
+	virtual bool HasAnyGeometrySelection() const;
+
+	/** initialize the Geometry Selection array and the boolean arrays according to the number of targets */
+	virtual void InitializeGeometrySelectionArrays(const int NumTargets);
+
+protected:
+	TArray<UE::Geometry::FGeometrySelection> GeometrySelectionArray;
+	TArray<bool> GeometrySelectionBoolArray;
+
+	UPROPERTY()
+	TObjectPtr<UGeometrySelectionVisualizationProperties> GeometrySelectionVizProperties = nullptr;
+
+	UPROPERTY()
+	TObjectPtr<UPreviewGeometry> GeometrySelectionViz = nullptr;
+};
\ No newline at end of file
