[Skeletal Mesh Editor] Making sure UI reacts to new weight map layers being added as new Vertex Attributes

diff --git a/Engine/Plugins/Runtime/MeshModelingToolset/Source/ModelingComponentsEditorOnly/Private/ToolTargets/SkeletalMeshToolTarget.cpp b/Engine/Plugins/Runtime/MeshModelingToolset/Source/ModelingComponentsEditorOnly/Private/ToolTargets/SkeletalMeshToolTarget.cpp
--- a/Engine/Plugins/Runtime/MeshModelingToolset/Source/ModelingComponentsEditorOnly/Private/ToolTargets/SkeletalMeshToolTarget.cpp
+++ b/Engine/Plugins/Runtime/MeshModelingToolset/Source/ModelingComponentsEditorOnly/Private/ToolTargets/SkeletalMeshToolTarget.cpp
@@ -2,6 +2,7 @@
 
 #include "ToolTargets/SkeletalMeshToolTarget.h"
 
+#include "LODUtilities.h"
 #include "ConversionUtils/DynamicMeshViaMeshDescriptionUtil.h"
 #include "DynamicMesh/DynamicMesh3.h"
 #include "Engine/SkinnedAssetCommon.h"
@@ -270,6 +271,10 @@ void USkeletalMeshToolTarget::CommitMeshDescription(USkeletalMesh* SkeletalMesh,
 	// Commit the edited mesh description to bulk data.
 	SkeletalMesh->CommitMeshDescription(LODIndex);
 
+	// Update the vertex attribute infos now that we have committed the mesh description
+	// the actual copying from SourceModel to LODModel happens during build time
+	FLODUtilities::UpdateLODInfoVertexAttributes(SkeletalMesh, LODIndex, LODIndex, false);
+	
 	// This call will invoke the skeletal mesh render data rebuild machinery.
 	SkeletalMesh->PostEditChange();
 }
diff --git a/Engine/Source/Developer/SkeletalMeshUtilitiesCommon/Private/LODUtilities.cpp b/Engine/Source/Developer/SkeletalMeshUtilitiesCommon/Private/LODUtilities.cpp
--- a/Engine/Source/Developer/SkeletalMeshUtilitiesCommon/Private/LODUtilities.cpp
+++ b/Engine/Source/Developer/SkeletalMeshUtilitiesCommon/Private/LODUtilities.cpp
@@ -3169,6 +3169,11 @@ bool FLODUtilities::UpdateLODInfoVertexAttributes(
 	// Wait for all the attribute conversion tasks to complete.
 	UE::Tasks::Wait(ConversionTasks);
 
+#if WITH_EDITOR
+	// Notify UI and other systems of the change
+	InSkeletalMesh->GetOnVertexAttributesArrayChanged().Broadcast();
+#endif
+	
 	return true;
 }
 
diff --git a/Engine/Source/Editor/Persona/Private/LODInfoUILayout.cpp b/Engine/Source/Editor/Persona/Private/LODInfoUILayout.cpp
--- a/Engine/Source/Editor/Persona/Private/LODInfoUILayout.cpp
+++ b/Engine/Source/Editor/Persona/Private/LODInfoUILayout.cpp
@@ -25,6 +25,15 @@ void ULODInfoUILayout::SetReferenceLODInfo(TWeakPtr<IPersonaToolkit> InPersonaTo
 	LODInfo = *SkeletalMeshLODInfo;
 }
 
+void ULODInfoUILayout::RefreshReferenceLODInfo()
+{
+	USkeletalMesh* SkeletalMesh = PersonaToolkit.Pin()->GetPreviewMesh();
+	const FSkeletalMeshLODInfo* SkeletalMeshLODInfo = SkeletalMesh->GetLODInfo(LODIndex);
+	check(SkeletalMeshLODInfo != nullptr);
+	//Copy the LODInfo Array to the temporary
+	LODInfo = *SkeletalMeshLODInfo;	
+}
+
 void ULODInfoUILayout::PostEditChangeProperty(struct FPropertyChangedEvent& PropertyChangedEvent)
 {
 	Super::PostEditChangeProperty(PropertyChangedEvent);
diff --git a/Engine/Source/Editor/Persona/Private/LODInfoUILayout.h b/Engine/Source/Editor/Persona/Private/LODInfoUILayout.h
--- a/Engine/Source/Editor/Persona/Private/LODInfoUILayout.h
+++ b/Engine/Source/Editor/Persona/Private/LODInfoUILayout.h
@@ -21,6 +21,7 @@ class ULODInfoUILayout : public UObject
 
 public:
 	void SetReferenceLODInfo(TWeakPtr<IPersonaToolkit> InPersonaToolkit, int32 InLODIndex);
+	void RefreshReferenceLODInfo();
 
 	TSharedPtr<IPersonaToolkit> GetPersonaToolkit() const
 	{
diff --git a/Engine/Source/Editor/Persona/Private/PersonaMeshDetails.cpp b/Engine/Source/Editor/Persona/Private/PersonaMeshDetails.cpp
--- a/Engine/Source/Editor/Persona/Private/PersonaMeshDetails.cpp
+++ b/Engine/Source/Editor/Persona/Private/PersonaMeshDetails.cpp
@@ -2880,10 +2880,11 @@ void FPersonaMeshDetails::AddLODLevelCategories(IDetailLayoutBuilder& DetailLayo
 
 		LodCategories.Empty(SkelMeshLODCount);
 		DetailDisplayLODs.Reset();
-		auto ClearLODInfoLayouts = [this, &SkelMeshLODCount]()
+		auto ClearLODInfoLayouts = [this, &SkelMeshLODCount, SkelMesh]()
 		{
 			for (ULODInfoUILayout* LODInfoUILayout : LODInfoUILayouts)
 			{
+				SkelMesh->GetOnVertexAttributesArrayChanged().RemoveAll(LODInfoUILayout);
 				LODInfoUILayout->RemoveFromRoot();
 				LODInfoUILayout->MarkAsGarbage();
 				LODInfoUILayout = nullptr;
@@ -2908,6 +2909,7 @@ void FPersonaMeshDetails::AddLODLevelCategories(IDetailLayoutBuilder& DetailLayo
 			}
 			const FSkeletalMeshLODModel& LODModel = SkelMesh->GetImportedModel()->LODModels[LODIndex];
 			LODInfoUILayout->SetReferenceLODInfo(GetPersonaToolkit(), LODIndex);
+			SkelMesh->GetOnVertexAttributesArrayChanged().AddUObject(LODInfoUILayout, &ULODInfoUILayout::RefreshReferenceLODInfo);
 			LODInfoUILayouts.Add(LODInfoUILayout);
 
 			//Show the viewport LOD at start
diff --git a/Engine/Source/Runtime/Engine/Classes/Engine/SkeletalMesh.h b/Engine/Source/Runtime/Engine/Classes/Engine/SkeletalMesh.h
--- a/Engine/Source/Runtime/Engine/Classes/Engine/SkeletalMesh.h
+++ b/Engine/Source/Runtime/Engine/Classes/Engine/SkeletalMesh.h
@@ -936,6 +936,10 @@ private:
 	UPROPERTY(EditAnywhere, EditFixedSize, Category=LevelOfDetail)
 	TArray<FSkeletalMeshLODInfo> LODInfo;
 
+#if WITH_EDITOR
+	FSimpleMulticastDelegate OnVertexAttributesArrayChanged;
+#endif
+	
 #if !WITH_EDITOR
 	/** Acceleration struct used for faster socket lookups */
 	struct FSocketInfo
@@ -2925,6 +2929,10 @@ public:
 	 */	
 	ENGINE_API virtual const FSkeletalMeshLODInfo* GetLODInfo(int32 Index) const override;
 
+#if WITH_EDITOR
+	ENGINE_API FSimpleMulticastDelegate& GetOnVertexAttributesArrayChanged();
+#endif
+	
 	/**
 	 *	Get BakePose for the given LOD
 	 */
diff --git a/Engine/Source/Runtime/Engine/Private/SkeletalMesh.cpp b/Engine/Source/Runtime/Engine/Private/SkeletalMesh.cpp
--- a/Engine/Source/Runtime/Engine/Private/SkeletalMesh.cpp
+++ b/Engine/Source/Runtime/Engine/Private/SkeletalMesh.cpp
@@ -5379,6 +5379,13 @@ const FSkeletalMeshLODInfo* USkeletalMesh::GetLODInfo(int32 Index) const
 	PRAGMA_ENABLE_DEPRECATION_WARNINGS
 }
 
+#if WITH_EDITOR
+FSimpleMulticastDelegate& USkeletalMesh::GetOnVertexAttributesArrayChanged()
+{
+	return OnVertexAttributesArrayChanged;
+}
+#endif
+
 const UAnimSequence* USkeletalMesh::GetBakePose(int32 LODIndex) const
 {
 	const FSkeletalMeshLODInfo* LOD = GetLODInfo(LODIndex);
