fix crash if trying to convert a cooked-for-editor geometry collection

diff --git a/Engine/Plugins/Runtime/MeshModelingToolset/Source/ModelingComponents/Private/ConversionUtils/SceneComponentToDynamicMesh.cpp b/Engine/Plugins/Runtime/MeshModelingToolset/Source/ModelingComponents/Private/ConversionUtils/SceneComponentToDynamicMesh.cpp
--- a/Engine/Plugins/Runtime/MeshModelingToolset/Source/ModelingComponents/Private/ConversionUtils/SceneComponentToDynamicMesh.cpp
+++ b/Engine/Plugins/Runtime/MeshModelingToolset/Source/ModelingComponents/Private/ConversionUtils/SceneComponentToDynamicMesh.cpp
@@ -50,7 +50,7 @@ bool CanConvertSceneComponentToDynamicMesh(USceneComponent* Component)
 	{
 #if WITH_EDITOR
 		const USkinnedAsset* SkinnedAsset = (!SkinnedMeshComponent->IsUnreachable() && SkinnedMeshComponent->IsValidLowLevel()) ? SkinnedMeshComponent->GetSkinnedAsset() : nullptr;
-		return SkinnedAsset && !SkinnedAsset->GetOutermost()->bIsCookedForEditor;;
+		return SkinnedAsset && !SkinnedAsset->GetOutermost()->bIsCookedForEditor;
 #else
 		return true;
 #endif
@@ -76,9 +76,14 @@ bool CanConvertSceneComponentToDynamicMesh(USceneComponent* Component)
 	{
 		return true;
 	}
-	else if (Cast<UGeometryCollectionComponent>(Component))
+	else if (UGeometryCollectionComponent* GeometryCollectionComponent = Cast<UGeometryCollectionComponent>(Component))
 	{
+#if WITH_EDITOR
+		const UGeometryCollection* GeometryCollectionAsset = (!GeometryCollectionComponent->IsUnreachable() && GeometryCollectionComponent->IsValidLowLevel()) ? GeometryCollectionComponent->GetRestCollection() : nullptr;
+		return GeometryCollectionAsset && !GeometryCollectionAsset->GetOutermost()->bIsCookedForEditor;
+#else
 		return true;
+#endif
 	}
 	return false;
 }
