add capsule primitive tool

diff --git a/Engine/Plugins/Editor/ModelingToolsEditorMode/Source/ModelingToolsEditorMode/Private/ModelingToolsEditorMode.cpp b/Engine/Plugins/Editor/ModelingToolsEditorMode/Source/ModelingToolsEditorMode/Private/ModelingToolsEditorMode.cpp
--- a/Engine/Plugins/Editor/ModelingToolsEditorMode/Source/ModelingToolsEditorMode/Private/ModelingToolsEditorMode.cpp
+++ b/Engine/Plugins/Editor/ModelingToolsEditorMode/Source/ModelingToolsEditorMode/Private/ModelingToolsEditorMode.cpp
@@ -562,6 +562,9 @@ void UModelingToolsEditorMode::Enter()
 	RegisterPrimitiveToolFunc(ToolManagerCommands.BeginAddCylinderPrimitiveTool,
 							  TEXT("BeginAddCylinderPrimitiveTool"),
 							  UAddPrimitiveToolBuilder::EMakeMeshShapeType::Cylinder);
+	RegisterPrimitiveToolFunc(ToolManagerCommands.BeginAddCapsulePrimitiveTool,
+	                          TEXT("BeginAddCapsulePrimitiveTool"),
+	                          UAddPrimitiveToolBuilder::EMakeMeshShapeType::Capsule);
 	RegisterPrimitiveToolFunc(ToolManagerCommands.BeginAddConePrimitiveTool,
 							  TEXT("BeginAddConePrimitiveTool"),
 							  UAddPrimitiveToolBuilder::EMakeMeshShapeType::Cone);
diff --git a/Engine/Plugins/Editor/ModelingToolsEditorMode/Source/ModelingToolsEditorMode/Private/ModelingToolsEditorModeToolkit.cpp b/Engine/Plugins/Editor/ModelingToolsEditorMode/Source/ModelingToolsEditorMode/Private/ModelingToolsEditorModeToolkit.cpp
--- a/Engine/Plugins/Editor/ModelingToolsEditorMode/Source/ModelingToolsEditorMode/Private/ModelingToolsEditorModeToolkit.cpp
+++ b/Engine/Plugins/Editor/ModelingToolsEditorMode/Source/ModelingToolsEditorMode/Private/ModelingToolsEditorModeToolkit.cpp
@@ -251,6 +251,7 @@ void FModelingToolsEditorModeToolkit::RegisterPalettes()
 		Commands.BeginAddBoxPrimitiveTool,
 		Commands.BeginAddSpherePrimitiveTool,
 		Commands.BeginAddCylinderPrimitiveTool,
+		Commands.BeginAddCapsulePrimitiveTool,
 		Commands.BeginAddConePrimitiveTool,
 		Commands.BeginAddTorusPrimitiveTool,
 		Commands.BeginAddArrowPrimitiveTool,
@@ -1853,6 +1854,7 @@ void FModelingToolsEditorModeToolkit::BuildToolPalette(FName PaletteIndex, class
 		ToolbarBuilder.AddToolBarButton(Commands.BeginAddBoxPrimitiveTool);
 		ToolbarBuilder.AddToolBarButton(Commands.BeginAddSpherePrimitiveTool);
 		ToolbarBuilder.AddToolBarButton(Commands.BeginAddCylinderPrimitiveTool);
+		ToolbarBuilder.AddToolBarButton(Commands.BeginAddCapsulePrimitiveTool);
 		ToolbarBuilder.AddToolBarButton(Commands.BeginAddConePrimitiveTool);
 		ToolbarBuilder.AddToolBarButton(Commands.BeginAddTorusPrimitiveTool);
 		ToolbarBuilder.AddToolBarButton(Commands.BeginAddArrowPrimitiveTool);
diff --git a/Engine/Plugins/Editor/ModelingToolsEditorMode/Source/ModelingToolsEditorMode/Private/ModelingToolsManagerActions.cpp b/Engine/Plugins/Editor/ModelingToolsEditorMode/Source/ModelingToolsEditorMode/Private/ModelingToolsManagerActions.cpp
--- a/Engine/Plugins/Editor/ModelingToolsEditorMode/Source/ModelingToolsEditorMode/Private/ModelingToolsManagerActions.cpp
+++ b/Engine/Plugins/Editor/ModelingToolsEditorMode/Source/ModelingToolsEditorMode/Private/ModelingToolsManagerActions.cpp
@@ -61,6 +61,7 @@ void FModelingToolsManagerCommands::RegisterCommands()
 	REGISTER_MODELING_TOOL_COMMAND(BeginAddArrowPrimitiveTool, "Arrow", "Create new arrow objects");
 	REGISTER_MODELING_TOOL_COMMAND(BeginAddDiscPrimitiveTool, "Disc", "Create new disc objects");
 	REGISTER_MODELING_TOOL_COMMAND(BeginAddStairsPrimitiveTool, "Stairs", "Create new stairs objects");
+	REGISTER_MODELING_TOOL_COMMAND(BeginAddCapsulePrimitiveTool, "Capsule", "Create new capsule objects");
 	if ( bUseShortNames )
 	{
 		REGISTER_MODELING_TOOL_COMMAND(BeginAddCylinderPrimitiveTool, "Cyl", "Create new cylinder objects");
diff --git a/Engine/Plugins/Editor/ModelingToolsEditorMode/Source/ModelingToolsEditorMode/Public/ModelingToolsManagerActions.h b/Engine/Plugins/Editor/ModelingToolsEditorMode/Source/ModelingToolsEditorMode/Public/ModelingToolsManagerActions.h
--- a/Engine/Plugins/Editor/ModelingToolsEditorMode/Source/ModelingToolsEditorMode/Public/ModelingToolsManagerActions.h
+++ b/Engine/Plugins/Editor/ModelingToolsEditorMode/Source/ModelingToolsEditorMode/Public/ModelingToolsManagerActions.h
@@ -40,6 +40,7 @@ public:
 
 	TSharedPtr<FUICommandInfo> BeginAddBoxPrimitiveTool;
 	TSharedPtr<FUICommandInfo> BeginAddCylinderPrimitiveTool;
+	TSharedPtr<FUICommandInfo> BeginAddCapsulePrimitiveTool;
 	TSharedPtr<FUICommandInfo> BeginAddConePrimitiveTool;
 	TSharedPtr<FUICommandInfo> BeginAddArrowPrimitiveTool;
 	TSharedPtr<FUICommandInfo> BeginAddRectanglePrimitiveTool;
diff --git a/Engine/Plugins/Runtime/MeshModelingToolset/Source/MeshModelingTools/Private/AddPrimitiveTool.cpp b/Engine/Plugins/Runtime/MeshModelingToolset/Source/MeshModelingTools/Private/AddPrimitiveTool.cpp
--- a/Engine/Plugins/Runtime/MeshModelingToolset/Source/MeshModelingTools/Private/AddPrimitiveTool.cpp
+++ b/Engine/Plugins/Runtime/MeshModelingToolset/Source/MeshModelingTools/Private/AddPrimitiveTool.cpp
@@ -17,6 +17,7 @@
 #include "Components/StaticMeshComponent.h"
 #include "Engine/StaticMeshActor.h"
 
+#include "Generators/CapsuleGenerator.h"
 #include "Generators/SweepGenerator.h"
 #include "Generators/GridBoxMeshGenerator.h"
 #include "Generators/RectangleMeshGenerator.h"
@@ -77,6 +78,9 @@ UInteractiveTool* UAddPrimitiveToolBuilder::BuildTool(const FToolBuilderState& S
 	case EMakeMeshShapeType::Stairs:
 		NewTool = NewObject<UAddStairsPrimitiveTool>(SceneState.ToolManager);
 		break;
+	case EMakeMeshShapeType::Capsule:
+		NewTool = NewObject<UAddCapsulePrimitiveTool>(SceneState.ToolManager);
+		break;
 	default:
 		break;
 	}
@@ -698,6 +702,26 @@ void UAddCylinderPrimitiveTool::GenerateMesh(FDynamicMesh3* OutMesh) const
 }
 
 
+UAddCapsulePrimitiveTool::UAddCapsulePrimitiveTool(const FObjectInitializer& ObjectInitializer)
+	: Super(ObjectInitializer.SetDefaultSubobjectClass<UProceduralCapsuleToolProperties>(TEXT("ShapeSettings")))
+{
+	AssetName = TEXT("Capsule");
+	UInteractiveTool::SetToolDisplayName(LOCTEXT("CapsuleToolName", "Create Capsule"));
+}
+
+void UAddCapsulePrimitiveTool::GenerateMesh(FDynamicMesh3* OutMesh) const
+{
+	FCapsuleGenerator CapGen;
+	const UProceduralCapsuleToolProperties* CapsuleSettings = Cast<UProceduralCapsuleToolProperties>(ShapeSettings);
+	CapGen.Radius = CapsuleSettings->Radius;
+	CapGen.SegmentLength = CapsuleSettings->CylinderLength;
+	CapGen.NumHemisphereArcSteps = CapsuleSettings->HemisphereSlices;
+	CapGen.NumCircleSteps = CapsuleSettings->CylinderSlices;
+	CapGen.bPolygroupPerQuad = ShapeSettings->PolygroupMode == EMakeMeshPolygroupMode::PerQuad;
+	CapGen.Generate();
+	OutMesh->Copy(&CapGen);
+}
+
 
 UAddConePrimitiveTool::UAddConePrimitiveTool(const FObjectInitializer& ObjectInitializer)
 	: Super(ObjectInitializer.SetDefaultSubobjectClass<UProceduralConeToolProperties>(TEXT("ShapeSettings")))
diff --git a/Engine/Plugins/Runtime/MeshModelingToolset/Source/MeshModelingTools/Public/AddPrimitiveTool.h b/Engine/Plugins/Runtime/MeshModelingToolset/Source/MeshModelingTools/Public/AddPrimitiveTool.h
--- a/Engine/Plugins/Runtime/MeshModelingToolset/Source/MeshModelingTools/Public/AddPrimitiveTool.h
+++ b/Engine/Plugins/Runtime/MeshModelingToolset/Source/MeshModelingTools/Public/AddPrimitiveTool.h
@@ -36,7 +36,8 @@ public:
 		Disc,
 		Torus,
 		Sphere,
-		Stairs
+		Stairs,
+		Capsule
 	};
 
 	virtual bool CanBuildTool(const FToolBuilderState& SceneState) const override;
@@ -372,6 +373,29 @@ public:
 	int VerticalSlices = 16;
 };
 
+UCLASS()
+class MESHMODELINGTOOLS_API UProceduralCapsuleToolProperties : public UProceduralShapeToolProperties
+{
+	GENERATED_BODY()
+
+public:
+	/** Radius of the capsule */
+	UPROPERTY(EditAnywhere, Category = Shape, meta = (UIMin = "1.0", UIMax = "1000.0", ClampMin = "0.0001", ClampMax = "1000000.0", ProceduralShapeSetting))
+	float Radius = 25.f;
+
+	/** Length of cylindrical section of the capsule */
+	UPROPERTY(EditAnywhere, Category = Shape, meta = (UIMin = "1.0", UIMax = "1000.0", ClampMin = "0.0001", ClampMax = "1000000.0", ProceduralShapeSetting))
+	float CylinderLength = 50.f;
+
+	/** Number of slices of the hemispherical end caps. */
+	UPROPERTY(EditAnywhere, Category = Shape, meta = (UIMin = "2", UIMax = "100", ClampMin = "2", ClampMax = "500", ProceduralShapeSetting))
+	int HemisphereSlices = 8;
+
+	/** Number of radial slices of the cylindrical section. */
+	UPROPERTY(EditAnywhere, Category = Shape, meta = (UIMin = "3", UIMax = "100", ClampMin = "3", ClampMax = "500", ProceduralShapeSetting))
+	int CylinderSlices = 16;
+};
+
 UENUM()
 enum class EProceduralStairsType
 {
@@ -562,6 +586,16 @@ protected:
 	virtual void GenerateMesh(FDynamicMesh3* OutMesh) const override;
 };
 
+UCLASS()
+class UAddCapsulePrimitiveTool : public UAddPrimitiveTool
+{
+	GENERATED_BODY()
+public:
+	explicit UAddCapsulePrimitiveTool(const FObjectInitializer& ObjectInitializer);
+protected:
+	virtual void GenerateMesh(FDynamicMesh3* OutMesh) const override;
+};
+
 UCLASS()
 class UAddConePrimitiveTool : public UAddPrimitiveTool
 {
