add a nullptr check to component collision util's ConvertSimpleCollisionToMeshes, method, fixing consistent crash in Geometry Script's CopyCollisionMeshesFromObject method

diff --git a/Engine/Plugins/Runtime/MeshModelingToolset/Source/ModelingComponents/Private/Physics/ComponentCollisionUtil.cpp b/Engine/Plugins/Runtime/MeshModelingToolset/Source/ModelingComponents/Private/Physics/ComponentCollisionUtil.cpp
--- a/Engine/Plugins/Runtime/MeshModelingToolset/Source/ModelingComponents/Private/Physics/ComponentCollisionUtil.cpp
+++ b/Engine/Plugins/Runtime/MeshModelingToolset/Source/ModelingComponents/Private/Physics/ComponentCollisionUtil.cpp
@@ -651,7 +651,10 @@ void UE::Geometry::ConvertSimpleCollisionToMeshes(
 		MeshAttributeSettings,
 		[&](int32 Index, const FKShapeElem& Elem, const UE::Geometry::FDynamicMesh3& Mesh)
 		{
-			PerElementMeshCallback((int)Elem.GetShapeType(), Mesh);
+			if (PerElementMeshCallback)
+			{
+				PerElementMeshCallback((int)Elem.GetShapeType(), Mesh);
+			}
 			Mappings.Reset();
 			Editor.AppendMesh(&Mesh, Mappings,
 				[&TransformSequence](int32 vid, const FVector3d& P) { return TransformSequence.TransformPosition(P); },
