add a 'cut outside' option to the polycut tool

diff --git a/Engine/Plugins/Experimental/MeshModelingToolsetExp/Source/MeshModelingToolsEditorOnlyExp/Public/PolygonOnMeshTool.h b/Engine/Plugins/Experimental/MeshModelingToolsetExp/Source/MeshModelingToolsEditorOnlyExp/Public/PolygonOnMeshTool.h
--- a/Engine/Plugins/Experimental/MeshModelingToolsetExp/Source/MeshModelingToolsEditorOnlyExp/Public/PolygonOnMeshTool.h
+++ b/Engine/Plugins/Experimental/MeshModelingToolsetExp/Source/MeshModelingToolsEditorOnlyExp/Public/PolygonOnMeshTool.h
@@ -76,7 +76,9 @@ public:
 	bool bCutWithBoolean = true;
 
 	/** Automatically attempt to fill any open boundaries left by CSG (e.g. due to numerical errors) */
-	UPROPERTY(EditAnywhere, Category = Operation, meta = (EditCondition = "bCutWithBoolean && Operation == EEmbeddedPolygonOpMethod::CutThrough || bCutWithBoolean && Operation == EEmbeddedPolygonOpMethod::InsertPolygon", EditConditionHides))
+	UPROPERTY(EditAnywhere, Category = Operation, meta = (EditCondition = 
+		"bCutWithBoolean && (Operation == EEmbeddedPolygonOpMethod::CutThrough || Operation == EEmbeddedPolygonOpMethod::CutOutside || Operation == EEmbeddedPolygonOpMethod::InsertPolygon)",
+		EditConditionHides))
 	bool bTryToFixHoles = false;
 
 	// TODO: re-add if/when extrude is added as a supported operation
diff --git a/Engine/Plugins/Runtime/MeshModelingToolset/Source/ModelingOperatorsEditorOnly/Private/CuttingOps/EmbedPolygonsOp.cpp b/Engine/Plugins/Runtime/MeshModelingToolset/Source/ModelingOperatorsEditorOnly/Private/CuttingOps/EmbedPolygonsOp.cpp
--- a/Engine/Plugins/Runtime/MeshModelingToolset/Source/ModelingOperatorsEditorOnly/Private/CuttingOps/EmbedPolygonsOp.cpp
+++ b/Engine/Plugins/Runtime/MeshModelingToolset/Source/ModelingOperatorsEditorOnly/Private/CuttingOps/EmbedPolygonsOp.cpp
@@ -208,6 +208,9 @@ void FEmbedPolygonsOp::BooleanPath(FProgressCancel* Progress)
 	case EEmbeddedPolygonOpMethod::CutThrough:
 		BoolOp = FMeshBoolean::EBooleanOp::Difference;
 		break;
+	case EEmbeddedPolygonOpMethod::CutOutside:
+		BoolOp = FMeshBoolean::EBooleanOp::Intersect;
+		break;
 	default:
 		unimplemented();
 	}
@@ -299,7 +302,7 @@ void FEmbedPolygonsOp::BooleanPath(FProgressCancel* Progress)
 
 	// Non-mesh-boundary cut boundary edges can be detected by the group change across the cut. Not relevant to trim
 	// operations, where we get them from the created boundary edges above.
-	if (Operation == EEmbeddedPolygonOpMethod::CutThrough || Operation == EEmbeddedPolygonOpMethod::InsertPolygon)
+	if (Operation == EEmbeddedPolygonOpMethod::CutOutside || Operation == EEmbeddedPolygonOpMethod::CutThrough || Operation == EEmbeddedPolygonOpMethod::InsertPolygon)
 	{
 		for (int EID : ResultMesh->EdgeIndicesItr())
 		{
@@ -357,7 +360,7 @@ void FEmbedPolygonsOp::CalculateResult(FProgressCancel* Progress)
 	}
 
 	int SecondHit = 1;
-	if (Operation == EEmbeddedPolygonOpMethod::CutThrough)
+	if (Operation == EEmbeddedPolygonOpMethod::CutThrough || Operation == EEmbeddedPolygonOpMethod::CutOutside)
 	{
 		while (SecondHit < SortedHitTriangles.Num() && FMath::IsNearlyEqual(SortedHitTriangles[SecondHit].Key, SortedHitTriangles[0].Key))
 		{
@@ -518,12 +521,13 @@ void FEmbedPolygonsOp::CalculateResult(FProgressCancel* Progress)
 	{
 		DeleteMethod = EDeleteMethod::DeleteNone;
 	}
-	else if (Operation == EEmbeddedPolygonOpMethod::TrimOutside)
+	else if (Operation == EEmbeddedPolygonOpMethod::TrimOutside || Operation == EEmbeddedPolygonOpMethod::CutOutside)
 	{
 		DeleteMethod = EDeleteMethod::DeleteOutside;
 	}
 
-	if (Operation != EEmbeddedPolygonOpMethod::CutThrough || SecondHit == -1)
+	bool bIsCutOperation = Operation == EEmbeddedPolygonOpMethod::CutThrough || Operation == EEmbeddedPolygonOpMethod::CutOutside;
+	if (!bIsCutOperation || SecondHit == -1)
 	{
 		TArray<int> PathVertIDs, PathVertCorrespond;
 		bool bCutSide1 = CutHole(*ResultMesh, Frame, SortedHitTriangles[0].Value, Polygon, DeleteMethod, PathVertIDs, PathVertCorrespond, bCollapseDegenerateEdges);
@@ -534,7 +538,7 @@ void FEmbedPolygonsOp::CalculateResult(FProgressCancel* Progress)
 			return;
 		}
 	}
-	else //Operation == EEmbeddedPolygonOpMethod::CutThrough
+	else // bIsCutOperation == true
 	{
 		TArray<int> HitTris;
 		HitTris.Add(SortedHitTriangles[0].Value);
@@ -550,7 +554,8 @@ void FEmbedPolygonsOp::CalculateResult(FProgressCancel* Progress)
 		}
 		FDynamicMeshEditor MeshEditor(ResultMesh.Get());
 		FDynamicMeshEditResult ResultOut;
-		bool bStitched = MeshEditor.StitchSparselyCorrespondedVertexLoops(AllPathVertIDs[0], AllPathVertCorrespond[0], AllPathVertIDs[1], AllPathVertCorrespond[1], ResultOut);
+		bool bReverseOrientation = Operation == EEmbeddedPolygonOpMethod::CutOutside;
+		bool bStitched = MeshEditor.StitchSparselyCorrespondedVertexLoops(AllPathVertIDs[0], AllPathVertCorrespond[0], AllPathVertIDs[1], AllPathVertCorrespond[1], ResultOut, bReverseOrientation);
 		if (!bStitched)
 		{
 			// Don't set bOperationSucceeded to true
@@ -558,7 +563,7 @@ void FEmbedPolygonsOp::CalculateResult(FProgressCancel* Progress)
 		}
 		if (ResultMesh->HasAttributes())
 		{
-			MeshEditor.SetTubeNormals(ResultOut.NewTriangles, AllPathVertIDs[0], AllPathVertCorrespond[0], AllPathVertIDs[1], AllPathVertCorrespond[1]);
+			MeshEditor.SetTubeNormals(ResultOut.NewTriangles, AllPathVertIDs[0], AllPathVertCorrespond[0], AllPathVertIDs[1], AllPathVertCorrespond[1], bReverseOrientation);
 			TArray<float> UValues; UValues.SetNumUninitialized(AllPathVertCorrespond[1].Num() + 1);
 			FVector3f ZVec = -(FVector3f)Frame.Z();
 			float Along = 0;
diff --git a/Engine/Plugins/Runtime/MeshModelingToolset/Source/ModelingOperatorsEditorOnly/Public/CuttingOps/EmbedPolygonsOp.h b/Engine/Plugins/Runtime/MeshModelingToolset/Source/ModelingOperatorsEditorOnly/Public/CuttingOps/EmbedPolygonsOp.h
--- a/Engine/Plugins/Runtime/MeshModelingToolset/Source/ModelingOperatorsEditorOnly/Public/CuttingOps/EmbedPolygonsOp.h
+++ b/Engine/Plugins/Runtime/MeshModelingToolset/Source/ModelingOperatorsEditorOnly/Public/CuttingOps/EmbedPolygonsOp.h
@@ -17,7 +17,8 @@ enum class EEmbeddedPolygonOpMethod : uint8
 	TrimOutside,
 	TrimInside,
 	InsertPolygon,
-	CutThrough
+	CutThrough,
+	CutOutside
 };
 
 namespace UE
diff --git a/Engine/Source/Runtime/GeometryCore/Private/DynamicMeshEditor.cpp b/Engine/Source/Runtime/GeometryCore/Private/DynamicMeshEditor.cpp
--- a/Engine/Source/Runtime/GeometryCore/Private/DynamicMeshEditor.cpp
+++ b/Engine/Source/Runtime/GeometryCore/Private/DynamicMeshEditor.cpp
@@ -955,7 +955,7 @@ void FDynamicMeshEditor::SetTriangleNormals(const TArray<int>& Triangles)
 
 
 
-void FDynamicMeshEditor::SetTubeNormals(const TArray<int>& Triangles, const TArray<int>& VertexIDs1, const TArray<int>& MatchedIndices1, const TArray<int>& VertexIDs2, const TArray<int>& MatchedIndices2)
+void FDynamicMeshEditor::SetTubeNormals(const TArray<int>& Triangles, const TArray<int>& VertexIDs1, const TArray<int>& MatchedIndices1, const TArray<int>& VertexIDs2, const TArray<int>& MatchedIndices2, bool bReverseNormals)
 {
 	check(Mesh->HasAttributes());
 	check(MatchedIndices1.Num() == MatchedIndices2.Num());
@@ -998,6 +998,7 @@ void FDynamicMeshEditor::SetTubeNormals(const TArray<int>& Triangles, const TArr
 		MatchedVertNormals[1][Idx] = Normalized(MatchedEdgeNormals[1][LastMatchedIdx] + MatchedEdgeNormals[1][Idx]);
 	}
 
+	float NormalScale = bReverseNormals ? -1.0f : 1.0f;
 	TMap<int, int> VertToElID;
 	for (int Side = 0; Side < 2; Side++)
 	{
@@ -1034,7 +1035,7 @@ void FDynamicMeshEditor::SetTubeNormals(const TArray<int>& Triangles, const TArr
 		for (int Idx = 0; Idx < NumVertices; Idx++)
 		{
 			int VID = VertexIDs[Idx];
-			VertToElID.Add(VID, Normals->AppendElement(VertNormals[Side][Idx]));
+			VertToElID.Add(VID, Normals->AppendElement(NormalScale * VertNormals[Side][Idx]));
 		}
 	}
 	for (int TID : Triangles)
diff --git a/Engine/Source/Runtime/GeometryCore/Public/DynamicMeshEditor.h b/Engine/Source/Runtime/GeometryCore/Public/DynamicMeshEditor.h
--- a/Engine/Source/Runtime/GeometryCore/Public/DynamicMeshEditor.h
+++ b/Engine/Source/Runtime/GeometryCore/Public/DynamicMeshEditor.h
@@ -362,7 +362,7 @@ public:
 	/**
 	 * For a 'tube' of triangles connecting loops of corresponded vertices, set smooth normals such that corresponding vertices have corresponding normals
 	 */
-	GEOMETRYCORE_API void SetTubeNormals(const TArray<int>& Triangles, const TArray<int>& VertexIDs1, const TArray<int>& MatchedIndices1, const TArray<int>& VertexIDs2, const TArray<int>& MatchedIndices2);
+	GEOMETRYCORE_API void SetTubeNormals(const TArray<int>& Triangles, const TArray<int>& VertexIDs1, const TArray<int>& MatchedIndices1, const TArray<int>& VertexIDs2, const TArray<int>& MatchedIndices2, bool bReverseNormals = false);
 
 
 	//////////////////////////////////////////////////////////////////////////
