ModelingTools: Fix typo in the Cone primitive tooltip for the radial slices property.

diff --git a/Engine/Plugins/Runtime/MeshModelingToolset/Source/MeshModelingTools/Public/AddPrimitiveTool.h b/Engine/Plugins/Runtime/MeshModelingToolset/Source/MeshModelingTools/Public/AddPrimitiveTool.h
--- a/Engine/Plugins/Runtime/MeshModelingToolset/Source/MeshModelingTools/Public/AddPrimitiveTool.h
+++ b/Engine/Plugins/Runtime/MeshModelingToolset/Source/MeshModelingTools/Public/AddPrimitiveTool.h
@@ -294,7 +294,7 @@ public:
 	UPROPERTY(EditAnywhere, Category = Shape, meta = (UIMin = "1.0", UIMax = "1000.0", ClampMin = "0.0001", ClampMax = "1000000.0", ProceduralShapeSetting))
 	float Height = 200.f;
 
-	/** Number of radial slices for the cylinder */
+	/** Number of radial slices for the cone */
 	UPROPERTY(EditAnywhere, Category = Shape, meta = (UIMin = "3", UIMax = "128", ClampMin = "3", ClampMax = "500", ProceduralShapeSetting))
 	int RadialSlices = 16;
 
