allow edit tri tool to collapse seam edges

diff --git a/Engine/Plugins/Runtime/MeshModelingToolset/Source/MeshModelingTools/Private/EditMeshPolygonsTool.cpp b/Engine/Plugins/Runtime/MeshModelingToolset/Source/MeshModelingTools/Private/EditMeshPolygonsTool.cpp
--- a/Engine/Plugins/Runtime/MeshModelingToolset/Source/MeshModelingTools/Private/EditMeshPolygonsTool.cpp
+++ b/Engine/Plugins/Runtime/MeshModelingToolset/Source/MeshModelingTools/Private/EditMeshPolygonsTool.cpp
@@ -2563,7 +2563,7 @@ void UEditMeshPolygonsTool::ApplyCollapseSingleEdge()
 	for (FSelectedEdge& Edge : ActiveEdgeSelection)
 	{
 		int32 eid = Edge.EdgeIDs[0];
-		if (Mesh->IsEdge(eid) && Mesh->Attributes()->IsSeamEdge(eid) == false)
+		if (Mesh->IsEdge(eid))
 		{
 			ValidEdgeIDs.Add(eid);
 		}
