Make modeling mode element selections auto-convert when user changes the active selection mode (if a conversion is implemented)

diff --git a/Engine/Plugins/Editor/ModelingToolsEditorMode/Source/ModelingToolsEditorMode/Private/ModelingToolsEditorMode.cpp b/Engine/Plugins/Editor/ModelingToolsEditorMode/Source/ModelingToolsEditorMode/Private/ModelingToolsEditorMode.cpp
--- a/Engine/Plugins/Editor/ModelingToolsEditorMode/Source/ModelingToolsEditorMode/Private/ModelingToolsEditorMode.cpp
+++ b/Engine/Plugins/Editor/ModelingToolsEditorMode/Source/ModelingToolsEditorMode/Private/ModelingToolsEditorMode.cpp
@@ -875,8 +875,7 @@ void UModelingToolsEditorMode::Enter()
 				if ( GetToolManager() && GetToolManager()->GetContextTransactionsAPI() && GetSelectionManager() )
 				{
 					GetToolManager()->GetContextTransactionsAPI()->BeginUndoTransaction(LOCTEXT("ChangeSelectionMode", "Selection Mode"));
-					GetSelectionManager()->SetMeshTopologyMode(TopoMode); 
-					GetSelectionManager()->SetSelectionElementType(ElementMode);
+					GetSelectionManager()->SetMeshSelectionTypeAndMode(ElementMode, TopoMode, true);
 					GetToolManager()->GetContextTransactionsAPI()->EndUndoTransaction();
 					if (UModelingToolsModeCustomizationSettings* ModelingEditorSettings = GetMutableDefault<UModelingToolsModeCustomizationSettings>())
 					{
diff --git a/Engine/Plugins/Runtime/MeshModelingToolset/Source/ModelingComponents/Private/Selection/GeometrySelectionManager.cpp b/Engine/Plugins/Runtime/MeshModelingToolset/Source/ModelingComponents/Private/Selection/GeometrySelectionManager.cpp
--- a/Engine/Plugins/Runtime/MeshModelingToolset/Source/ModelingComponents/Private/Selection/GeometrySelectionManager.cpp
+++ b/Engine/Plugins/Runtime/MeshModelingToolset/Source/ModelingComponents/Private/Selection/GeometrySelectionManager.cpp
@@ -206,6 +206,64 @@ void UGeometrySelectionManager::SetMeshTopologyMode(EMeshTopologyMode NewTopolog
 }
 
 
+void UGeometrySelectionManager::SetMeshSelectionTypeAndMode(EGeometryElementType NewElementType, EMeshTopologyMode NewTopologyMode, bool bConvertSelection)
+{
+	if (MeshTopologyMode != NewTopologyMode || SelectionElementType != NewElementType)
+	{
+		bool bHasSelection = HasSelection();
+
+		// If we're converting selections, save the old one; we will re-add it after changing the mode
+		TArray<FGeometrySelection> OldTypeSelections;
+		if (bHasSelection && bConvertSelection)
+		{
+			for (TSharedPtr<FGeometrySelectionTarget> Target : ActiveTargetReferences)
+			{
+				OldTypeSelections.Add(Target->Selection);
+			}
+		}
+
+		ClearSelection();
+
+		GetTransactionsAPI()->BeginUndoTransaction(LOCTEXT("ChangeElementMethod", "Change Selection Method"));
+
+		// We have to undo/redo the change to the selection type because if we want to 'undo' this later and restore
+		// the current selection, we need the active element type to be correct. Note that it goes *after* the Clear
+		// so that when we undo, we change to the correct type before we restore
+		TUniquePtr<FGeometrySelectionManager_SelectionTypeChange> TypeChange = MakeUnique<FGeometrySelectionManager_SelectionTypeChange>();
+		TypeChange->FromElementType = SelectionElementType;
+		TypeChange->ToElementType = NewElementType;
+		TypeChange->FromTopologyMode = MeshTopologyMode;
+		TypeChange->ToTopologyMode = NewTopologyMode;
+		GetTransactionsAPI()->AppendChange(this, MoveTemp(TypeChange), LOCTEXT("ChangeElementMethod", "Change Selection Method"));
+
+		SetSelectionElementTypeInternal(NewElementType);
+		SetMeshTopologyModeInternal(NewTopologyMode);
+
+		if (bHasSelection && bConvertSelection && ensure(ActiveTargetReferences.Num() == OldTypeSelections.Num()))
+		{
+			for (int32 TargetIdx = 0; TargetIdx < ActiveTargetReferences.Num(); ++TargetIdx)
+			{
+				// Add back the old selection, converted to the new mode/type
+				TSharedPtr<FGeometrySelectionTarget> Target = ActiveTargetReferences[TargetIdx];
+				FGeometrySelection InitialSelection = Target->Selection;
+				FGeometrySelectionDelta AfterDelta;
+				Target->Selector->UpdateSelectionFromSelection(OldTypeSelections[TargetIdx], true, *Target->SelectionEditor, FGeometrySelectionUpdateConfig{EGeometrySelectionChangeType::Replace}, &AfterDelta);
+				if (!AfterDelta.IsEmpty())
+				{
+					TUniquePtr<FGeometrySelectionReplaceChange> NewSelectionChange = MakeUnique<FGeometrySelectionReplaceChange>();
+					NewSelectionChange->Identifier = Target->TargetIdentifier;
+					NewSelectionChange->After = Target->Selection;
+					NewSelectionChange->Before = InitialSelection;
+					GetTransactionsAPI()->AppendChange(this, MoveTemp(NewSelectionChange), LOCTEXT("ConvertSelection", "Convert Selection"));
+				}
+			}
+		}
+
+		GetTransactionsAPI()->EndUndoTransaction();
+
+	}
+}
+
 
 EGeometryTopologyType UGeometrySelectionManager::GetSelectionTopologyType() const
 {
diff --git a/Engine/Plugins/Runtime/MeshModelingToolset/Source/ModelingComponents/Public/Selection/GeometrySelectionManager.h b/Engine/Plugins/Runtime/MeshModelingToolset/Source/ModelingComponents/Public/Selection/GeometrySelectionManager.h
--- a/Engine/Plugins/Runtime/MeshModelingToolset/Source/ModelingComponents/Public/Selection/GeometrySelectionManager.h
+++ b/Engine/Plugins/Runtime/MeshModelingToolset/Source/ModelingComponents/Public/Selection/GeometrySelectionManager.h
@@ -100,6 +100,9 @@ public:
 	virtual EMeshTopologyMode GetMeshTopologyMode() const { return MeshTopologyMode; }
 	virtual EGeometryTopologyType GetSelectionTopologyType() const;
 
+	// Switch the selection mode and type, optionally converting any existing selection to the new type and mode
+	virtual void SetMeshSelectionTypeAndMode(EGeometryElementType NewElementType, EMeshTopologyMode NewSelectionMode, bool bConvertSelection);
+
 	
 
 
