Fix mesh elements visualizer tooltip typo

diff --git a/Engine/Plugins/Runtime/MeshModelingToolset/Source/ModelingComponents/Public/Drawing/MeshElementsVisualizer.h b/Engine/Plugins/Runtime/MeshModelingToolset/Source/ModelingComponents/Public/Drawing/MeshElementsVisualizer.h
--- a/Engine/Plugins/Runtime/MeshModelingToolset/Source/ModelingComponents/Public/Drawing/MeshElementsVisualizer.h
+++ b/Engine/Plugins/Runtime/MeshModelingToolset/Source/ModelingComponents/Public/Drawing/MeshElementsVisualizer.h
@@ -19,7 +19,7 @@ class MODELINGCOMPONENTS_API UMeshElementsVisualizerProperties : public UInterac
 {
 	GENERATED_BODY()
 public:
-	/** Should any be mesh elements be shown */
+	/** Should any mesh elements be shown */
 	UPROPERTY(EditAnywhere, Category = MeshElementVisualization)
 	bool bVisible = true;
 
@@ -47,7 +47,7 @@ public:
 	UPROPERTY(EditAnywhere, Category = MeshElementVisualization, meta = (EditCondition = "bVisible"))
 	bool bShowColorSeams = true;
 
-	/** multiplier on edge thicknesses */
+	/** Multiplier on edge thicknesses */
 	UPROPERTY(EditAnywhere, Category = MeshElementVisualization, AdvancedDisplay, meta = (UIMin = 0.1, UIMax = 10.0, EditCondition = "bVisible"))
 	float ThicknessScale = 1.0;
 
@@ -75,7 +75,7 @@ public:
 	UPROPERTY(EditAnywhere, Category = MeshElementVisualization, AdvancedDisplay, meta = (EditCondition = "bVisible"))
 	FColor ColorSeamColor = FColor(46, 204, 113);
 
-	/** depth bias used to slightly shift depth of lines */
+	/** Depth bias used to slightly shift depth of lines */
 	UPROPERTY(EditAnywhere, Category = MeshElementVisualization, AdvancedDisplay, meta = (UIMin = -2.0, UIMax = 2.0, EditCondition = "bVisible"))
 	float DepthBias = 0.2;
 
