Make draw polygon, draw polypath, and draw-and-revolve tools initialize with a 'reasonable' drawing frame instead of starting at the origin

Where reasonable means centered on the selection (if available), or otherwise centered in the view.
If a geometry element is selected, we will also align to that (e.g., aligning to the plane of a selected triangle)

diff --git a/Engine/Plugins/Experimental/MeshModelingToolsetExp/Source/MeshModelingToolsExp/Private/DrawPolyPathTool.cpp b/Engine/Plugins/Experimental/MeshModelingToolsetExp/Source/MeshModelingToolsExp/Private/DrawPolyPathTool.cpp
--- a/Engine/Plugins/Experimental/MeshModelingToolsetExp/Source/MeshModelingToolsExp/Private/DrawPolyPathTool.cpp
+++ b/Engine/Plugins/Experimental/MeshModelingToolsetExp/Source/MeshModelingToolsExp/Private/DrawPolyPathTool.cpp
@@ -160,6 +160,7 @@ UInteractiveTool* UDrawPolyPathToolBuilder::BuildTool(const FToolBuilderState& S
 {
 	UDrawPolyPathTool* NewTool = NewObject<UDrawPolyPathTool>(SceneState.ToolManager);
 	NewTool->SetWorld(SceneState.World);
+	NewTool->SetInitialDrawFrame(ToolSetupUtil::GetDefaultWorldReferenceFrame(SceneState.ToolManager));
 	return NewTool;
 }
 
@@ -186,8 +187,6 @@ void UDrawPolyPathTool::Setup()
 	HoverBehavior->Modifiers.RegisterModifier(ShiftModifierID, FInputDeviceState::IsShiftKeyDown);
 	AddInputBehavior(HoverBehavior);
 
-	DrawPlaneWorld = FFrame3d();
-
 	PlaneMechanic = NewObject<UConstructionPlaneMechanic>(this);
 	PlaneMechanic->Setup(this);
 	PlaneMechanic->CanUpdatePlaneFunc = [this]() { return CanUpdateDrawPlane(); };
diff --git a/Engine/Plugins/Experimental/MeshModelingToolsetExp/Source/MeshModelingToolsExp/Public/DrawPolyPathTool.h b/Engine/Plugins/Experimental/MeshModelingToolsetExp/Source/MeshModelingToolsExp/Public/DrawPolyPathTool.h
--- a/Engine/Plugins/Experimental/MeshModelingToolsetExp/Source/MeshModelingToolsExp/Public/DrawPolyPathTool.h
+++ b/Engine/Plugins/Experimental/MeshModelingToolsetExp/Source/MeshModelingToolsExp/Public/DrawPolyPathTool.h
@@ -192,6 +192,11 @@ public:
 	// IModifierToggleBehaviorTarget
 	virtual void OnUpdateModifierState(int ModifierID, bool bIsOn) override;
 
+	void SetInitialDrawFrame(UE::Geometry::FFrame3d Frame)
+	{
+		DrawPlaneWorld = Frame;
+	}
+
 protected:
 	UWorld* TargetWorld;
 
diff --git a/Engine/Plugins/Runtime/MeshModelingToolset/Source/MeshModelingTools/Private/DrawAndRevolveTool.cpp b/Engine/Plugins/Runtime/MeshModelingToolset/Source/MeshModelingTools/Private/DrawAndRevolveTool.cpp
--- a/Engine/Plugins/Runtime/MeshModelingToolset/Source/MeshModelingTools/Private/DrawAndRevolveTool.cpp
+++ b/Engine/Plugins/Runtime/MeshModelingToolset/Source/MeshModelingTools/Private/DrawAndRevolveTool.cpp
@@ -36,8 +36,9 @@ bool UDrawAndRevolveToolBuilder::CanBuildTool(const FToolBuilderState& SceneStat
 UInteractiveTool* UDrawAndRevolveToolBuilder::BuildTool(const FToolBuilderState& SceneState) const
 {
 	UDrawAndRevolveTool* NewTool = NewObject<UDrawAndRevolveTool>(SceneState.ToolManager);
-
 	NewTool->SetWorld(SceneState.World);
+	FQuaterniond DefaultOrientation = (FQuaterniond)FRotator(90, 0, 0);
+	NewTool->SetInitialDrawFrame(ToolSetupUtil::GetDefaultWorldReferenceFrame(SceneState.ToolManager, DefaultOrientation));
 	return NewTool;
 }
 
@@ -127,6 +128,9 @@ void UDrawAndRevolveTool::Setup()
 	Settings->bAllowedToEditDrawPlane = true;
 	AddToolPropertySource(Settings);
 
+	Settings->DrawPlaneOrigin = InitialDrawFrame.Origin;
+	Settings->DrawPlaneOrientation = (FRotator)InitialDrawFrame.Rotation;
+
 	MaterialProperties = NewObject<UNewMeshMaterialProperties>(this);
 	AddToolPropertySource(MaterialProperties);
 	MaterialProperties->RestoreProperties(this);
@@ -168,16 +172,6 @@ void UDrawAndRevolveTool::Setup()
 
 	UpdateRevolutionAxis();
 
-	FViewCameraState InitialCameraState;
-	GetToolManager()->GetContextQueriesAPI()->GetCurrentViewState(InitialCameraState);
-	if (FVector::DistSquared(InitialCameraState.Position, Settings->DrawPlaneOrigin) > FarDrawPlaneThreshold * FarDrawPlaneThreshold)
-	{
-		bHasFarPlaneWarning = true;
-		GetToolManager()->DisplayMessage(
-			LOCTEXT("FarDrawPlane", "The axis of revolution is far from the camera. Note that you can ctrl-click to place the axis on a visible surface."),
-			EToolMessageLevel::UserWarning);
-	}
-
 	// The plane mechanic lets us update the plane in which we draw the profile curve, as long as we haven't
 	// started adding points to it already.
 	FFrame3d ProfileDrawPlane(Settings->DrawPlaneOrigin, Settings->DrawPlaneOrientation.Quaternion());
@@ -198,11 +192,6 @@ void UDrawAndRevolveTool::Setup()
 			ControlPointsMechanic->SetPlane(PlaneMechanic->Plane);
 		}
 		UpdateRevolutionAxis();
-		if (bHasFarPlaneWarning) // if the user has changed the plane, no longer need a warning
-		{
-			bHasFarPlaneWarning = false;
-			GetToolManager()->DisplayMessage(FText(), EToolMessageLevel::UserWarning);
-		}
 		});
 
 	ControlPointsMechanic->SetPlane(PlaneMechanic->Plane);
@@ -345,13 +334,6 @@ void UDrawAndRevolveTool::Render(IToolsContextRenderAPI* RenderAPI)
 {
 	GetToolManager()->GetContextQueriesAPI()->GetCurrentViewState(CameraState);
 
-	if (bHasFarPlaneWarning && FVector::DistSquared(CameraState.Position, Settings->DrawPlaneOrigin) < FarDrawPlaneThreshold * FarDrawPlaneThreshold)
-	{
-		// if camera is now closer to the axis, no longer need a warning
-		bHasFarPlaneWarning = false;
-		GetToolManager()->DisplayMessage(FText(), EToolMessageLevel::UserWarning);
-	}
-
 	if (PlaneMechanic != nullptr)
 	{
 		PlaneMechanic->Render(RenderAPI);
diff --git a/Engine/Plugins/Runtime/MeshModelingToolset/Source/MeshModelingTools/Private/DrawPolygonTool.cpp b/Engine/Plugins/Runtime/MeshModelingToolset/Source/MeshModelingTools/Private/DrawPolygonTool.cpp
--- a/Engine/Plugins/Runtime/MeshModelingToolset/Source/MeshModelingTools/Private/DrawPolygonTool.cpp
+++ b/Engine/Plugins/Runtime/MeshModelingToolset/Source/MeshModelingTools/Private/DrawPolygonTool.cpp
@@ -33,6 +33,7 @@
 #include "Drawing/MeshDebugDrawing.h"
 
 #include "Selection/SelectClickedAction.h"
+#include "Selection/StoredMeshSelectionUtil.h"
 #include "Selection/ToolSelectionUtil.h"
 #include "ModelingObjectsCreationAPI.h"
 #include "Mechanics/ConstructionPlaneMechanic.h"
@@ -61,6 +62,7 @@ UInteractiveTool* UDrawPolygonToolBuilder::BuildTool(const FToolBuilderState& Sc
 {
 	UDrawPolygonTool* NewTool = NewObject<UDrawPolygonTool>(SceneState.ToolManager);
 	NewTool->SetWorld(SceneState.World);
+	NewTool->SetInitialDrawFrame(ToolSetupUtil::GetDefaultWorldReferenceFrame(SceneState.ToolManager));
 	return NewTool;
 }
 
@@ -109,7 +111,7 @@ void UDrawPolygonTool::Setup()
 	PlaneMechanic->Setup(this);
 	PlaneMechanic->CanUpdatePlaneFunc = [this]() { return AllowDrawPlaneUpdates(); };
 	PlaneMechanic->OnPlaneChanged.AddLambda([this]() { SnapEngine.Plane = PlaneMechanic->Plane; });	// Keep SnapEngine plane up to date with PlaneMechanic's plane
-	PlaneMechanic->Initialize(TargetWorld, FFrame3d());
+	PlaneMechanic->Initialize(TargetWorld, InitialDrawFrame);
 	
 	DragAlignmentMechanic = NewObject<UDragAlignmentMechanic>(this);
 	DragAlignmentMechanic->Setup(this);
diff --git a/Engine/Plugins/Runtime/MeshModelingToolset/Source/MeshModelingTools/Public/DrawAndRevolveTool.h b/Engine/Plugins/Runtime/MeshModelingToolset/Source/MeshModelingTools/Public/DrawAndRevolveTool.h
--- a/Engine/Plugins/Runtime/MeshModelingToolset/Source/MeshModelingTools/Public/DrawAndRevolveTool.h
+++ b/Engine/Plugins/Runtime/MeshModelingToolset/Source/MeshModelingTools/Public/DrawAndRevolveTool.h
@@ -108,6 +108,11 @@ public:
 
 	virtual void OnPropertyModified(UObject* PropertySet, FProperty* Property) override;
 
+	void SetInitialDrawFrame(UE::Geometry::FFrame3d InFrame)
+	{
+		InitialDrawFrame = InFrame;
+	}
+
 protected:
 
 	UWorld* TargetWorld;
@@ -119,6 +124,9 @@ protected:
 	FVector3d RevolutionAxisOrigin;
 	FVector3d RevolutionAxisDirection;
 
+	// The initial frame, used in tool setup to place the axis
+	UE::Geometry::FFrame3d InitialDrawFrame;
+
 	bool bProfileCurveComplete = false;
 
 	void UpdateRevolutionAxis();
@@ -147,8 +155,4 @@ protected:
 	void GenerateAsset(const FDynamicMeshOpResult& Result);
 
 	friend class URevolveOperatorFactory;
-
-private:
-	constexpr static double FarDrawPlaneThreshold = 100 * 100;
-	bool bHasFarPlaneWarning = false;
 };
diff --git a/Engine/Plugins/Runtime/MeshModelingToolset/Source/MeshModelingTools/Public/DrawPolygonTool.h b/Engine/Plugins/Runtime/MeshModelingToolset/Source/MeshModelingTools/Public/DrawPolygonTool.h
--- a/Engine/Plugins/Runtime/MeshModelingToolset/Source/MeshModelingTools/Public/DrawPolygonTool.h
+++ b/Engine/Plugins/Runtime/MeshModelingToolset/Source/MeshModelingTools/Public/DrawPolygonTool.h
@@ -216,6 +216,11 @@ public:
 
 	virtual void ApplyUndoPoints(const TArray<FVector3d>& ClickPointsIn, const TArray<FVector3d>& PolygonVerticesIn);
 
+	void SetInitialDrawFrame(UE::Geometry::FFrame3d Frame)
+	{
+		InitialDrawFrame = Frame;
+	}
+
 
 protected:
 	// flags used to identify modifier keys/buttons
@@ -260,6 +265,8 @@ protected:
 
 	// polygon drawing
 
+	UE::Geometry::FFrame3d InitialDrawFrame;
+
 	bool bAbortActivePolygonDraw;
 
 	bool bInFixedPolygonMode = false;
diff --git a/Engine/Plugins/Runtime/MeshModelingToolset/Source/ModelingComponents/Private/Selection/StoredMeshSelectionUtil.cpp b/Engine/Plugins/Runtime/MeshModelingToolset/Source/ModelingComponents/Private/Selection/StoredMeshSelectionUtil.cpp
--- a/Engine/Plugins/Runtime/MeshModelingToolset/Source/ModelingComponents/Private/Selection/StoredMeshSelectionUtil.cpp
+++ b/Engine/Plugins/Runtime/MeshModelingToolset/Source/ModelingComponents/Private/Selection/StoredMeshSelectionUtil.cpp
@@ -42,6 +42,43 @@ bool UE::Geometry::GetCurrentGeometrySelectionForTarget(UInteractiveToolManager*
 	return bHaveSelection;
 }
 
+bool UE::Geometry::GetCurrentSelectionWorldFrameBounds(const FToolBuilderState& SceneState, FFrame3d& OutWorldFrame, FAxisAlignedBox3d& OutWorldBounds, bool& bOutIsElementSelection)
+{
+	UGeometrySelectionManager* SelectionManager = SceneState.ToolManager->GetContextObjectStore()->FindContext<UGeometrySelectionManager>();
+	if (SelectionManager && SelectionManager->HasSelection()) 
+	{
+		bOutIsElementSelection = true;
+		FGeometrySelectionBounds Bounds;
+		SelectionManager->GetSelectionBounds(Bounds);
+		OutWorldBounds = Bounds.WorldBounds;
+		SelectionManager->GetSelectionWorldFrame(OutWorldFrame);
+		return true;
+	}
+	else // use the selected actor/components
+	{
+		bOutIsElementSelection = false;
+		FBox Bounds;
+		for (AActor* Actor : SceneState.SelectedActors)
+		{
+			Bounds += Actor->GetComponentsBoundingBox(true);
+		}
+		for (UActorComponent* Component : SceneState.SelectedComponents)
+		{
+			if (UPrimitiveComponent* PrimComponent = Cast<UPrimitiveComponent>(Component))
+			{
+				Bounds += PrimComponent->Bounds.GetBox();
+			}
+		}
+		if (Bounds.IsValid)
+		{
+			OutWorldBounds = (FAxisAlignedBox3d)Bounds;
+			OutWorldFrame = FFrame3d(Bounds.GetCenter());
+			return true;
+		}
+	}
+	return false;
+}
+
 
 bool UE::Geometry::SetToolOutputGeometrySelectionForTarget(UInteractiveTool* Tool, UToolTarget* Target, const FGeometrySelection& OutputSelection)
 {
diff --git a/Engine/Plugins/Runtime/MeshModelingToolset/Source/ModelingComponents/Private/ToolSetupUtil.cpp b/Engine/Plugins/Runtime/MeshModelingToolset/Source/ModelingComponents/Private/ToolSetupUtil.cpp
--- a/Engine/Plugins/Runtime/MeshModelingToolset/Source/ModelingComponents/Private/ToolSetupUtil.cpp
+++ b/Engine/Plugins/Runtime/MeshModelingToolset/Source/ModelingComponents/Private/ToolSetupUtil.cpp
@@ -4,6 +4,7 @@
 #include "ToolSetupUtil.h"
 #include "ModelingComponentsSettings.h"
 #include "Curves/CurveFloat.h"
+#include "FrameTypes.h"
 #include "InteractiveTool.h"
 #include "InteractiveToolManager.h"
 #include "MaterialDomain.h"
@@ -13,6 +14,10 @@
 #include "Components/BaseDynamicMeshComponent.h"
 #include "PreviewMesh.h"
 
+#include "SceneQueries/SceneSnappingManager.h"
+#include "Selection/StoredMeshSelectionUtil.h"
+#include "ToolContextInterfaces.h"
+#include "ToolSceneQueriesUtil.h"
 
 UMaterialInterface* ToolSetupUtil::GetDefaultMaterial()
 {
@@ -390,4 +395,41 @@ void ToolSetupUtil::ApplyRenderingConfigurationToPreview(UPreviewMesh* PreviewMe
 		return;
 	}
 	ApplyRenderingConfigurationToPreview(Component, SourceTarget);
-}
\ No newline at end of file
+}
+
+UE::Geometry::FFrame3d ToolSetupUtil::GetDefaultWorldReferenceFrame(UInteractiveToolManager* ToolManager, UE::Geometry::FQuaterniond DefaultOrientation, double DefaultPlacementDistance)
+{
+	using namespace UE::Geometry;
+
+	FFrame3d Frame;
+	FAxisAlignedBox3d Bounds;
+	bool bIsElementSelection;
+	FToolBuilderState SceneState;
+	ToolManager->GetContextQueriesAPI()->GetCurrentSelectionState(SceneState);
+	if (GetCurrentSelectionWorldFrameBounds(SceneState, Frame, Bounds, bIsElementSelection))
+	{
+		if (bIsElementSelection)
+		{
+			return Frame;
+		}
+		else
+		{
+			Frame.Rotation = DefaultOrientation;
+			return Frame;
+		}
+	}
+	else
+	{
+		FViewCameraState CameraState;
+		ToolManager->GetContextQueriesAPI()->GetCurrentViewState(CameraState);
+		FHitResult WorldHitResult;
+		FRay Ray(CameraState.Position, CameraState.Forward(), true);
+		double HitParam = DefaultPlacementDistance;
+		if (ToolSceneQueriesUtil::FindNearestVisibleObjectHit(USceneSnappingManager::Find(ToolManager), WorldHitResult, Ray))
+		{
+			HitParam = (double)Ray.GetParameter(WorldHitResult.ImpactPoint);
+		}
+
+		return FFrame3d(CameraState.Position + CameraState.Forward() * HitParam, DefaultOrientation);
+	}
+}
diff --git a/Engine/Plugins/Runtime/MeshModelingToolset/Source/ModelingComponents/Public/Selection/StoredMeshSelectionUtil.h b/Engine/Plugins/Runtime/MeshModelingToolset/Source/ModelingComponents/Public/Selection/StoredMeshSelectionUtil.h
--- a/Engine/Plugins/Runtime/MeshModelingToolset/Source/ModelingComponents/Public/Selection/StoredMeshSelectionUtil.h
+++ b/Engine/Plugins/Runtime/MeshModelingToolset/Source/ModelingComponents/Public/Selection/StoredMeshSelectionUtil.h
@@ -41,6 +41,16 @@ namespace Geometry
 	MODELINGCOMPONENTS_API bool GetCurrentGeometrySelectionForTarget(UInteractiveToolManager* ToolManager, UToolTarget* Target,
 		FGeometrySelection& SelectionOut);
 
+	/**
+	 * Find the frame and bounds representative of the current selection, if there is one.
+	 * @param SceneState				Scene state info, including current selection
+	 * @param OutWorldFrame				Frame centered on the current selection and (when sensible) aligned to the selected elements
+	 * @param OutWorldBounds			Bounds of the current selection
+	 * @param bOutIsElementSelection	Whether the selection is a mesh element selection, or a regular editor/actor selection
+	 * @return	true if a selection was found, false otherwise.
+	 */
+	MODELINGCOMPONENTS_API bool GetCurrentSelectionWorldFrameBounds(const FToolBuilderState& SceneState, FFrame3d& OutWorldFrame, FAxisAlignedBox3d& OutWorldBounds, bool& bOutIsElementSelection);
+
 	/**
 	 * Allow a Tool to return an "output" FGeometrySelection for the given ToolTarget (presumably on Tool Shutdown)
 	 * Target must be an "active" selection target in the current UGeometrySelectionManager.
diff --git a/Engine/Plugins/Runtime/MeshModelingToolset/Source/ModelingComponents/Public/ToolSetupUtil.h b/Engine/Plugins/Runtime/MeshModelingToolset/Source/ModelingComponents/Public/ToolSetupUtil.h
--- a/Engine/Plugins/Runtime/MeshModelingToolset/Source/ModelingComponents/Public/ToolSetupUtil.h
+++ b/Engine/Plugins/Runtime/MeshModelingToolset/Source/ModelingComponents/Public/ToolSetupUtil.h
@@ -4,6 +4,7 @@
 
 #include "CoreMinimal.h"
 #include "Materials/MaterialInstanceDynamic.h"
+#include "FrameTypes.h"
 
 class UCurveFloat;
 class UTexture;
@@ -205,7 +206,13 @@ namespace ToolSetupUtil
 	MODELINGCOMPONENTS_API UCurveFloat* GetContrastAdjustmentCurve(UInteractiveToolManager* ToolManager);
 
 
-
+	/**
+	 * Use the current selection, if any, or otherwise the camera position to find a relevant reference frame for creating new objects.
+	 * Note: Should be called in tool setup to access the active geometry element selection (if any) before it is cleared by the tool.
+	 * @param DefaultOrientation			Orientation to use if there is no geometry element selection providing a natural orientation
+	 * @param DefaultPlacementDistance		Distance from camera to place the frame if there is no selection and a raycast doesn't hit anything
+	 */
+	MODELINGCOMPONENTS_API UE::Geometry::FFrame3d GetDefaultWorldReferenceFrame(UInteractiveToolManager* ToolManager, UE::Geometry::FQuaterniond DefaultOrientation = UE::Geometry::FQuaterniond(), double NoSelectionPlacementDistance = 500);
 
 
 	//
