Build fix
#cis fix

diff --git a/Engine/Plugins/Experimental/Enterprise/DataprepGeometryOperations/Source/Private/DataprepGeometrySelectionTransforms.cpp b/Engine/Plugins/Experimental/Enterprise/DataprepGeometryOperations/Source/Private/DataprepGeometrySelectionTransforms.cpp
--- a/Engine/Plugins/Experimental/Enterprise/DataprepGeometryOperations/Source/Private/DataprepGeometrySelectionTransforms.cpp
+++ b/Engine/Plugins/Experimental/Enterprise/DataprepGeometryOperations/Source/Private/DataprepGeometrySelectionTransforms.cpp
@@ -14,6 +14,7 @@
 #include "MeshMergeModule.h"
 #include "MeshMergeData.h"
 #include "MeshAttributes.h"
+#include "MeshMerge/MeshMergingSettings.h"
 #include "Modules/ModuleManager.h"
 #include "StaticMeshOperations.h"
 
diff --git a/Engine/Plugins/Runtime/MeshModelingToolset/Source/ModelingOperatorsEditorOnly/Private/CleaningOps/SimplifyMeshOp.cpp b/Engine/Plugins/Runtime/MeshModelingToolset/Source/ModelingOperatorsEditorOnly/Private/CleaningOps/SimplifyMeshOp.cpp
--- a/Engine/Plugins/Runtime/MeshModelingToolset/Source/ModelingOperatorsEditorOnly/Private/CleaningOps/SimplifyMeshOp.cpp
+++ b/Engine/Plugins/Runtime/MeshModelingToolset/Source/ModelingOperatorsEditorOnly/Private/CleaningOps/SimplifyMeshOp.cpp
@@ -8,6 +8,7 @@
 #include "DynamicMesh/DynamicMesh3.h"
 #include "DynamicMesh/DynamicMeshAttributeSet.h"
 #include "DynamicMesh/MeshAttributeUtil.h"
+#include "MeshReductionSettings.h"
 #include "MeshSimplification.h"
 #include "MeshConstraintsUtil.h"
 #include "ProjectionTargets.h"
