fix gizmo not showing up for modeling mode geometry selection after selection element type / topology mode change

diff --git a/Engine/Plugins/Runtime/MeshModelingToolset/Source/ModelingComponents/Private/Selection/GeometrySelectionManager.cpp b/Engine/Plugins/Runtime/MeshModelingToolset/Source/ModelingComponents/Private/Selection/GeometrySelectionManager.cpp
--- a/Engine/Plugins/Runtime/MeshModelingToolset/Source/ModelingComponents/Private/Selection/GeometrySelectionManager.cpp
+++ b/Engine/Plugins/Runtime/MeshModelingToolset/Source/ModelingComponents/Private/Selection/GeometrySelectionManager.cpp
@@ -261,6 +261,7 @@ void UGeometrySelectionManager::SetMeshSelectionTypeAndMode(EGeometryElementType
 
 		GetTransactionsAPI()->EndUndoTransaction();
 
+		OnSelectionModified.Broadcast();
 	}
 }
 
