make poly edit bevel tool refuse to start when there are no bevel-able edges in the current selection + make bevel of polygroup faces succeed when only some of the selection-bordering edges are boundary edges

diff --git a/Engine/Plugins/Runtime/GeometryProcessing/Source/DynamicMesh/Private/Operations/MeshBevel.cpp b/Engine/Plugins/Runtime/GeometryProcessing/Source/DynamicMesh/Private/Operations/MeshBevel.cpp
--- a/Engine/Plugins/Runtime/GeometryProcessing/Source/DynamicMesh/Private/Operations/MeshBevel.cpp
+++ b/Engine/Plugins/Runtime/GeometryProcessing/Source/DynamicMesh/Private/Operations/MeshBevel.cpp
@@ -131,11 +131,44 @@ void FMeshBevel::InitializeFromGroupTopologyEdges(const FDynamicMesh3& Mesh, con
 
 bool FMeshBevel::InitializeFromGroupTopologyFaces(const FDynamicMesh3& Mesh, const FGroupTopology& Topology, const TArray<int32>& GroupFaces)
 {
-	FGroupTopologySelection Selection;
-	Selection.SelectedGroupIDs.Append(GroupFaces);
-	TArray<int32> Triangles;
-	Topology.GetSelectedTriangles(Selection, Triangles);
-	return InitializeFromTriangleSet(Mesh, Triangles);
+	TSet<int32> GroupSelection;
+	GroupSelection.Append(GroupFaces);
+	bool bFoundAnythingToBevel = false;
+	for (int32 GroupID : GroupFaces)
+	{
+		const FGroupTopology::FGroup* Group = Topology.FindGroupByID(GroupID);
+		for (const FGroupTopology::FGroupBoundary& Boundary : Group->Boundaries)
+		{
+			for (int32 GroupEdgeID : Boundary.GroupEdges)
+			{
+				const FGroupTopology::FGroupEdge& GroupEdge = Topology.Edges[GroupEdgeID];
+				int32 OtherGroupID = GroupEdge.OtherGroupID(GroupID);
+				// Do not bevel edges where both sides are selected, or boundary edges
+				if (OtherGroupID == IndexConstants::InvalidID || GroupSelection.Contains(OtherGroupID))
+				{
+					continue;
+				}
+				if (Topology.IsIsolatedLoop(GroupEdgeID))
+				{
+					FEdgeLoop NewLoop;
+					NewLoop.InitializeFromEdges(&Mesh, Topology.Edges[GroupEdgeID].Span.Edges);
+					AddBevelEdgeLoop(Mesh, NewLoop);
+					bFoundAnythingToBevel = true;
+				}
+				else
+				{
+					AddBevelGroupEdge(Mesh, Topology, GroupEdgeID);
+					bFoundAnythingToBevel = true;
+				}
+			}
+		}
+	}
+	if (bFoundAnythingToBevel)
+	{
+		// precompute topological information necessary to apply bevel to vertices/edges/loops
+		BuildVertexSets(Mesh);
+	}
+	return bFoundAnythingToBevel;
 }
 
 bool FMeshBevel::InitializeFromTriangleSet(const FDynamicMesh3& Mesh, const TArray<int32>& Triangles)
diff --git a/Engine/Plugins/Runtime/MeshModelingToolset/Source/MeshModelingTools/Private/ToolActivities/PolyEditBevelEdgeActivity.cpp b/Engine/Plugins/Runtime/MeshModelingToolset/Source/MeshModelingTools/Private/ToolActivities/PolyEditBevelEdgeActivity.cpp
--- a/Engine/Plugins/Runtime/MeshModelingToolset/Source/MeshModelingTools/Private/ToolActivities/PolyEditBevelEdgeActivity.cpp
+++ b/Engine/Plugins/Runtime/MeshModelingToolset/Source/MeshModelingTools/Private/ToolActivities/PolyEditBevelEdgeActivity.cpp
@@ -178,7 +178,35 @@ bool UPolyEditBevelEdgeActivity::CanStart() const
 	if (ActivityContext)
 	{
 		const FGroupTopologySelection& Selection = ActivityContext->SelectionMechanic->GetActiveSelection();
-		return (Selection.SelectedEdgeIDs.Num() > 0) ||  (Selection.SelectedGroupIDs.Num() > 0);
+		// Search for any edges that we can bevel
+		for (int32 GroupEdge : Selection.SelectedEdgeIDs)
+		{
+			// Can't bevel boundary edges
+			if (ActivityContext->CurrentTopology->Edges[GroupEdge].Groups.B != IndexConstants::InvalidID)
+			{
+				return true;
+			}
+		}
+		for (int32 GroupID : Selection.SelectedGroupIDs)
+		{
+			const FGroupTopology::FGroup* Group = ActivityContext->CurrentTopology->FindGroupByID(GroupID);
+			if (Group)
+			{
+				for (const FGroupTopology::FGroupBoundary& Boundary : Group->Boundaries)
+				{
+					for (int32 GroupEdgeID : Boundary.GroupEdges)
+					{
+						const FGroupTopology::FGroupEdge& GroupEdge = ActivityContext->CurrentTopology->Edges[GroupEdgeID];
+						int32 OtherGroupID = GroupEdge.OtherGroupID(GroupID);
+						// Only bevel non-boundary edges where one side is not in the selection
+						if (OtherGroupID != IndexConstants::InvalidID && !Selection.SelectedGroupIDs.Contains(OtherGroupID))
+						{
+							return true;
+						}
+					}
+				}
+			}
+		}
 	}
 	return false;
 }
@@ -187,7 +215,7 @@ EToolActivityStartResult UPolyEditBevelEdgeActivity::Start()
 {
 	if (!CanStart())
 	{
-		ParentTool->GetToolManager()->DisplayMessage( LOCTEXT("OnBevelFailedMesssage", "Action requires edge or face selection."), EToolMessageLevel::UserWarning);
+		ParentTool->GetToolManager()->DisplayMessage( LOCTEXT("OnBevelFailedMesssage", "Cannot Bevel with current selection: Make sure your selection includes (or borders) non-boundary edges."), EToolMessageLevel::UserWarning);
 		return EToolActivityStartResult::FailedStart;
 	}
 
