Added the ability to generate UVs with full floating point precision when building a new Mesh asset programmatically

diff --git a/Engine/Plugins/Runtime/MeshModelingToolset/Source/ModelingComponentsEditorOnly/Private/AssetUtils/CreateStaticMeshUtil.cpp b/Engine/Plugins/Runtime/MeshModelingToolset/Source/ModelingComponentsEditorOnly/Private/AssetUtils/CreateStaticMeshUtil.cpp
--- a/Engine/Plugins/Runtime/MeshModelingToolset/Source/ModelingComponentsEditorOnly/Private/AssetUtils/CreateStaticMeshUtil.cpp
+++ b/Engine/Plugins/Runtime/MeshModelingToolset/Source/ModelingComponentsEditorOnly/Private/AssetUtils/CreateStaticMeshUtil.cpp
@@ -72,7 +72,8 @@ UE::AssetUtils::ECreateStaticMeshResult UE::AssetUtils::CreateStaticMeshAsset(
 		BuildSettings.bRecomputeNormals = Options.bEnableRecomputeNormals;
 		BuildSettings.bRecomputeTangents = Options.bEnableRecomputeTangents;
 		BuildSettings.bGenerateLightmapUVs = Options.bGenerateLightmapUVs;
-
+		BuildSettings.bUseFullPrecisionUVs = Options.bUseFullPrecisionUVs;
+		
 		if (!Options.bAllowDistanceField)
 		{
 			BuildSettings.DistanceFieldResolutionScale = 0.0f;
diff --git a/Engine/Plugins/Runtime/MeshModelingToolset/Source/ModelingComponentsEditorOnly/Public/AssetUtils/CreateStaticMeshUtil.h b/Engine/Plugins/Runtime/MeshModelingToolset/Source/ModelingComponentsEditorOnly/Public/AssetUtils/CreateStaticMeshUtil.h
--- a/Engine/Plugins/Runtime/MeshModelingToolset/Source/ModelingComponentsEditorOnly/Public/AssetUtils/CreateStaticMeshUtil.h
+++ b/Engine/Plugins/Runtime/MeshModelingToolset/Source/ModelingComponentsEditorOnly/Public/AssetUtils/CreateStaticMeshUtil.h
@@ -79,6 +79,9 @@ namespace AssetUtils
 		// Whether to generate lightmap uvs for the generated mesh
 		bool bGenerateLightmapUVs = false;
 
+		// Whether to generate UVs with full floating point precision.
+		bool bUseFullPrecisionUVs = false;
+		
 		// Controls whether the UBodySetup on the Asset will be created (generally should be true)
 		bool bCreatePhysicsBody = true;
 		// set asset collision type
