Fix incorrect initialization of section index -> material index mapping to fix crash in copy mesh to static mesh geometry script method + add additional validation of polygon group IDs

diff --git a/Engine/Plugins/Runtime/MeshModelingToolset/Source/ModelingComponents/Private/ConversionUtils/SceneComponentToDynamicMesh.cpp b/Engine/Plugins/Runtime/MeshModelingToolset/Source/ModelingComponents/Private/ConversionUtils/SceneComponentToDynamicMesh.cpp
--- a/Engine/Plugins/Runtime/MeshModelingToolset/Source/ModelingComponents/Private/ConversionUtils/SceneComponentToDynamicMesh.cpp
+++ b/Engine/Plugins/Runtime/MeshModelingToolset/Source/ModelingComponents/Private/ConversionUtils/SceneComponentToDynamicMesh.cpp
@@ -103,7 +103,7 @@ namespace Private::ConversionHelper
 		check(Mesh);
 		TMap<int32, int32> SectionToMaterial;
 		const int32 NumMaterials = Mesh->GetStaticMaterials().Num();
-		int32 MaxMaterialIndex = -1;
+		int32 NumSectionIndex = 0;
 		if (bHighResLOD)
 		{
 			// custom path for HiResSource, where the section info map isn't available so we use mesh description slot names
@@ -119,16 +119,16 @@ namespace Private::ConversionHelper
 			int32 SectionIndex = 0;
 			for (FPolygonGroupID PolygonGroupID : MeshDescription->PolygonGroups().GetElementIDs())
 			{
-				int32 MaterialIndex = Mesh->GetStaticMaterials().IndexOfByPredicate(
+				int32 MaterialIndex = PolygonGroupID >= 0 && PolygonGroupID < MaterialSlotNames.GetNumElements() ? Mesh->GetStaticMaterials().IndexOfByPredicate(
 					[&MaterialSlotName = MaterialSlotNames[PolygonGroupID]](const FStaticMaterial& StaticMaterial) { return StaticMaterial.MaterialSlotName == MaterialSlotName; }
-				);
+					) : INDEX_NONE;
 				if (MaterialIndex != INDEX_NONE)
 				{
 					SectionToMaterial.Add(SectionIndex, MaterialIndex);
-					MaxMaterialIndex = FMath::Max(MaterialIndex, MaxMaterialIndex);
 				}
 				++SectionIndex;
 			}
+			NumSectionIndex = SectionIndex;
 		}
 		else
 		{
@@ -142,14 +142,13 @@ namespace Private::ConversionHelper
 				{
 					int32 MaterialIndex = SectionMap.Get(UseLOD, SectionIndex).MaterialIndex;
 					SectionToMaterial.Add(SectionIndex, MaterialIndex);
-					MaxMaterialIndex = FMath::Max(MaterialIndex, MaxMaterialIndex);
 				}
 			}
+			NumSectionIndex = LODSectionNum;
 		}
 
-		MaxMaterialIndex = FMath::Min(MaxMaterialIndex, NumMaterials - 1);
 		TArray<int32> Result;
-		Result.SetNumUninitialized(MaxMaterialIndex + 1);
+		Result.SetNumUninitialized(NumSectionIndex);
 		// Fill in identity mapping first to cover any unmapped indices
 		for (int32 Idx = 0; Idx < Result.Num(); ++Idx)
 		{
