ModelingTools: Added a new cvar (modeling.CapturePostEditChangeInTransactions, default=true). When set to false, PostEditChange will be excluded from tool-target based modeling tool transactions.

diff --git a/Engine/Plugins/Runtime/MeshModelingToolset/Source/ModelingComponents/Private/ModelingToolTargetUtil.cpp b/Engine/Plugins/Runtime/MeshModelingToolset/Source/ModelingComponents/Private/ModelingToolTargetUtil.cpp
--- a/Engine/Plugins/Runtime/MeshModelingToolset/Source/ModelingComponents/Private/ModelingToolTargetUtil.cpp
+++ b/Engine/Plugins/Runtime/MeshModelingToolset/Source/ModelingComponents/Private/ModelingToolTargetUtil.cpp
@@ -26,6 +26,10 @@
 #include "MeshDescriptionToDynamicMesh.h"
 #include "DynamicMeshToMeshDescription.h"
 #include "StaticMeshAttributes.h"
+#if WITH_EDITOR
+#include "StaticMeshCompiler.h"
+#include "SkinnedAssetCompiler.h"
+#endif
 
 #include "DynamicMesh/NonManifoldMappingSupport.h"
 
@@ -33,6 +37,14 @@
 
 using namespace UE::Geometry;
 
+namespace UE::ToolTarget::Internal
+{
+	static TAutoConsoleVariable<bool> CVarCapturePostEditChangeInTransactions(
+	TEXT("modeling.CapturePostEditChangeInTransactions"),
+	true,
+	TEXT("When true, PostEditChange will be included in tool-target based tool transactions."));
+}
+
 AActor* UE::ToolTarget::GetTargetActor(UToolTarget* Target)
 {
 	IPrimitiveComponentBackedTarget* TargetComponent = Cast<IPrimitiveComponentBackedTarget>(Target);
@@ -609,6 +621,33 @@ USkeletalMesh* UE::ToolTarget::GetSkeletalMeshFromTargetIfAvailable(UToolTarget*
 	return TargetSkeletalMesh;
 }
 
+#if WITH_EDITOR
+void UE::ToolTarget::Internal::PostEditChangeWithConditionalUndo(UObject* Object)
+{
+	if (CVarCapturePostEditChangeInTransactions->GetBool())
+	{
+		Object->PostEditChange();
+	}
+	else
+	{
+		TGuardValue<ITransaction*> SuppressTransaction(GUndo, nullptr);
+		Object->PostEditChange();
+
+		// For StaticMesh & SkeletalMesh, we additionally block on the StaticMesh/SkeletalMesh build
+		// to ensure it completes while transactions are disabled. This is necessary because
+		// closing out the transaction scope will force the compile to complete outside of this scope
+		// where transactions are still captured.
+		if (UStaticMesh* StaticMesh = Cast<UStaticMesh>(Object))
+		{
+			FStaticMeshCompilingManager::Get().FinishCompilation({StaticMesh});
+		}
+		else if (USkinnedAsset* SkinnedAsset = Cast<USkinnedAsset>(Object))
+		{
+			FSkinnedAssetCompilingManager::Get().FinishCompilation({SkinnedAsset});
+		}
+	}
+}
+#endif
 
 
 #undef LOCTEXT_NAMESPACE
\ No newline at end of file
diff --git a/Engine/Plugins/Runtime/MeshModelingToolset/Source/ModelingComponents/Public/ModelingToolTargetUtil.h b/Engine/Plugins/Runtime/MeshModelingToolset/Source/ModelingComponents/Public/ModelingToolTargetUtil.h
--- a/Engine/Plugins/Runtime/MeshModelingToolset/Source/ModelingComponents/Public/ModelingToolTargetUtil.h
+++ b/Engine/Plugins/Runtime/MeshModelingToolset/Source/ModelingComponents/Public/ModelingToolTargetUtil.h
@@ -253,6 +253,23 @@ namespace Internal
 	MODELINGCOMPONENTS_API void CommitDynamicMeshViaIPersistentDynamicMeshSource(
 		IPersistentDynamicMeshSource& DynamicMeshSource,
 		const UE::Geometry::FDynamicMesh3& UpdatedMesh, bool bHaveModifiedTopology);
+
+#if WITH_EDITOR
+	/**
+	 * Internal path for use by tool targets to invoke PostEditChange with conditional capture for undo
+	 * transactions. This capture is controlled by the CVar 'modeling.CapturePostEditChangeInTransactions'.
+	 *
+	 * PostEditChange can include arbitrary user code which can lead to a tool transaction capturing data
+	 * beyond the modifications the tool has made, opening up the risk of various side effects. It is
+	 * recommended that PostEditChange is not included in a transaction scope since PostEditUndo will
+	 * reinvoke PostEditChange and provide the user code the opportunity to reapply its changes in response
+	 * to the data modifications being made.
+	 *
+	 * That said, not all user code in PostEditChange is well behaved and deterministic so we provide the CVar
+	 * to toggle this behavior.
+	 */
+	MODELINGCOMPONENTS_API void PostEditChangeWithConditionalUndo(UObject* Object);
+#endif
 }
 
 /**
diff --git a/Engine/Plugins/Runtime/MeshModelingToolset/Source/ModelingComponentsEditorOnly/Private/ToolTargets/SkeletalMeshToolTarget.cpp b/Engine/Plugins/Runtime/MeshModelingToolset/Source/ModelingComponentsEditorOnly/Private/ToolTargets/SkeletalMeshToolTarget.cpp
--- a/Engine/Plugins/Runtime/MeshModelingToolset/Source/ModelingComponentsEditorOnly/Private/ToolTargets/SkeletalMeshToolTarget.cpp
+++ b/Engine/Plugins/Runtime/MeshModelingToolset/Source/ModelingComponentsEditorOnly/Private/ToolTargets/SkeletalMeshToolTarget.cpp
@@ -9,6 +9,7 @@
 #include "Rendering/SkeletalMeshModel.h"
 #include "Widgets/Notifications/SNotificationList.h"
 #include "Framework/Notifications/NotificationManager.h"
+#include "ModelingToolTargetUtil.h"
 #include "SkeletalMeshAttributes.h"
 #include "Misc/SlowTask.h"
 #include "UObject/Package.h"
@@ -151,7 +152,7 @@ bool USkeletalMeshReadOnlyToolTarget::CommitMaterialSetUpdate(USkeletalMesh* Ske
 		}
 	}
 
-	SkeletalMeshIn->PostEditChange();
+	UE::ToolTarget::Internal::PostEditChangeWithConditionalUndo(SkeletalMeshIn);
 
 	return true;
 }
@@ -276,7 +277,7 @@ void USkeletalMeshToolTarget::CommitMeshDescription(USkeletalMesh* SkeletalMesh,
 	FLODUtilities::UpdateLODInfoVertexAttributes(SkeletalMesh, LODIndex, LODIndex, false);
 	
 	// This call will invoke the skeletal mesh render data rebuild machinery.
-	SkeletalMesh->PostEditChange();
+	UE::ToolTarget::Internal::PostEditChangeWithConditionalUndo(SkeletalMesh);
 }
 
 void USkeletalMeshToolTarget::CommitDynamicMesh(const FDynamicMesh3& Mesh, const FDynamicMeshCommitInfo& CommitInfo)
diff --git a/Engine/Plugins/Runtime/MeshModelingToolset/Source/ModelingComponentsEditorOnly/Private/ToolTargets/StaticMeshToolTarget.cpp b/Engine/Plugins/Runtime/MeshModelingToolset/Source/ModelingComponentsEditorOnly/Private/ToolTargets/StaticMeshToolTarget.cpp
--- a/Engine/Plugins/Runtime/MeshModelingToolset/Source/ModelingComponentsEditorOnly/Private/ToolTargets/StaticMeshToolTarget.cpp
+++ b/Engine/Plugins/Runtime/MeshModelingToolset/Source/ModelingComponentsEditorOnly/Private/ToolTargets/StaticMeshToolTarget.cpp
@@ -9,6 +9,7 @@
 #include "Engine/StaticMesh.h"
 #include "MaterialDomain.h"
 #include "Materials/Material.h"
+#include "ModelingToolTargetUtil.h"
 #include "RenderingThread.h"
 #include "Widgets/Notifications/SNotificationList.h"
 #include "Framework/Notifications/NotificationManager.h"
@@ -255,8 +256,7 @@ bool UStaticMeshToolTarget::CommitMaterialSetUpdate(UStaticMesh* StaticMeshIn,
 		}
 	}
 
-	StaticMeshIn->PostEditChange();
-
+	UE::ToolTarget::Internal::PostEditChangeWithConditionalUndo(StaticMeshIn);
 	return true;
 }
 
@@ -482,7 +482,7 @@ void UStaticMeshToolTarget::CommitMeshDescription(UStaticMesh* StaticMeshIn, con
 		ThisSourceModel.ResetReductionSetting();
 	}
 
-	StaticMeshIn->PostEditChange();
+	UE::ToolTarget::Internal::PostEditChangeWithConditionalUndo(StaticMeshIn);
 }
 
 FDynamicMesh3 UStaticMeshToolTarget::GetDynamicMesh()
diff --git a/Engine/Plugins/Runtime/MeshModelingToolset/Source/ModelingComponentsEditorOnly/Private/ToolTargets/VolumeComponentToolTarget.cpp b/Engine/Plugins/Runtime/MeshModelingToolset/Source/ModelingComponentsEditorOnly/Private/ToolTargets/VolumeComponentToolTarget.cpp
--- a/Engine/Plugins/Runtime/MeshModelingToolset/Source/ModelingComponentsEditorOnly/Private/ToolTargets/VolumeComponentToolTarget.cpp
+++ b/Engine/Plugins/Runtime/MeshModelingToolset/Source/ModelingComponentsEditorOnly/Private/ToolTargets/VolumeComponentToolTarget.cpp
@@ -8,6 +8,7 @@
 #include "DynamicMeshToMeshDescription.h"
 #include "GameFramework/Volume.h"
 #include "MeshDescriptionToDynamicMesh.h"
+#include "ModelingToolTargetUtil.h"
 #include "DynamicMesh/MeshNormals.h"
 #include "StaticMeshAttributes.h"
 #include "ToolSetupUtil.h"
@@ -102,7 +103,8 @@ void UVolumeComponentToolTarget::CommitDynamicMesh(const UE::Geometry::FDynamicM
 	UE::Conversion::DynamicMeshToVolume(Mesh, Volume, ConversionOptions);
 
 	Volume->SetActorTransform(Transform);
-	Volume->PostEditChange();
+	
+	UE::ToolTarget::Internal::PostEditChangeWithConditionalUndo(Volume);
 }
 
 const FMeshDescription* UVolumeComponentToolTarget::GetMeshDescription(const FGetMeshParameters& GetMeshParams)
