ModelingTools: Fixed Inset tooltip error in PolyEdit/TriEdit.

diff --git a/Engine/Plugins/Runtime/MeshModelingToolset/Source/MeshModelingTools/Public/ToolActivities/PolyEditInsetOutsetActivity.h b/Engine/Plugins/Runtime/MeshModelingToolset/Source/MeshModelingTools/Public/ToolActivities/PolyEditInsetOutsetActivity.h
--- a/Engine/Plugins/Runtime/MeshModelingToolset/Source/MeshModelingTools/Public/ToolActivities/PolyEditInsetOutsetActivity.h
+++ b/Engine/Plugins/Runtime/MeshModelingToolset/Source/MeshModelingTools/Public/ToolActivities/PolyEditInsetOutsetActivity.h
@@ -19,12 +19,12 @@ class MESHMODELINGTOOLS_API UPolyEditInsetOutsetProperties : public UInteractive
 	GENERATED_BODY()
 
 public:
-	/** Amount of smoothing applied to outset boundary */
+	/** Amount of smoothing applied to the boundary */
 	UPROPERTY(EditAnywhere, Category = "Region", 
 		meta = (UIMin = "0.0", UIMax = "1.0", EditCondition = "bBoundaryOnly == false"))
 	float Softness = 0.5;
 
-	/** Controls whether outset operation will move interior vertices as well as border vertices */
+	/** Controls whether operation will move interior vertices as well as border vertices */
 	UPROPERTY(EditAnywhere, Category = "Region", AdvancedDisplay)
 	bool bBoundaryOnly = false;
 
