Clarify when the Rotation setting is used in the Add Primitive tools by disabling it when it is not applicable, renaming it to Initial Rotation, and expanding its tooltip.

diff --git a/Engine/Plugins/Runtime/MeshModelingToolset/Source/MeshModelingTools/Public/AddPrimitiveTool.h b/Engine/Plugins/Runtime/MeshModelingToolset/Source/MeshModelingTools/Public/AddPrimitiveTool.h
--- a/Engine/Plugins/Runtime/MeshModelingToolset/Source/MeshModelingTools/Public/AddPrimitiveTool.h
+++ b/Engine/Plugins/Runtime/MeshModelingToolset/Source/MeshModelingTools/Public/AddPrimitiveTool.h
@@ -95,8 +95,8 @@ public:
 	UPROPERTY(EditAnywhere, Category = Positioning, meta = (ProceduralShapeSetting))
 	EMakeMeshPivotLocation PivotLocation = EMakeMeshPivotLocation::Base;
 
-	/** Rotation of the shape around its up axis */
-	UPROPERTY(EditAnywhere, Category = Positioning, meta = (UIMin = "0.0", UIMax = "360.0"))
+	/** Initial rotation of the shape around its up axis, before placement. After placement, use the gizmo to control rotation. */
+	UPROPERTY(EditAnywhere, Category = Positioning, DisplayName = "Initial Rotation", meta = (UIMin = "0.0", UIMax = "360.0", EditCondition = "!bShowGizmoOptions", HideEditConditionToggle))
 	float Rotation = 0.0;
 
 	/** If true, aligns the shape along the normal of the surface it is placed on. */
