Modeling Tools: Replaces the existing PDI implementation of MeshElementSelection to use PreviewGeometry

diff --git a/Engine/Plugins/Runtime/MeshModelingToolset/Source/ModelingComponents/Private/Selection/GeometrySelectionManager.cpp b/Engine/Plugins/Runtime/MeshModelingToolset/Source/ModelingComponents/Private/Selection/GeometrySelectionManager.cpp
--- a/Engine/Plugins/Runtime/MeshModelingToolset/Source/ModelingComponents/Private/Selection/GeometrySelectionManager.cpp
+++ b/Engine/Plugins/Runtime/MeshModelingToolset/Source/ModelingComponents/Private/Selection/GeometrySelectionManager.cpp
@@ -10,6 +10,9 @@
 #include "InteractiveToolManager.h"
 #include "ToolContextInterfaces.h"
 #include "ToolDataVisualizer.h"
+#include "ToolSetupUtil.h"
+#include "Drawing/PreviewGeometryActor.h"
+#include "Materials/MaterialInstanceDynamic.h"
 #include "Selections/GeometrySelectionUtil.h"
 #include "Util/ColorConstants.h"
 
@@ -33,6 +36,7 @@ void UGeometrySelectionManager::Initialize( UInteractiveToolsContext* ToolsConte
 {
 	ToolsContext = ToolsContextIn;
 	TransactionsAPI = TransactionsAPIIn;
+	PreviewGeometry = NewObject<UPreviewGeometry>(this);
 }
 
 void UGeometrySelectionManager::RegisterSelectorFactory(TUniquePtr<IGeometrySelectorFactory> Factory)
@@ -48,6 +52,13 @@ void UGeometrySelectionManager::Shutdown()
 	ToolsContext = nullptr;
 	TransactionsAPI = nullptr;
 
+	if (PreviewGeometry)
+	{
+		RemoveAllSets();
+		PreviewGeometry->Disconnect();
+		PreviewGeometry = nullptr;
+	}
+
 	for (TSharedPtr<FGeometrySelectionTarget> Target : ActiveTargetReferences)
 	{
 		SleepOrShutdownTarget(Target, false);
@@ -125,7 +136,7 @@ void UGeometrySelectionManager::SetSelectionElementTypeInternal(EGeometryElement
 			Target->SelectionEditor->UpdateQueryConfig(GetCurrentSelectionQueryConfig(), bEnableTopologyFilter);
 		}
 
-		MarkRenderCachesDirty();
+		MarkRenderCachesDirty(false);
 		ClearActivePreview();
 	}
 }
@@ -224,6 +235,12 @@ void UGeometrySelectionManager::SetMeshSelectionTypeAndMode(EGeometryElementType
 
 		ClearSelection();
 
+		// clear preview geometry sets when in Object selection mode
+		if (NewTopologyMode == EMeshTopologyMode::None && PreviewGeometry)
+		{
+			RemoveAllSets();
+		}
+
 		GetTransactionsAPI()->BeginUndoTransaction(LOCTEXT("ChangeElementMethod", "Change Selection Method"));
 
 		// We have to undo/redo the change to the selection type because if we want to 'undo' this later and restore
@@ -715,7 +732,7 @@ void UGeometrySelectionManager::ClearSelection(bool bSaveSelectionBeforeClear)
 
 	GetTransactionsAPI()->EndUndoTransaction();
 
-	MarkRenderCachesDirty();
+	MarkRenderCachesDirty(false);
 	OnSelectionModified.Broadcast();
 }
 
@@ -755,7 +772,7 @@ void UGeometrySelectionManager::UpdateSelectionViaRaycast(
 		GetTransactionsAPI()->AppendChange(this, MoveTemp(DeltaChange), LOCTEXT("UpdateSelectionViaRaycast", "Change Selection"));
 		GetTransactionsAPI()->EndUndoTransaction();
 
-		MarkRenderCachesDirty();
+		MarkRenderCachesDirty(false);
 		OnSelectionModified.Broadcast();
 	}
 	else if (ResultOut.bSelectionMissed && UpdateConfig.ChangeType == EGeometrySelectionChangeType::Replace)
@@ -798,7 +815,7 @@ void UGeometrySelectionManager::UpdateSelectionViaConvex(
 		GetTransactionsAPI()->AppendChange(this, MoveTemp(DeltaChange), LOCTEXT("UpdateSelectionViaConvex", "Change Selection"));
 		GetTransactionsAPI()->EndUndoTransaction();
 
-		MarkRenderCachesDirty();
+		MarkRenderCachesDirty(false);
 		OnSelectionModified.Broadcast();
 	}
 	else if (ResultOut.bSelectionMissed && UpdateConfig.ChangeType == EGeometrySelectionChangeType::Replace)
@@ -843,7 +860,7 @@ bool UGeometrySelectionManager::BeginTrackedSelectionChange(FGeometrySelectionUp
 
 	if (bClearOnBegin && InitialTrackedDelta.IsEmpty() == false)
 	{
-		MarkRenderCachesDirty();
+		MarkRenderCachesDirty(false);
 		OnSelectionModified.Broadcast();
 	}
 
@@ -873,7 +890,7 @@ void UGeometrySelectionManager::AccumulateSelectionUpdate_Raycast(
 		ActiveTrackedDelta.Added.Append( ResultOut.SelectionDelta.Added );
 		ActiveTrackedDelta.Removed.Append( ResultOut.SelectionDelta.Removed );
 
-		MarkRenderCachesDirty();
+		MarkRenderCachesDirty(false);
 		OnSelectionModified.Broadcast();
 	}
 }
@@ -928,7 +945,7 @@ bool UGeometrySelectionManager::SetSelectionForComponent(UPrimitiveComponent* Co
 				NewSelectionChange->Before = InitialSelection;
 				GetTransactionsAPI()->AppendChange(this, MoveTemp(NewSelectionChange), LOCTEXT("NewSelection", "New Selection"));
 
-				MarkRenderCachesDirty();
+				MarkRenderCachesDirty(false);
 				OnSelectionModified.Broadcast();
 			}
 			return true;
@@ -988,7 +1005,7 @@ bool UGeometrySelectionManager::RestoreSavedSelection()
 					NewSelectionChange->Before = InitialSelection;
 					GetTransactionsAPI()->AppendChange(this, MoveTemp(NewSelectionChange), LOCTEXT("RestoreSelection", "Restore Selection"));
 
-					MarkRenderCachesDirty();
+					MarkRenderCachesDirty(false);
 					OnSelectionModified.Broadcast();
 				}
 				bFound = true;
@@ -1090,7 +1107,11 @@ bool UGeometrySelectionManager::UpdateSelectionPreviewViaRaycast(
 		}
 		CachedSelectedPreviewRenderElements.Reset();
 		CachedUnselectedPreviewRenderElements.Reset();
-		MarkRenderCachesDirty();
+
+		RenderCachesDirtyFlags |= EEnumerateRenderCachesDirtyFlags::PreviewCachesDirty;
+
+		RemoveSets(HoverOverSelectedParams.Identifiers);
+		RemoveSets(HoverOverUnselectedParams.Identifiers);
 	}
 
 	return (ActivePreviewSelection.IsEmpty() == false);
@@ -1220,7 +1241,7 @@ void UGeometrySelectionManager::UpdateTransformation(
 		}
 	}
 
-	bSelectionRenderCachesDirty = true;
+	RenderCachesDirtyFlags |= EEnumerateRenderCachesDirtyFlags::SelectionCachesDirty;
 }
 
 void UGeometrySelectionManager::EndTransformation()
@@ -1242,7 +1263,7 @@ void UGeometrySelectionManager::EndTransformation()
 
 	GetTransactionsAPI()->EndUndoTransaction();
 
-	bSelectionRenderCachesDirty = true;
+	RenderCachesDirtyFlags |= EEnumerateRenderCachesDirtyFlags::SelectionCachesDirty;
 }
 
 
@@ -1333,8 +1354,8 @@ void UGeometrySelectionManager::ExecuteSelectionCommand(UGeometrySelectionEditCo
 
 	GetTransactionsAPI()->EndUndoTransaction();
 
-	// assume this is true for now
-	MarkRenderCachesDirty();
+	// assume marking render caches (except selectable) is true for now
+	MarkRenderCachesDirty(false);
 	OnSelectionModified.Broadcast();
 }
 
@@ -1381,7 +1402,7 @@ void UGeometrySelectionManager::ApplyChange(IGeometrySelectionChange* Change)
 
 			if (ApplyDelta.IsEmpty() == false)
 			{
-				MarkRenderCachesDirty();
+				MarkRenderCachesDirty(false);
 				OnSelectionModified.Broadcast();
 			}
 
@@ -1410,7 +1431,7 @@ void UGeometrySelectionManager::RevertChange(IGeometrySelectionChange* Change)
 
 			if (RevertDelta.IsEmpty() == false)
 			{
-				MarkRenderCachesDirty();
+				MarkRenderCachesDirty(false);
 				OnSelectionModified.Broadcast();
 			}
 
@@ -1425,6 +1446,8 @@ void UGeometrySelectionManager::OnTargetGeometryModified(IGeometrySelector* Sele
 	CachedSelectableRenderElements.Reset();
 	
 	CachedSelectableRenderElements.SetNum(ActiveTargetReferences.Num());
+
+	RemoveSets(UnselectedParams.Identifiers);
 	
 	MarkRenderCachesDirty();
 	ClearActivePreview();
@@ -1438,16 +1461,21 @@ void UGeometrySelectionManager::UpdateSelectionRenderCacheOnTargetChange()
 	
 	CachedSelectionRenderElements.SetNum(ActiveTargetReferences.Num());
 	CachedSelectableRenderElements.SetNum(ActiveTargetReferences.Num());
+
+	RemoveSets(SelectedParams.Identifiers);
+	RemoveSets(UnselectedParams.Identifiers);
 	
 	MarkRenderCachesDirty();
 	ClearActivePreview();
 }
 
-void UGeometrySelectionManager::MarkRenderCachesDirty()
+void UGeometrySelectionManager::MarkRenderCachesDirty(bool bMarkSelectableDirty)
 {
-	bSelectionRenderCachesDirty = true;
-	bSelectableRenderCachesDirty = true;
-	bPreviewRenderCachesDirty = true;
+	if (bMarkSelectableDirty)
+	{
+		RenderCachesDirtyFlags |= EEnumerateRenderCachesDirtyFlags::UnselectedCachesDirty;
+	}
+	RenderCachesDirtyFlags |= EEnumerateRenderCachesDirtyFlags::PreviewCachesDirty | EEnumerateRenderCachesDirtyFlags::SelectionCachesDirty;
 }
 
 void UGeometrySelectionManager::RebuildSelectionRenderCaches()
@@ -1464,17 +1492,74 @@ void UGeometrySelectionManager::ClearActivePreview()
 	UnselectedActivePreviewSelection.Reset();
 	CachedSelectedPreviewRenderElements.Reset();
 	CachedUnselectedPreviewRenderElements.Reset();
-	bSelectableRenderCachesDirty = true;
-	bPreviewRenderCachesDirty = true;
+	
+	RenderCachesDirtyFlags |= EEnumerateRenderCachesDirtyFlags::PreviewCachesDirty;
+
+	RemoveSets(HoverOverSelectedParams.Identifiers);
+	RemoveSets(HoverOverUnselectedParams.Identifiers);
 }
 
-void UGeometrySelectionManager::RebuildSelectionRenderCache()
+
+void UGeometrySelectionManager::RemoveAllSets() const
 {
-	if (bSelectionRenderCachesDirty == false)
+	if(!ensure(PreviewGeometry))
 	{
 		return;
 	}
+	PreviewGeometry->RemoveAllLineSets();
+	PreviewGeometry->RemoveAllPointSets();
+	PreviewGeometry->RemoveAllTriangleSets();
+}
+
+void UGeometrySelectionManager::RemoveSets(const TArrayView<FString>& SetIdentifiers) const
+{
+	
+	if (PreviewGeometry)
+	{
+		PreviewGeometry->RemoveLineSet(SetIdentifiers[1]);
+		PreviewGeometry->RemovePointSet(SetIdentifiers[0]);
+		PreviewGeometry->RemoveTriangleSet(SetIdentifiers[2]);
+	}
+}
+
+void UGeometrySelectionManager::CreateOrUpdateAllSets(const FGeometrySelectionElements& Elements, const FMeshElementSelectionParams SelectionParams) const
+{
+
+	UWorld* World = GetWorld();
+	if (World && !PreviewGeometry->ParentActor)
+	{
+		PreviewGeometry->CreateInWorld(World, FTransform::Identity);
+	}
+	
+	PreviewGeometry->CreateOrUpdateLineSet(SelectionParams.Identifiers[1], Elements.Segments.Num(), [&](int32 j, TArray<FRenderableLine>& LinesOut)
+		{
+		const FSegment3d Seg = Elements.Segments[j];
+		LinesOut.Add(FRenderableLine(Seg.StartPoint(), Seg.EndPoint(),  SelectionParams.Color,  SelectionParams.LineThickness, SelectionParams.DepthBias));
+		}, 1);
+	
+	PreviewGeometry->CreateOrUpdateTriangleSet(SelectionParams.Identifiers[2], Elements.Triangles.Num(), [&](int32 k, TArray<FRenderableTriangle>& TrianglesOut)
+		{
+		const FTriangle3d Triangle = Elements.Triangles[k];
+		const FVector3d Normal = Triangle.Normal();
+		FRenderableTriangleVertex A(Triangle.V[0], FVector2D(0,0), Normal, SelectionParams.Color);
+		FRenderableTriangleVertex B(Triangle.V[1], FVector2D(1,0), Normal, SelectionParams.Color);
+		FRenderableTriangleVertex C(Triangle.V[2], FVector2D(1,1), Normal, SelectionParams.Color);
+		TrianglesOut.Add(FRenderableTriangle(SelectionParams.SelectionFillColor, A, B, C));
+		},1 );
+	
+	PreviewGeometry->CreateOrUpdatePointSet(SelectionParams.Identifiers[0], Elements.Points.Num(), [&](int32 k, TArray<FRenderablePoint>& PointsOut)
+		{
+		const FVector3d Point = Elements.Points[k];
+		PointsOut.Add(FRenderablePoint(Point, SelectionParams.Color, SelectionParams.PointSize, SelectionParams.DepthBias));
+		});
+}
 
+void UGeometrySelectionManager::RebuildSelectionRenderCache()
+{
+	if ((RenderCachesDirtyFlags & EEnumerateRenderCachesDirtyFlags::SelectionCachesDirty) == EEnumerateRenderCachesDirtyFlags::None)
+	{
+		return;
+	}
 	check(ActiveTargetReferences.Num() == CachedSelectionRenderElements.Num());
 	for (int32 k = 0; k < ActiveTargetReferences.Num(); ++k)
 	{
@@ -1483,14 +1568,15 @@ void UGeometrySelectionManager::RebuildSelectionRenderCache()
 		FGeometrySelectionElements& SelectionElements = CachedSelectionRenderElements[k];
 		SelectionElements.Reset();
 		Target->Selector->AccumulateSelectionElements(Target->Selection, SelectionElements, true, EEnumerateSelectionMapping::Default);
+		CreateOrUpdateAllSets(SelectionElements, SelectedParams);
 	}
 
-	bSelectionRenderCachesDirty = false;
+	RenderCachesDirtyFlags &= ~EEnumerateRenderCachesDirtyFlags::SelectionCachesDirty;
 }
 
 void UGeometrySelectionManager::RebuildSelectableRenderCache()
 {
-	if (bSelectableRenderCachesDirty == false || MeshTopologyMode == EMeshTopologyMode::None)
+	if ((RenderCachesDirtyFlags & EEnumerateRenderCachesDirtyFlags::UnselectedCachesDirty) == EEnumerateRenderCachesDirtyFlags::None || MeshTopologyMode == EMeshTopologyMode::None)
 	{
 		return;
 	}
@@ -1513,14 +1599,15 @@ void UGeometrySelectionManager::RebuildSelectableRenderCache()
 			
 			return true;
 		});
+		CreateOrUpdateAllSets(AllElements, UnselectedParams);
 	}
 
-	bSelectableRenderCachesDirty = false;
+	RenderCachesDirtyFlags &= ~EEnumerateRenderCachesDirtyFlags::UnselectedCachesDirty;
 }
 
 void UGeometrySelectionManager::RebuildPreviewRenderCache()
 {
-	if (bPreviewRenderCachesDirty == false || ActiveTargetReferences.Num() == 0)
+	if ((RenderCachesDirtyFlags & EEnumerateRenderCachesDirtyFlags::PreviewCachesDirty) == EEnumerateRenderCachesDirtyFlags::None || ActiveTargetReferences.Num() == 0)
 	{
 		return;
 	}
@@ -1539,9 +1626,10 @@ void UGeometrySelectionManager::RebuildPreviewRenderCache()
 		
 		Target->Selector->AccumulateSelectionElements(SelectedActivePreviewSelection, CachedSelectedPreviewRenderElements, true, MappingFlags);
 		Target->Selector->AccumulateSelectionElements(UnselectedActivePreviewSelection, CachedUnselectedPreviewRenderElements, true, MappingFlags);
+		CreateOrUpdateAllSets(CachedSelectedPreviewRenderElements, HoverOverSelectedParams);
+		CreateOrUpdateAllSets(CachedUnselectedPreviewRenderElements, HoverOverUnselectedParams);
 	}
-	
-	bPreviewRenderCachesDirty = false;
+	RenderCachesDirtyFlags &= ~EEnumerateRenderCachesDirtyFlags::PreviewCachesDirty;
 }
 
 
@@ -1577,39 +1665,42 @@ void UGeometrySelectionManager::DebugRender(IToolsContextRenderAPI* RenderAPI)
 
 		return;
 	}
-
-	//const UMaterialInterface* TriangleMaterial = ToolSetupUtil::GetSelectionMaterial(FLinearColor(1.0f, 0, 0, 0.5f), nullptr, 0.5f);
-
 	if (!this->ToolsContext->ToolManager->HasAnyActiveTool())
 	{
 		RebuildSelectionRenderCaches();
-		
-		for ( const FGeometrySelectionElements& Elements : CachedSelectionRenderElements )
-		{
-			// render geometry selection
-			ToolSelectionUtil::DebugRender(RenderAPI, Elements, 4.f, GeometrySelectedColor, 10.f, GeometrySelectedColor, 6.f, GeometrySelectedColor);
-		}
-
-		if (MeshTopologyMode != EMeshTopologyMode::None)
-		{
-			for ( const FGeometrySelectionElements& Elements : CachedSelectableRenderElements )
-			{
-				// render unselected geometry (wireframe overlay color in MeshElementSelection mode)
-				ToolSelectionUtil::DebugRender(RenderAPI, Elements, 2.f, UnselectedColor, 8.f, UnselectedColor, 5.f, UnselectedColor);
-			}
-		}
-		// render both hover cases (if element is part of the selection, or if not)
-		ToolSelectionUtil::DebugRender(RenderAPI, CachedSelectedPreviewRenderElements, 4.f, HoverOverSelectedColor, 10.f, HoverOverSelectedColor, 7.f, HoverOverSelectedColor);
-		ToolSelectionUtil::DebugRender(RenderAPI, CachedUnselectedPreviewRenderElements, 4.f, HoverOverUnselectedColor, 10.f, HoverOverUnselectedColor, 7.f, HoverOverUnselectedColor);
+	}
+	else
+	{
+		// disables PreviewGeometry when in a tool
+		RemoveAllSets();
+		RenderCachesDirtyFlags |= EEnumerateRenderCachesDirtyFlags::UnselectedCachesDirty;
 	}
 }
 
 void UGeometrySelectionManager::SetSelectionColors(const FLinearColor UnselectedCol, const FLinearColor HoverOverSelectedCol, const FLinearColor HoverOverUnselectedCol, const FLinearColor GeometrySelectedCol)
 {
-	UnselectedColor = UnselectedCol;
-	HoverOverSelectedColor = HoverOverSelectedCol;
-	HoverOverUnselectedColor = HoverOverUnselectedCol;
-	GeometrySelectedColor = GeometrySelectedCol;
+	UnselectedParams.Color = UnselectedCol.ToFColor(true);
+	HoverOverSelectedParams.Color = HoverOverSelectedCol.ToFColor(true);
+	HoverOverUnselectedParams.Color = HoverOverUnselectedCol.ToFColor(true);
+	SelectedParams.Color = GeometrySelectedCol.ToFColor(true);
+
+	const float Opacity = .6f;
+
+	UnselectedParams.SelectionFillColor =
+		ToolSetupUtil::GetCustomTwoSidedDepthOffsetMaterial(this->ToolsContext->ToolManager, UnselectedCol, UnselectedParams.DepthBias, Opacity);
+
+	HoverOverSelectedParams.SelectionFillColor =
+		ToolSetupUtil::GetCustomTwoSidedDepthOffsetMaterial(this->ToolsContext->ToolManager, HoverOverSelectedCol, HoverOverSelectedParams.DepthBias, Opacity);
+
+	HoverOverUnselectedParams.SelectionFillColor =
+		ToolSetupUtil::GetCustomTwoSidedDepthOffsetMaterial(this->ToolsContext->ToolManager, HoverOverUnselectedCol, HoverOverUnselectedParams.DepthBias, Opacity);
+
+	SelectedParams.SelectionFillColor =
+		ToolSetupUtil::GetCustomTwoSidedDepthOffsetMaterial(this->ToolsContext->ToolManager, GeometrySelectedCol, SelectedParams.DepthBias, Opacity);
+
+	// ensures that when color is changed in Editor Preferences, colors are immediately updated in the UI
+	MarkRenderCachesDirty();
+	RebuildSelectionRenderCaches();
 }
 
 
diff --git a/Engine/Plugins/Runtime/MeshModelingToolset/Source/ModelingComponents/Public/Selection/GeometrySelectionManager.h b/Engine/Plugins/Runtime/MeshModelingToolset/Source/ModelingComponents/Public/Selection/GeometrySelectionManager.h
--- a/Engine/Plugins/Runtime/MeshModelingToolset/Source/ModelingComponents/Public/Selection/GeometrySelectionManager.h
+++ b/Engine/Plugins/Runtime/MeshModelingToolset/Source/ModelingComponents/Public/Selection/GeometrySelectionManager.h
@@ -10,6 +10,7 @@
 
 #include "GeometrySelectionManager.generated.h"
 
+class UPreviewGeometry;
 class IGeometrySelector;
 class UInteractiveToolsContext;
 class IToolsContextRenderAPI;
@@ -17,6 +18,30 @@ class IToolsContextTransactionsAPI;
 class UGeometrySelectionEditCommand;
 class UGeometrySelectionEditCommandArguments;
 
+USTRUCT()
+struct FMeshElementSelectionParams
+{
+	GENERATED_BODY()
+	TStaticArray<FString, 3> Identifiers {};
+	float DepthBias = 0.0f;
+	float LineThickness = 0.0f;
+	float PointSize = 0.0f;
+	FColor Color = FColor(0,0,0);
+	UPROPERTY()
+	TObjectPtr<UMaterialInstanceDynamic> SelectionFillColor = nullptr;
+};
+
+enum class EEnumerateRenderCachesDirtyFlags: uint8
+{
+	None = 0,
+
+	SelectionCachesDirty = 1 << 0,
+	UnselectedCachesDirty = 1 << 1, //SelectableRenderCaches
+	PreviewCachesDirty = 1 << 2,
+
+	Default = SelectionCachesDirty | UnselectedCachesDirty | PreviewCachesDirty,
+};
+ENUM_CLASS_FLAGS(EEnumerateRenderCachesDirtyFlags);
 
 /**
  * UGeometrySelectionManager provides the infrastructure for "Element Selection", ie 
@@ -103,6 +128,11 @@ public:
 	// Switch the selection mode and type, optionally converting any existing selection to the new type and mode
 	virtual void SetMeshSelectionTypeAndMode(EGeometryElementType NewElementType, EMeshTopologyMode NewSelectionMode, bool bConvertSelection);
 
+	/**
+	 * Removes Triangle, Line, and Point sets with the given SetIdentifier prefix
+	 * Full Identifier strings are in the format: SetIdentifier_Triangles, SetIdentifier_Lines, and SetIdentifier_Points
+	 */
+	void RemoveSets(const TArrayView<FString>& SetIdentifiers) const;
 	
 
 
@@ -156,9 +186,10 @@ public:
 	);
 
 	/**
-	 * Invalidates all cached selection elements.
+	 * Invalidates all cached selection elements by default.
+	 * When desired, can choose to not mark the Selectable render cache as dirty
 	 */
-	void MarkRenderCachesDirty();
+	void MarkRenderCachesDirty(bool bMarkSelectableDirty = true);
 	
 	//
 	// Selection Updates
@@ -468,9 +499,18 @@ protected:
 
 	
 	// todo [nickolas.drake]: cane we move CachedSelectionRenderElements, Cached[Un]selectedPreviewRenderElements, and bSelectionRenderCachesDirty to private?
+
+	UPROPERTY()
+	TObjectPtr<UPreviewGeometry> PreviewGeometry = nullptr;
+
+	/**
+	 * Calls the CreateOrUpdate function for Triangle, Line, and Point sets to build sets for PreviewGeometry
+	 * @param SelectionParams struct containing information (color, depth bias, line thickness, identifier) for the type of selection we are currently
+	 * building a set for (as Selected, Selectable, HoverSelected, HoverUnselected all have different visual information needed to contruct their sets)
+	 */
+	void CreateOrUpdateAllSets(const FGeometrySelectionElements& Elements, const FMeshElementSelectionParams SelectionParams) const;
 	
 	TArray<FGeometrySelectionElements> CachedSelectionRenderElements;				// Cached 3D geometry for current selection
-	bool bSelectionRenderCachesDirty = false;
 	void UpdateSelectionRenderCacheOnTargetChange();
 	void RebuildSelectionRenderCaches();
 
@@ -482,12 +522,20 @@ protected:
 	FGeometrySelectionElements CachedUnselectedPreviewRenderElements;  // Cached 3D geometry for active preview elements that are NOT in the current selection
     void ClearActivePreview();
 
+	inline static TStaticArray<FString, 3> UnselectedSetIds = { "Unselected_Vertices", "Unselected_Lines", "Unselected_Triangles"};
+	inline static TStaticArray<FString, 3> HoverOverSelectedSetIds = {  "HoverSelected_Vertices", "HoverSelected_Lines", "HoverSelected_Triangles"};
+	inline static TStaticArray<FString, 3> HoverOverUnselectedSetIds = { "HoverUnselected_Vertices", "HoverUnselected_Lines", "HoverUnselected_Triangles" };
+	inline static TStaticArray<FString, 3> SelectedSetIds = {"Selected_Vertices", "Selected_Lines", "Selected_Triangles"};
+
+	UPROPERTY()
+	FMeshElementSelectionParams UnselectedParams {UnselectedSetIds, 5.f, 4.f, 8.f};
+	UPROPERTY()
+	FMeshElementSelectionParams HoverOverSelectedParams {HoverOverSelectedSetIds, 10.f, 8.f, 10.f};
+	UPROPERTY()
+	FMeshElementSelectionParams HoverOverUnselectedParams {HoverOverUnselectedSetIds, 10.f, 8.f, 10.f};
+	UPROPERTY()
+	FMeshElementSelectionParams SelectedParams {SelectedSetIds, 6.f, 6.f, 10.f};
 	
-	// colors used in mesh element selection
-	FLinearColor UnselectedColor;
-	FLinearColor HoverOverSelectedColor;
-	FLinearColor HoverOverUnselectedColor;
-	FLinearColor GeometrySelectedColor;
 
 	
 	// various change types need internal access
@@ -513,19 +561,18 @@ private:
 	// 3D geometry for element selections of each ActiveTarget is cached
 	// to improve rendering performance
 	//
-	
-	// todo [nickolas.drake]: refactor Rebuild_X_RenderCache() functions below to populate a UPreviewGeometry with the accumulated elements and remove PDI rendering
-	
+
 	void RebuildSelectionRenderCache();
 	
 	TArray<FGeometrySelectionElements> CachedSelectableRenderElements;				// Cached 3D geometry for all selectable elements
 	void RebuildSelectableRenderCache();
-	bool bSelectableRenderCachesDirty = false;
 	
 	void RebuildPreviewRenderCache();
-	bool bPreviewRenderCachesDirty = false;
 
-	
+	EEnumerateRenderCachesDirtyFlags RenderCachesDirtyFlags = EEnumerateRenderCachesDirtyFlags::None;
+
+	void RemoveAllSets() const;
+
 	// Tracks saved selection state. Useful when the selection is temporarily cleared (e.g., for a tool)
 	struct FSavedSelection
 	{
