DynamicMeshProvider: full support for MeshDescription LODs when converting to DynamicMesh
    a FGetMeshParameters parameter is now used to request a specific LOD, as well as the option of requesting tangents
    this will also allow further customization in the future, if required, by adding attributes to the FGetMeshParameters struct

#rnx

diff --git a/Engine/Plugins/ChaosClothAssetEditor/Source/ChaosClothAssetEditorTools/Private/ChaosClothAsset/ClothTransferSkinWeightsTool.cpp b/Engine/Plugins/ChaosClothAssetEditor/Source/ChaosClothAssetEditorTools/Private/ChaosClothAsset/ClothTransferSkinWeightsTool.cpp
--- a/Engine/Plugins/ChaosClothAssetEditor/Source/ChaosClothAssetEditorTools/Private/ChaosClothAsset/ClothTransferSkinWeightsTool.cpp
+++ b/Engine/Plugins/ChaosClothAssetEditor/Source/ChaosClothAssetEditorTools/Private/ChaosClothAsset/ClothTransferSkinWeightsTool.cpp
@@ -78,7 +78,9 @@ void UClothTransferSkinWeightsTool::Setup()
 	TargetClothPreview->OnMeshUpdated.AddUObject(this, &UClothTransferSkinWeightsTool::PreviewMeshUpdatedCallback);
 
 	// Set the initial preview mesh before any computation runs
-	UE::Geometry::FDynamicMesh3 InitialPreviewMesh = UE::ToolTarget::GetDynamicMeshCopy(Target, true);
+	static FGetMeshParameters GetMeshParams;
+	GetMeshParams.bWantMeshTangents = true;
+	UE::Geometry::FDynamicMesh3 InitialPreviewMesh = UE::ToolTarget::GetDynamicMeshCopy(Target, GetMeshParams);
 	TargetClothPreview->PreviewMesh->UpdatePreview(MoveTemp(InitialPreviewMesh));
 
 	TargetClothPreview->SetVisibility(true);
diff --git a/Engine/Plugins/Experimental/Dataflow/Source/DataflowEditor/Private/Dataflow/DataflowComponentToolTarget.cpp b/Engine/Plugins/Experimental/Dataflow/Source/DataflowEditor/Private/Dataflow/DataflowComponentToolTarget.cpp
--- a/Engine/Plugins/Experimental/Dataflow/Source/DataflowEditor/Private/Dataflow/DataflowComponentToolTarget.cpp
+++ b/Engine/Plugins/Experimental/Dataflow/Source/DataflowEditor/Private/Dataflow/DataflowComponentToolTarget.cpp
@@ -182,11 +182,6 @@ FDynamicMesh3 UDataflowComponentReadOnlyToolTarget::GetDynamicMesh()
 	return DynamicMesh;
 }
 
-FDynamicMesh3 UDataflowComponentReadOnlyToolTarget::GetDynamicMesh(bool bRequestTangents)
-{
-	return GetDynamicMesh();
-}
-
 //
 // UDataflowComponentToolTarget
 //
diff --git a/Engine/Plugins/Experimental/Dataflow/Source/DataflowEditor/Private/Dataflow/DataflowToolTarget.cpp b/Engine/Plugins/Experimental/Dataflow/Source/DataflowEditor/Private/Dataflow/DataflowToolTarget.cpp
--- a/Engine/Plugins/Experimental/Dataflow/Source/DataflowEditor/Private/Dataflow/DataflowToolTarget.cpp
+++ b/Engine/Plugins/Experimental/Dataflow/Source/DataflowEditor/Private/Dataflow/DataflowToolTarget.cpp
@@ -110,11 +110,6 @@ FDynamicMesh3 UDataflowReadOnlyToolTarget::GetDynamicMesh()
 	return DynamicMesh;
 }
 
-FDynamicMesh3 UDataflowReadOnlyToolTarget::GetDynamicMesh(bool bRequestTangents)
-{
-	return GetDynamicMesh();
-}
-
 //
 // UDataflowToolTarget
 //
diff --git a/Engine/Plugins/Experimental/Dataflow/Source/DataflowEditor/Public/Dataflow/DataflowComponentToolTarget.h b/Engine/Plugins/Experimental/Dataflow/Source/DataflowEditor/Public/Dataflow/DataflowComponentToolTarget.h
--- a/Engine/Plugins/Experimental/Dataflow/Source/DataflowEditor/Public/Dataflow/DataflowComponentToolTarget.h
+++ b/Engine/Plugins/Experimental/Dataflow/Source/DataflowEditor/Public/Dataflow/DataflowComponentToolTarget.h
@@ -43,7 +43,6 @@ public:
 
 	// IDynamicMeshProvider
 	virtual UE::Geometry::FDynamicMesh3 GetDynamicMesh() override;
-	virtual UE::Geometry::FDynamicMesh3 GetDynamicMesh(bool bRequestTangents) override;
 
 protected:
 	// So that the tool target factory can poke into Component.
diff --git a/Engine/Plugins/Experimental/Dataflow/Source/DataflowEditor/Public/Dataflow/DataflowToolTarget.h b/Engine/Plugins/Experimental/Dataflow/Source/DataflowEditor/Public/Dataflow/DataflowToolTarget.h
--- a/Engine/Plugins/Experimental/Dataflow/Source/DataflowEditor/Public/Dataflow/DataflowToolTarget.h
+++ b/Engine/Plugins/Experimental/Dataflow/Source/DataflowEditor/Public/Dataflow/DataflowToolTarget.h
@@ -45,7 +45,6 @@ public:
 
 	// IDynamicMeshProvider
 	virtual UE::Geometry::FDynamicMesh3 GetDynamicMesh() override;
-	virtual UE::Geometry::FDynamicMesh3 GetDynamicMesh(bool bRequestTangents) override;
 
 protected:
 	// So that the tool target factory can poke into Component.
diff --git a/Engine/Plugins/Experimental/MeshModelingToolsetExp/Source/MeshModelingToolsEditorOnlyExp/Private/MeshTangentsTool.cpp b/Engine/Plugins/Experimental/MeshModelingToolsetExp/Source/MeshModelingToolsEditorOnlyExp/Private/MeshTangentsTool.cpp
--- a/Engine/Plugins/Experimental/MeshModelingToolsetExp/Source/MeshModelingToolsEditorOnlyExp/Private/MeshTangentsTool.cpp
+++ b/Engine/Plugins/Experimental/MeshModelingToolsetExp/Source/MeshModelingToolsEditorOnlyExp/Private/MeshTangentsTool.cpp
@@ -106,7 +106,9 @@ void UMeshTangentsTool::Setup()
 	PreviewMesh->SetMaterials(MaterialSet.Materials);
 	// configure mesh
 	PreviewMesh->SetTangentsMode(EDynamicMeshComponentTangentsMode::ExternallyProvided);
-	FDynamicMesh3 InputMeshWithTangents = UE::ToolTarget::GetDynamicMeshCopy(Target, true);
+	static FGetMeshParameters GetMeshParams;
+	GetMeshParams.bWantMeshTangents = true;
+	FDynamicMesh3 InputMeshWithTangents = UE::ToolTarget::GetDynamicMeshCopy(Target, GetMeshParams);
 	PreviewMesh->ReplaceMesh(MoveTemp(InputMeshWithTangents));
 
 	// make a copy of initialized mesh and tangents
diff --git a/Engine/Plugins/Experimental/MeshModelingToolsetExp/Source/MeshModelingToolsEditorOnlyExp/Private/SimplifyMeshTool.cpp b/Engine/Plugins/Experimental/MeshModelingToolsetExp/Source/MeshModelingToolsEditorOnlyExp/Private/SimplifyMeshTool.cpp
--- a/Engine/Plugins/Experimental/MeshModelingToolsetExp/Source/MeshModelingToolsEditorOnlyExp/Private/SimplifyMeshTool.cpp
+++ b/Engine/Plugins/Experimental/MeshModelingToolsetExp/Source/MeshModelingToolsEditorOnlyExp/Private/SimplifyMeshTool.cpp
@@ -98,7 +98,9 @@ void USimplifyMeshTool::Setup()
 		// Note that we only copy over a mesh description if we use the UEStandard path (conditionally done in MakeNewOperator)
 		// to avoid doing conversions back and forth for non mesh-description-backed targets (where the conversions can 
 		// occasionally do odd things to the attributes), and to avoid the slow copy unless the user needs it.
-		OriginalMesh = MakeShared<FDynamicMesh3, ESPMode::ThreadSafe>(UE::ToolTarget::GetDynamicMeshCopy(Target, true));
+		static FGetMeshParameters GetMeshParams;
+		GetMeshParams.bWantMeshTangents = true;
+		OriginalMesh = MakeShared<FDynamicMesh3, ESPMode::ThreadSafe>(UE::ToolTarget::GetDynamicMeshCopy(Target, GetMeshParams));
 
 		EnterProgressFrame(1);
 		OriginalMeshSpatial = MakeShared<FDynamicMeshAABBTree3, ESPMode::ThreadSafe>(OriginalMesh.Get(), true);
diff --git a/Engine/Plugins/Experimental/MeshModelingToolsetExp/Source/MeshModelingToolsExp/Private/BakeMeshAttributeMapsTool.cpp b/Engine/Plugins/Experimental/MeshModelingToolsetExp/Source/MeshModelingToolsExp/Private/BakeMeshAttributeMapsTool.cpp
--- a/Engine/Plugins/Experimental/MeshModelingToolsetExp/Source/MeshModelingToolsExp/Private/BakeMeshAttributeMapsTool.cpp
+++ b/Engine/Plugins/Experimental/MeshModelingToolsetExp/Source/MeshModelingToolsExp/Private/BakeMeshAttributeMapsTool.cpp
@@ -554,7 +554,9 @@ void UBakeMeshAttributeMapsTool::UpdateDetailMesh()
 	UToolTarget* DetailTarget = Targets[bIsBakeToSelf ? 0 : 1];
 
 	const bool bWantMeshTangents = (InputMeshSettings->SourceNormalMap != nullptr);
-	DetailMesh = MakeShared<FDynamicMesh3, ESPMode::ThreadSafe>(UE::ToolTarget::GetDynamicMeshCopy(DetailTarget, bWantMeshTangents));
+	FGetMeshParameters GetMeshParams;
+	GetMeshParams.bWantMeshTangents = bWantMeshTangents;
+	DetailMesh = MakeShared<FDynamicMesh3, ESPMode::ThreadSafe>(UE::ToolTarget::GetDynamicMeshCopy(DetailTarget, GetMeshParams));
 
 	if (InputMeshSettings->bProjectionInWorldSpace && bIsBakeToSelf == false)
 	{
diff --git a/Engine/Plugins/Experimental/MeshModelingToolsetExp/Source/MeshModelingToolsExp/Private/BakeMeshAttributeVertexTool.cpp b/Engine/Plugins/Experimental/MeshModelingToolsetExp/Source/MeshModelingToolsExp/Private/BakeMeshAttributeVertexTool.cpp
--- a/Engine/Plugins/Experimental/MeshModelingToolsetExp/Source/MeshModelingToolsExp/Private/BakeMeshAttributeVertexTool.cpp
+++ b/Engine/Plugins/Experimental/MeshModelingToolsetExp/Source/MeshModelingToolsExp/Private/BakeMeshAttributeVertexTool.cpp
@@ -281,7 +281,9 @@ void UBakeMeshAttributeVertexTool::Setup()
 	// TargetMesh stores the original target mesh. It is intended to remain
 	// const throughout this tool and is used to refresh the PreviewMesh back
 	// to its original state.
-	TargetMesh = MakeShared<FDynamicMesh3, ESPMode::ThreadSafe>(UE::ToolTarget::GetDynamicMeshCopy(Targets[0], true));
+	static FGetMeshParameters GetMeshParams;
+	GetMeshParams.bWantMeshTangents = true;
+	TargetMesh = MakeShared<FDynamicMesh3, ESPMode::ThreadSafe>(UE::ToolTarget::GetDynamicMeshCopy(Targets[0], GetMeshParams));
 	TargetMeshTangents = MakeShared<FMeshTangentsd, ESPMode::ThreadSafe>(TargetMesh.Get());
 	TargetMeshTangents->CopyTriVertexTangents(*TargetMesh);
 
@@ -477,7 +479,9 @@ void UBakeMeshAttributeVertexTool::UpdateDetailMesh()
 	IPrimitiveComponentBackedTarget* DetailComponent = Cast<IPrimitiveComponentBackedTarget>(Targets[bIsBakeToSelf ? 0 : 1]);
 	UToolTarget* DetailTargetMesh = Targets[bIsBakeToSelf ? 0 : 1];
 
-	const FDynamicMesh3 DetailMeshCopy = UE::ToolTarget::GetDynamicMeshCopy(DetailTargetMesh, true);
+	static FGetMeshParameters GetMeshParams;
+	GetMeshParams.bWantMeshTangents = true;
+	const FDynamicMesh3 DetailMeshCopy = UE::ToolTarget::GetDynamicMeshCopy(DetailTargetMesh, GetMeshParams);
 	DetailMesh = MakeShared<FDynamicMesh3, ESPMode::ThreadSafe>();
 	DetailMesh->Copy(DetailMeshCopy);
 	if (InputMeshSettings->bProjectionInWorldSpace && bIsBakeToSelf == false)
diff --git a/Engine/Plugins/Experimental/MeshModelingToolsetExp/Source/MeshModelingToolsExp/Private/BakeToolUtils.cpp b/Engine/Plugins/Experimental/MeshModelingToolsetExp/Source/MeshModelingToolsExp/Private/BakeToolUtils.cpp
--- a/Engine/Plugins/Experimental/MeshModelingToolsetExp/Source/MeshModelingToolsExp/Private/BakeToolUtils.cpp
+++ b/Engine/Plugins/Experimental/MeshModelingToolsetExp/Source/MeshModelingToolsExp/Private/BakeToolUtils.cpp
@@ -39,7 +39,9 @@ UPreviewMesh* CreateBakePreviewMesh(
 	UToolTarget* ToolTarget,
 	UWorld* World)
 {
-	const FDynamicMesh3 InputMesh = UE::ToolTarget::GetDynamicMeshCopy(ToolTarget, true);
+	static FGetMeshParameters GetMeshParams;
+	GetMeshParams.bWantMeshTangents = true;
+	const FDynamicMesh3 InputMesh = UE::ToolTarget::GetDynamicMeshCopy(ToolTarget, GetMeshParams);
 	const FTransformSRT3d BaseToWorld = UE::ToolTarget::GetLocalToWorldTransform(ToolTarget);
 	const FComponentMaterialSet MaterialSet = UE::ToolTarget::GetMaterialSet(ToolTarget);
 
diff --git a/Engine/Plugins/Experimental/MeshModelingToolsetExp/Source/MeshModelingToolsExp/Private/MeshInspectorTool.cpp b/Engine/Plugins/Experimental/MeshModelingToolsetExp/Source/MeshModelingToolsExp/Private/MeshInspectorTool.cpp
--- a/Engine/Plugins/Experimental/MeshModelingToolsetExp/Source/MeshModelingToolsExp/Private/MeshInspectorTool.cpp
+++ b/Engine/Plugins/Experimental/MeshModelingToolsetExp/Source/MeshModelingToolsExp/Private/MeshInspectorTool.cpp
@@ -204,7 +204,9 @@ void UMeshInspectorTool::Setup()
 	DefaultMaterial = PreviewMesh->GetMaterial(0);
 
 	PreviewMesh->SetTangentsMode(EDynamicMeshComponentTangentsMode::ExternallyProvided);
-	FDynamicMesh3 InputMeshWithTangents = UE::ToolTarget::GetDynamicMeshCopy(Target, true);
+	static FGetMeshParameters GetMeshParams;
+	GetMeshParams.bWantMeshTangents = true;
+	FDynamicMesh3 InputMeshWithTangents = UE::ToolTarget::GetDynamicMeshCopy(Target, GetMeshParams);
 	PreviewMesh->ReplaceMesh(MoveTemp(InputMeshWithTangents));
 
 	DrawnLineSet = NewObject<ULineSetComponent>(PreviewMesh->GetRootComponent(), "MeshInspectorToolLineSet");
diff --git a/Engine/Plugins/Experimental/MeshModelingToolsetExp/Source/MeshModelingToolsExp/Private/PatternTool.cpp b/Engine/Plugins/Experimental/MeshModelingToolsetExp/Source/MeshModelingToolsExp/Private/PatternTool.cpp
--- a/Engine/Plugins/Experimental/MeshModelingToolsetExp/Source/MeshModelingToolsExp/Private/PatternTool.cpp
+++ b/Engine/Plugins/Experimental/MeshModelingToolsetExp/Source/MeshModelingToolsExp/Private/PatternTool.cpp
@@ -1773,7 +1773,9 @@ void UPatternTool::EmitResults()
 		
 		if (Element.SourceDynamicMesh != nullptr || bConvertToDynamic )
 		{
-			FDynamicMesh3 ElementMesh = UE::ToolTarget::GetDynamicMeshCopy(Targets[ElemIdx], true);
+			static FGetMeshParameters GetMeshParams;
+			GetMeshParams.bWantMeshTangents = true;
+			FDynamicMesh3 ElementMesh = UE::ToolTarget::GetDynamicMeshCopy(Targets[ElemIdx], GetMeshParams);
 
 			// this lambda creates a new dynamic mesh actor w/ the materials from the current Element
 			auto EmitDynamicMeshActor = [this, &Element, ModelingSettings](FDynamicMesh3&& MoveMesh, FString BaseName, FTransformSRT3d Transform)
diff --git a/Engine/Plugins/Experimental/MeshModelingToolsetExp/Source/MeshModelingToolsExp/Private/Sculpting/MeshSculptToolBase.cpp b/Engine/Plugins/Experimental/MeshModelingToolsetExp/Source/MeshModelingToolsExp/Private/Sculpting/MeshSculptToolBase.cpp
--- a/Engine/Plugins/Experimental/MeshModelingToolsetExp/Source/MeshModelingToolsExp/Private/Sculpting/MeshSculptToolBase.cpp
+++ b/Engine/Plugins/Experimental/MeshModelingToolsetExp/Source/MeshModelingToolsExp/Private/Sculpting/MeshSculptToolBase.cpp
@@ -305,7 +305,9 @@ void UMeshSculptToolBase::InitializeSculptMeshComponent(UBaseDynamicMeshComponen
 	Component->RegisterComponent();
 
 	// initialize from LOD-0 MeshDescription
-	Component->SetMesh(UE::ToolTarget::GetDynamicMeshCopy(Target, true));
+	static FGetMeshParameters GetMeshParams;
+	GetMeshParams.bWantMeshTangents = true;
+	Component->SetMesh(UE::ToolTarget::GetDynamicMeshCopy(Target, GetMeshParams));
 	double MaxDimension = Component->GetMesh()->GetBounds(true).MaxDim();
 
 	// bake rotation and scaling into mesh because handling these inside sculpting is a mess
diff --git a/Engine/Plugins/Experimental/MeshModelingToolsetExp/Source/MeshModelingToolsExp/Private/SplitMeshesTool.cpp b/Engine/Plugins/Experimental/MeshModelingToolsetExp/Source/MeshModelingToolsExp/Private/SplitMeshesTool.cpp
--- a/Engine/Plugins/Experimental/MeshModelingToolsetExp/Source/MeshModelingToolsExp/Private/SplitMeshesTool.cpp
+++ b/Engine/Plugins/Experimental/MeshModelingToolsetExp/Source/MeshModelingToolsExp/Private/SplitMeshesTool.cpp
@@ -69,11 +69,14 @@ void USplitMeshesTool::Setup()
 	BasicProperties->WatchProperty(BasicProperties->bShowPreview, [this](bool bShowPreview) { UpdatePreviewVisibility(bShowPreview); });
 	AddToolPropertySource(BasicProperties);
 
+	static FGetMeshParameters GetMeshParams;
+	GetMeshParams.bWantMeshTangents = true;
+	
 	SourceMeshes.SetNum(Targets.Num());
 	bool bHasSelection = false;
 	for (int32 k = 0; k < Targets.Num(); ++k)
 	{
-		SourceMeshes[k].Mesh = UE::ToolTarget::GetDynamicMeshCopy(Targets[k], true);
+		SourceMeshes[k].Mesh = UE::ToolTarget::GetDynamicMeshCopy(Targets[k], GetMeshParams);
 		SourceMeshes[k].Materials = UE::ToolTarget::GetMaterialSet(Targets[k]).Materials;
 		bHasSelection = bHasSelection || HasGeometrySelection(k);
 	}
diff --git a/Engine/Plugins/Experimental/MeshModelingToolsetExp/Source/MeshModelingToolsExp/Private/WeldMeshEdgesTool.cpp b/Engine/Plugins/Experimental/MeshModelingToolsetExp/Source/MeshModelingToolsExp/Private/WeldMeshEdgesTool.cpp
--- a/Engine/Plugins/Experimental/MeshModelingToolsetExp/Source/MeshModelingToolsExp/Private/WeldMeshEdgesTool.cpp
+++ b/Engine/Plugins/Experimental/MeshModelingToolsetExp/Source/MeshModelingToolsExp/Private/WeldMeshEdgesTool.cpp
@@ -214,8 +214,9 @@ void UWeldMeshEdgesTool::Setup()
 {
 	UInteractiveTool::Setup();
 
-
-	SourceMesh = MakeShared<FDynamicMesh3, ESPMode::ThreadSafe>(UE::ToolTarget::GetDynamicMeshCopy(Target, true));
+	static FGetMeshParameters GetMeshParams;
+	GetMeshParams.bWantMeshTangents = true;
+	SourceMesh = MakeShared<FDynamicMesh3, ESPMode::ThreadSafe>(UE::ToolTarget::GetDynamicMeshCopy(Target, GetMeshParams));
 
 	// initialize selection if exists
 	if (HasGeometrySelection())
diff --git a/Engine/Plugins/Runtime/MeshModelingToolset/Source/MeshModelingTools/Private/CombineMeshesTool.cpp b/Engine/Plugins/Runtime/MeshModelingToolset/Source/MeshModelingTools/Private/CombineMeshesTool.cpp
--- a/Engine/Plugins/Runtime/MeshModelingToolset/Source/MeshModelingTools/Private/CombineMeshesTool.cpp
+++ b/Engine/Plugins/Runtime/MeshModelingToolset/Source/MeshModelingTools/Private/CombineMeshesTool.cpp
@@ -186,11 +186,13 @@ void UCombineMeshesTool::CreateNewAsset()
 
 	// Make sure meshes are available before we open transaction. This is to avoid potential stability issues related 
 	// to creation/load of meshes inside a transaction, for assets that possibly do not have bulk data currently loaded.
+	static FGetMeshParameters GetMeshParams;
+	GetMeshParams.bWantMeshTangents = true;
 	TArray<FDynamicMesh3> InputMeshes;
 	InputMeshes.Reserve(Targets.Num());
 	for (int32 ComponentIdx = 0; ComponentIdx < Targets.Num(); ComponentIdx++)
 	{
-		InputMeshes.Add(UE::ToolTarget::GetDynamicMeshCopy(Targets[ComponentIdx], true));
+		InputMeshes.Add(UE::ToolTarget::GetDynamicMeshCopy(Targets[ComponentIdx], GetMeshParams));
 	}
 
 	GetToolManager()->BeginUndoTransaction( bDuplicateMode ? 
@@ -369,11 +371,13 @@ void UCombineMeshesTool::UpdateExistingAsset()
 
 	// Make sure meshes are available before we open transaction. This is to avoid potential stability issues related 
 	// to creation/load of meshes inside a transaction, for assets that possibly do not have bulk data currently loaded.
+	static FGetMeshParameters GetMeshParams;
+	GetMeshParams.bWantMeshTangents = true;
 	TArray<FDynamicMesh3> InputMeshes;
 	InputMeshes.Reserve(Targets.Num());
 	for (int32 ComponentIdx = 0; ComponentIdx < Targets.Num(); ComponentIdx++)
 	{
-		InputMeshes.Add(UE::ToolTarget::GetDynamicMeshCopy(Targets[ComponentIdx], true));
+		InputMeshes.Add(UE::ToolTarget::GetDynamicMeshCopy(Targets[ComponentIdx], GetMeshParams));
 	}
 
 	check(!bDuplicateMode);
diff --git a/Engine/Plugins/Runtime/MeshModelingToolset/Source/ModelingComponents/Private/ConversionUtils/DynamicMeshViaMeshDescriptionUtil.cpp b/Engine/Plugins/Runtime/MeshModelingToolset/Source/ModelingComponents/Private/ConversionUtils/DynamicMeshViaMeshDescriptionUtil.cpp
--- a/Engine/Plugins/Runtime/MeshModelingToolset/Source/ModelingComponents/Private/ConversionUtils/DynamicMeshViaMeshDescriptionUtil.cpp
+++ b/Engine/Plugins/Runtime/MeshModelingToolset/Source/ModelingComponents/Private/ConversionUtils/DynamicMeshViaMeshDescriptionUtil.cpp
@@ -10,21 +10,28 @@
 
 FDynamicMesh3 UE::Geometry::GetDynamicMeshViaMeshDescription(
 	IMeshDescriptionProvider& MeshDescriptionProvider, bool bRequestTangents)
+{
+	FGetMeshParameters GetMeshParams;
+	GetMeshParams.bWantMeshTangents = bRequestTangents;
+	return GetDynamicMeshViaMeshDescription(MeshDescriptionProvider, GetMeshParams);
+}
+
+FDynamicMesh3 UE::Geometry::GetDynamicMeshViaMeshDescription(
+	IMeshDescriptionProvider& MeshDescriptionProvider,
+	const FGetMeshParameters& InGetMeshParams)
 {
 	FDynamicMesh3 DynamicMesh;
 	FMeshDescriptionToDynamicMesh Converter;
 	Converter.bVIDsFromNonManifoldMeshDescriptionAttr = true;
 	Converter.SetPolygonGroupToMaterialIndexMap(MeshDescriptionProvider.GetPolygonGroupToMaterialIndexMap());
-	if (bRequestTangents)
+	if (InGetMeshParams.bWantMeshTangents)
 	{
-		FGetMeshParameters GetMeshParams;
-		GetMeshParams.bWantMeshTangents = true;
-		FMeshDescription MeshDescriptionCopy = MeshDescriptionProvider.GetMeshDescriptionCopy(GetMeshParams);
-		Converter.Convert(&MeshDescriptionCopy, DynamicMesh, bRequestTangents);
+		FMeshDescription MeshDescriptionCopy = MeshDescriptionProvider.GetMeshDescriptionCopy(InGetMeshParams);
+		Converter.Convert(&MeshDescriptionCopy, DynamicMesh, InGetMeshParams.bWantMeshTangents);
 	}
 	else
 	{
-		Converter.Convert(MeshDescriptionProvider.GetMeshDescription(), DynamicMesh, bRequestTangents);
+		Converter.Convert(MeshDescriptionProvider.GetMeshDescription(InGetMeshParams), DynamicMesh);
 	}
 	return DynamicMesh;
 }
diff --git a/Engine/Plugins/Runtime/MeshModelingToolset/Source/ModelingComponents/Private/ModelingToolTargetUtil.cpp b/Engine/Plugins/Runtime/MeshModelingToolset/Source/ModelingComponents/Private/ModelingToolTargetUtil.cpp
--- a/Engine/Plugins/Runtime/MeshModelingToolset/Source/ModelingComponents/Private/ModelingToolTargetUtil.cpp
+++ b/Engine/Plugins/Runtime/MeshModelingToolset/Source/ModelingComponents/Private/ModelingToolTargetUtil.cpp
@@ -224,8 +224,14 @@ FMeshDescription UE::ToolTarget::GetMeshDescriptionCopy(UToolTarget* Target, con
 	return FMeshDescription();
 }
 
-
 FDynamicMesh3 UE::ToolTarget::GetDynamicMeshCopy(UToolTarget* Target, bool bWantMeshTangents)
+{
+	FGetMeshParameters GetMeshParams;
+	GetMeshParams.bWantMeshTangents = bWantMeshTangents;
+	return GetDynamicMeshCopy(Target, GetMeshParams);
+}
+
+FDynamicMesh3 UE::ToolTarget::GetDynamicMeshCopy(UToolTarget* Target, const FGetMeshParameters& InGetMeshParams)
 {
 	IPersistentDynamicMeshSource* DynamicMeshSource = Cast<IPersistentDynamicMeshSource>(Target);
 	if (DynamicMeshSource)
@@ -239,7 +245,7 @@ FDynamicMesh3 UE::ToolTarget::GetDynamicMeshCopy(UToolTarget* Target, bool bWant
 	IDynamicMeshProvider* DynamicMeshProvider = Cast<IDynamicMeshProvider>(Target);
 	if (DynamicMeshProvider)
 	{
-		return DynamicMeshProvider->GetDynamicMesh(bWantMeshTangents);
+		return DynamicMeshProvider->GetDynamicMesh(InGetMeshParams);
 	}
 
 	IMeshDescriptionProvider* MeshDescriptionProvider = Cast<IMeshDescriptionProvider>(Target);
@@ -250,16 +256,14 @@ FDynamicMesh3 UE::ToolTarget::GetDynamicMeshCopy(UToolTarget* Target, bool bWant
 		FMeshDescriptionToDynamicMesh Converter;
 		Converter.bVIDsFromNonManifoldMeshDescriptionAttr = true;
 		Converter.SetPolygonGroupToMaterialIndexMap(MeshDescriptionProvider->GetPolygonGroupToMaterialIndexMap());
-		if (bWantMeshTangents)
+		if (InGetMeshParams.bWantMeshTangents)
 		{
-			FGetMeshParameters GetMeshParams;
-			GetMeshParams.bWantMeshTangents = true;
-			FMeshDescription MeshDescriptionCopy = MeshDescriptionProvider->GetMeshDescriptionCopy(GetMeshParams);
-			Converter.Convert(&MeshDescriptionCopy, Mesh, bWantMeshTangents);
+			const FMeshDescription MeshDescriptionCopy = MeshDescriptionProvider->GetMeshDescriptionCopy(InGetMeshParams);
+			Converter.Convert(&MeshDescriptionCopy, Mesh, InGetMeshParams.bWantMeshTangents);
 		}
 		else
 		{
-			Converter.Convert(MeshDescriptionProvider->GetMeshDescription(), Mesh, bWantMeshTangents);
+			Converter.Convert(MeshDescriptionProvider->GetMeshDescription(InGetMeshParams), Mesh);
 		}
 
 		return Mesh;
diff --git a/Engine/Plugins/Runtime/MeshModelingToolset/Source/ModelingComponents/Private/TargetInterfaces/DynamicMeshProvider.cpp b/Engine/Plugins/Runtime/MeshModelingToolset/Source/ModelingComponents/Private/TargetInterfaces/DynamicMeshProvider.cpp
--- a/Engine/Plugins/Runtime/MeshModelingToolset/Source/ModelingComponents/Private/TargetInterfaces/DynamicMeshProvider.cpp
+++ b/Engine/Plugins/Runtime/MeshModelingToolset/Source/ModelingComponents/Private/TargetInterfaces/DynamicMeshProvider.cpp
@@ -2,12 +2,20 @@
 
 #include "TargetInterfaces/DynamicMeshProvider.h"
 #include "DynamicMesh/DynamicMesh3.h"
+#include "TargetInterfaces/MeshTargetInterfaceTypes.h"
 
 #define LOCTEXT_NAMESPACE "DynamicMeshProvider"
 
 using namespace UE::Geometry;
 
 FDynamicMesh3 IDynamicMeshProvider::GetDynamicMesh(bool bRequestTangents)
+{
+	FGetMeshParameters GetMeshParams;
+	GetMeshParams.bWantMeshTangents = bRequestTangents;
+	return GetDynamicMesh(GetMeshParams);
+}
+
+FDynamicMesh3 IDynamicMeshProvider::GetDynamicMesh(const FGetMeshParameters& InGetMeshParams)
 {
 	return GetDynamicMesh();
 }
diff --git a/Engine/Plugins/Runtime/MeshModelingToolset/Source/ModelingComponents/Public/ConversionUtils/DynamicMeshViaMeshDescriptionUtil.h b/Engine/Plugins/Runtime/MeshModelingToolset/Source/ModelingComponents/Public/ConversionUtils/DynamicMeshViaMeshDescriptionUtil.h
--- a/Engine/Plugins/Runtime/MeshModelingToolset/Source/ModelingComponents/Public/ConversionUtils/DynamicMeshViaMeshDescriptionUtil.h
+++ b/Engine/Plugins/Runtime/MeshModelingToolset/Source/ModelingComponents/Public/ConversionUtils/DynamicMeshViaMeshDescriptionUtil.h
@@ -5,6 +5,7 @@
 #include "CoreMinimal.h"
 
 #include "TargetInterfaces/DynamicMeshCommitter.h"
+#include "TargetInterfaces/MeshTargetInterfaceTypes.h"
 
 #include "GeometryBase.h"
 
@@ -16,8 +17,13 @@ struct FMeshDescription;
 namespace UE {
 namespace Geometry {
 
+	UE_DEPRECATED(5.5, "Use GetDynamicMeshViaMeshDescription which takes a FGetMeshParameters instead.")
 	MODELINGCOMPONENTS_API FDynamicMesh3 GetDynamicMeshViaMeshDescription(
-		IMeshDescriptionProvider& MeshDescriptionProvider, bool bRequestTangents = false);
+		IMeshDescriptionProvider& MeshDescriptionProvider, bool bRequestTangents);
+	
+	MODELINGCOMPONENTS_API FDynamicMesh3 GetDynamicMeshViaMeshDescription(
+		IMeshDescriptionProvider& MeshDescriptionProvider,
+		const FGetMeshParameters& InGetMeshParams = FGetMeshParameters());
 
 	MODELINGCOMPONENTS_API void CommitDynamicMeshViaMeshDescription(
 		FMeshDescription&& CurrentMeshDescription,
diff --git a/Engine/Plugins/Runtime/MeshModelingToolset/Source/ModelingComponents/Public/ModelingToolTargetUtil.h b/Engine/Plugins/Runtime/MeshModelingToolset/Source/ModelingComponents/Public/ModelingToolTargetUtil.h
--- a/Engine/Plugins/Runtime/MeshModelingToolset/Source/ModelingComponents/Public/ModelingToolTargetUtil.h
+++ b/Engine/Plugins/Runtime/MeshModelingToolset/Source/ModelingComponents/Public/ModelingToolTargetUtil.h
@@ -141,8 +141,19 @@ MODELINGCOMPONENTS_API FMeshDescription GetMeshDescriptionCopy(
  * @param bWantMeshTangents if true, tangents will be returned if the target has them available. This may require that they be auto-calculated in some cases (which may be expensive)
  * @return a created DynamicMesh3, which may be empty if the Target doesn't have a mesh 
  */
-MODELINGCOMPONENTS_API UE::Geometry::FDynamicMesh3 GetDynamicMeshCopy(UToolTarget* Target, bool bWantMeshTangents = false);
-
+	UE_DEPRECATED(5.5, "Use GetDynamicMeshCopy which takes a FGetMeshParameters instead.")
+	MODELINGCOMPONENTS_API UE::Geometry::FDynamicMesh3 GetDynamicMeshCopy(UToolTarget* Target, bool bWantMeshTangents);
+	
+/**
+ * Fetch a DynamicMesh3 representing the given ToolTarget. This may be a conversion of the output of GetMeshDescription().
+ * This function returns a copy, so the caller can take ownership of this Mesh.
+ * @param InGetMeshParams to specify various options like specific LOD and/or tangents on the returned mesh.
+ * if InGetMeshParams.bWantMeshTangents is true, tangents will be returned if the target has them available. This may require that they be auto-calculated in some cases (which may be expensive)
+ * @return a created DynamicMesh3, which may be empty if the Target doesn't have a mesh 
+ */
+	MODELINGCOMPONENTS_API UE::Geometry::FDynamicMesh3 GetDynamicMeshCopy(
+		UToolTarget* Target,
+		const FGetMeshParameters& InGetMeshParams = FGetMeshParameters());
 
 /**
  * EDynamicMeshUpdateResult is returned by functions below that update a ToolTarget with a new Mesh
diff --git a/Engine/Plugins/Runtime/MeshModelingToolset/Source/ModelingComponents/Public/TargetInterfaces/DynamicMeshProvider.h b/Engine/Plugins/Runtime/MeshModelingToolset/Source/ModelingComponents/Public/TargetInterfaces/DynamicMeshProvider.h
--- a/Engine/Plugins/Runtime/MeshModelingToolset/Source/ModelingComponents/Public/TargetInterfaces/DynamicMeshProvider.h
+++ b/Engine/Plugins/Runtime/MeshModelingToolset/Source/ModelingComponents/Public/TargetInterfaces/DynamicMeshProvider.h
@@ -9,6 +9,8 @@
 
 #include "DynamicMeshProvider.generated.h"
 
+struct FGetMeshParameters;
+
 PREDECLARE_GEOMETRY(class FDynamicMesh3);
 
 UINTERFACE()
@@ -34,5 +36,16 @@ public:
 	 *
 	 * Note: Default implementation simply returns GetDynamicMesh(). Overloaded implementations for e.g., Static and Skeletal Mesh sources will enable (and compute if needed) additional tangent data.
 	 */
+	UE_DEPRECATED(5.5, "Use GetDynamicMesh which takes a FGetMeshParameters instead.")
 	virtual UE::Geometry::FDynamicMesh3 GetDynamicMesh(bool bRequestTangents);
+	
+	/**
+	 * Gives a copy of a dynamic mesh for tools to operate on.
+	 * 
+	 * @param InGetMeshParams Request specific LOD and/or tangents on the returned mesh.
+	 * bWantMeshTangents not required if tangents are not on the source data and the provider does not have a standard way to generate them.
+	 *
+	 * Note: Default implementation simply returns GetDynamicMesh(). Overloaded implementations for e.g., Static and Skeletal Mesh sources will enable (and compute if needed) additional tangent data.
+	 */
+	virtual UE::Geometry::FDynamicMesh3 GetDynamicMesh(const FGetMeshParameters& InGetMeshParams);
 };
diff --git a/Engine/Plugins/Runtime/MeshModelingToolset/Source/ModelingComponentsEditorOnly/Private/ToolTargets/SkeletalMeshComponentToolTarget.cpp b/Engine/Plugins/Runtime/MeshModelingToolset/Source/ModelingComponentsEditorOnly/Private/ToolTargets/SkeletalMeshComponentToolTarget.cpp
--- a/Engine/Plugins/Runtime/MeshModelingToolset/Source/ModelingComponentsEditorOnly/Private/ToolTargets/SkeletalMeshComponentToolTarget.cpp
+++ b/Engine/Plugins/Runtime/MeshModelingToolset/Source/ModelingComponentsEditorOnly/Private/ToolTargets/SkeletalMeshComponentToolTarget.cpp
@@ -151,9 +151,9 @@ FDynamicMesh3 USkeletalMeshComponentReadOnlyToolTarget::GetDynamicMesh()
 	return GetDynamicMeshViaMeshDescription(*this);
 }
 
-FDynamicMesh3 USkeletalMeshComponentReadOnlyToolTarget::GetDynamicMesh(bool bRequestTangents)
+FDynamicMesh3 USkeletalMeshComponentReadOnlyToolTarget::GetDynamicMesh(const FGetMeshParameters& InGetMeshParams)
 {
-	return GetDynamicMeshViaMeshDescription(*this, bRequestTangents);
+	return GetDynamicMeshViaMeshDescription(*this, InGetMeshParams);
 }
 
 
diff --git a/Engine/Plugins/Runtime/MeshModelingToolset/Source/ModelingComponentsEditorOnly/Private/ToolTargets/SkeletalMeshToolTarget.cpp b/Engine/Plugins/Runtime/MeshModelingToolset/Source/ModelingComponentsEditorOnly/Private/ToolTargets/SkeletalMeshToolTarget.cpp
--- a/Engine/Plugins/Runtime/MeshModelingToolset/Source/ModelingComponentsEditorOnly/Private/ToolTargets/SkeletalMeshToolTarget.cpp
+++ b/Engine/Plugins/Runtime/MeshModelingToolset/Source/ModelingComponentsEditorOnly/Private/ToolTargets/SkeletalMeshToolTarget.cpp
@@ -221,9 +221,9 @@ FDynamicMesh3 USkeletalMeshReadOnlyToolTarget::GetDynamicMesh()
 	return GetDynamicMeshViaMeshDescription(*this);
 }
 
-FDynamicMesh3 USkeletalMeshReadOnlyToolTarget::GetDynamicMesh(bool bRequestTangents)
+FDynamicMesh3 USkeletalMeshReadOnlyToolTarget::GetDynamicMesh(const FGetMeshParameters& InGetMeshParams)
 {
-	return GetDynamicMeshViaMeshDescription(*this, bRequestTangents);
+	return GetDynamicMeshViaMeshDescription(*this, InGetMeshParams);
 }
 
 USkeletalMesh* USkeletalMeshReadOnlyToolTarget::GetSkeletalMesh() const
diff --git a/Engine/Plugins/Runtime/MeshModelingToolset/Source/ModelingComponentsEditorOnly/Private/ToolTargets/StaticMeshComponentToolTarget.cpp b/Engine/Plugins/Runtime/MeshModelingToolset/Source/ModelingComponentsEditorOnly/Private/ToolTargets/StaticMeshComponentToolTarget.cpp
--- a/Engine/Plugins/Runtime/MeshModelingToolset/Source/ModelingComponentsEditorOnly/Private/ToolTargets/StaticMeshComponentToolTarget.cpp
+++ b/Engine/Plugins/Runtime/MeshModelingToolset/Source/ModelingComponentsEditorOnly/Private/ToolTargets/StaticMeshComponentToolTarget.cpp
@@ -252,9 +252,9 @@ FDynamicMesh3 UStaticMeshComponentToolTarget::GetDynamicMesh()
 	return GetDynamicMeshViaMeshDescription(*this);
 }
 
-FDynamicMesh3 UStaticMeshComponentToolTarget::GetDynamicMesh(bool bRequestTangents)
+FDynamicMesh3 UStaticMeshComponentToolTarget::GetDynamicMesh(const FGetMeshParameters& InGetMeshParams)
 {
-	return GetDynamicMeshViaMeshDescription(*this, bRequestTangents);
+	return GetDynamicMeshViaMeshDescription(*this, InGetMeshParams);
 }
 
 void UStaticMeshComponentToolTarget::CommitDynamicMesh(const FDynamicMesh3& Mesh, const FDynamicMeshCommitInfo& CommitInfo)
diff --git a/Engine/Plugins/Runtime/MeshModelingToolset/Source/ModelingComponentsEditorOnly/Private/ToolTargets/StaticMeshToolTarget.cpp b/Engine/Plugins/Runtime/MeshModelingToolset/Source/ModelingComponentsEditorOnly/Private/ToolTargets/StaticMeshToolTarget.cpp
--- a/Engine/Plugins/Runtime/MeshModelingToolset/Source/ModelingComponentsEditorOnly/Private/ToolTargets/StaticMeshToolTarget.cpp
+++ b/Engine/Plugins/Runtime/MeshModelingToolset/Source/ModelingComponentsEditorOnly/Private/ToolTargets/StaticMeshToolTarget.cpp
@@ -490,9 +490,9 @@ FDynamicMesh3 UStaticMeshToolTarget::GetDynamicMesh()
 	return GetDynamicMeshViaMeshDescription(*this);
 }
 
-FDynamicMesh3 UStaticMeshToolTarget::GetDynamicMesh(bool bRequestTangents)
+FDynamicMesh3 UStaticMeshToolTarget::GetDynamicMesh(const FGetMeshParameters& InGetMeshParams)
 {
-	return GetDynamicMeshViaMeshDescription(*this, bRequestTangents);
+	return GetDynamicMeshViaMeshDescription(*this, InGetMeshParams);
 }
 
 void UStaticMeshToolTarget::CommitDynamicMesh(const FDynamicMesh3& Mesh, const FDynamicMeshCommitInfo& CommitInfo)
diff --git a/Engine/Plugins/Runtime/MeshModelingToolset/Source/ModelingComponentsEditorOnly/Public/ToolTargets/SkeletalMeshComponentToolTarget.h b/Engine/Plugins/Runtime/MeshModelingToolset/Source/ModelingComponentsEditorOnly/Public/ToolTargets/SkeletalMeshComponentToolTarget.h
--- a/Engine/Plugins/Runtime/MeshModelingToolset/Source/ModelingComponentsEditorOnly/Public/ToolTargets/SkeletalMeshComponentToolTarget.h
+++ b/Engine/Plugins/Runtime/MeshModelingToolset/Source/ModelingComponentsEditorOnly/Public/ToolTargets/SkeletalMeshComponentToolTarget.h
@@ -53,7 +53,7 @@ public:
 
 	// IDynamicMeshProvider
 	virtual UE::Geometry::FDynamicMesh3 GetDynamicMesh() override;
-	virtual UE::Geometry::FDynamicMesh3 GetDynamicMesh(bool bRequestTangents) override;
+	virtual UE::Geometry::FDynamicMesh3 GetDynamicMesh(const FGetMeshParameters& InGetMeshParams) override;
 
 	// ISkeletalMeshBackedTarget implementation
 	virtual USkeletalMesh* GetSkeletalMesh() const override;
diff --git a/Engine/Plugins/Runtime/MeshModelingToolset/Source/ModelingComponentsEditorOnly/Public/ToolTargets/SkeletalMeshToolTarget.h b/Engine/Plugins/Runtime/MeshModelingToolset/Source/ModelingComponentsEditorOnly/Public/ToolTargets/SkeletalMeshToolTarget.h
--- a/Engine/Plugins/Runtime/MeshModelingToolset/Source/ModelingComponentsEditorOnly/Public/ToolTargets/SkeletalMeshToolTarget.h
+++ b/Engine/Plugins/Runtime/MeshModelingToolset/Source/ModelingComponentsEditorOnly/Public/ToolTargets/SkeletalMeshToolTarget.h
@@ -52,7 +52,7 @@ public:
 
 	// IDynamicMeshProvider
 	virtual UE::Geometry::FDynamicMesh3 GetDynamicMesh() override;
-	virtual UE::Geometry::FDynamicMesh3 GetDynamicMesh(bool bRequestTangents) override;
+	virtual UE::Geometry::FDynamicMesh3 GetDynamicMesh(const FGetMeshParameters& InGetMeshParams) override;
 
 	// ISkeletalMeshBackedTarget implementation
 	virtual USkeletalMesh* GetSkeletalMesh() const override;
diff --git a/Engine/Plugins/Runtime/MeshModelingToolset/Source/ModelingComponentsEditorOnly/Public/ToolTargets/StaticMeshComponentToolTarget.h b/Engine/Plugins/Runtime/MeshModelingToolset/Source/ModelingComponentsEditorOnly/Public/ToolTargets/StaticMeshComponentToolTarget.h
--- a/Engine/Plugins/Runtime/MeshModelingToolset/Source/ModelingComponentsEditorOnly/Public/ToolTargets/StaticMeshComponentToolTarget.h
+++ b/Engine/Plugins/Runtime/MeshModelingToolset/Source/ModelingComponentsEditorOnly/Public/ToolTargets/StaticMeshComponentToolTarget.h
@@ -60,7 +60,7 @@ public:
 
 	// IDynamicMeshProvider
 	virtual UE::Geometry::FDynamicMesh3 GetDynamicMesh() override;
-	virtual UE::Geometry::FDynamicMesh3 GetDynamicMesh(bool bRequestTangents) override;
+	virtual UE::Geometry::FDynamicMesh3 GetDynamicMesh(const FGetMeshParameters& InGetMeshParams) override;
 
 	// IDynamicMeshCommitter
 	virtual void CommitDynamicMesh(const UE::Geometry::FDynamicMesh3& Mesh, const FDynamicMeshCommitInfo& CommitInfo) override;
diff --git a/Engine/Plugins/Runtime/MeshModelingToolset/Source/ModelingComponentsEditorOnly/Public/ToolTargets/StaticMeshToolTarget.h b/Engine/Plugins/Runtime/MeshModelingToolset/Source/ModelingComponentsEditorOnly/Public/ToolTargets/StaticMeshToolTarget.h
--- a/Engine/Plugins/Runtime/MeshModelingToolset/Source/ModelingComponentsEditorOnly/Public/ToolTargets/StaticMeshToolTarget.h
+++ b/Engine/Plugins/Runtime/MeshModelingToolset/Source/ModelingComponentsEditorOnly/Public/ToolTargets/StaticMeshToolTarget.h
@@ -66,7 +66,7 @@ public:
 
 	// IDynamicMeshProvider
 	virtual UE::Geometry::FDynamicMesh3 GetDynamicMesh() override;
-	virtual UE::Geometry::FDynamicMesh3 GetDynamicMesh(bool bRequestTangents) override;
+	virtual UE::Geometry::FDynamicMesh3 GetDynamicMesh(const FGetMeshParameters& InGetMeshParams) override;
 
 	// IDynamicMeshCommitter
 	virtual void CommitDynamicMesh(const UE::Geometry::FDynamicMesh3& Mesh, const FDynamicMeshCommitInfo& CommitInfo) override;
