Modeling Tools: Resolves flickering which happened when hovering over a selected element in MeshElementSelection

diff --git a/Engine/Plugins/Runtime/MeshModelingToolset/Source/ModelingComponents/Private/Selection/GeometrySelectionManager.cpp b/Engine/Plugins/Runtime/MeshModelingToolset/Source/ModelingComponents/Private/Selection/GeometrySelectionManager.cpp
--- a/Engine/Plugins/Runtime/MeshModelingToolset/Source/ModelingComponents/Private/Selection/GeometrySelectionManager.cpp
+++ b/Engine/Plugins/Runtime/MeshModelingToolset/Source/ModelingComponents/Private/Selection/GeometrySelectionManager.cpp
@@ -1690,10 +1690,10 @@ void UGeometrySelectionManager::SetSelectionColors(const FLinearColor Unselected
 		ToolSetupUtil::GetCustomTwoSidedDepthOffsetMaterial(this->ToolsContext->ToolManager, UnselectedCol, UnselectedParams.DepthBias, Opacity);
 
 	HoverOverSelectedParams.SelectionFillColor =
-		ToolSetupUtil::GetCustomTwoSidedDepthOffsetMaterial(this->ToolsContext->ToolManager, HoverOverSelectedCol, HoverOverSelectedParams.DepthBias, Opacity);
+		ToolSetupUtil::GetCustomTwoSidedDepthOffsetMaterial(this->ToolsContext->ToolManager, HoverOverSelectedCol, HoverOverSelectedParams.DepthBias);
 
 	HoverOverUnselectedParams.SelectionFillColor =
-		ToolSetupUtil::GetCustomTwoSidedDepthOffsetMaterial(this->ToolsContext->ToolManager, HoverOverUnselectedCol, HoverOverUnselectedParams.DepthBias, Opacity);
+		ToolSetupUtil::GetCustomTwoSidedDepthOffsetMaterial(this->ToolsContext->ToolManager, HoverOverUnselectedCol, HoverOverUnselectedParams.DepthBias, Opacity +.2);
 
 	SelectedParams.SelectionFillColor =
 		ToolSetupUtil::GetCustomTwoSidedDepthOffsetMaterial(this->ToolsContext->ToolManager, GeometrySelectedCol, SelectedParams.DepthBias, Opacity);
