Modeling Tools: Updates the visual display of Mesh Element Selection so that only applicable components are displayed depending on the selection mode

diff --git a/Engine/Plugins/Runtime/MeshModelingToolset/Source/ModelingComponents/Private/Selection/GeometrySelectionManager.cpp b/Engine/Plugins/Runtime/MeshModelingToolset/Source/ModelingComponents/Private/Selection/GeometrySelectionManager.cpp
--- a/Engine/Plugins/Runtime/MeshModelingToolset/Source/ModelingComponents/Private/Selection/GeometrySelectionManager.cpp
+++ b/Engine/Plugins/Runtime/MeshModelingToolset/Source/ModelingComponents/Private/Selection/GeometrySelectionManager.cpp
@@ -241,6 +241,18 @@ void UGeometrySelectionManager::SetMeshSelectionTypeAndMode(EGeometryElementType
 			RemoveAllSets();
 		}
 
+		// removes existing Line/Point/Triangle Sets when moving between vertex and face/edge modes
+		bool bRebuildSelectable = false;
+		if (
+			(((NewElementType == EGeometryElementType::Vertex) && (SelectionElementType == EGeometryElementType::Edge || SelectionElementType == EGeometryElementType::Face))
+			|| ((NewElementType == EGeometryElementType::Edge || NewElementType == EGeometryElementType::Face) && (SelectionElementType == EGeometryElementType::Vertex)))
+			&& (NewTopologyMode != EMeshTopologyMode::None)
+		) // ensure lines not rebuilt when changing to object mode
+		{
+			RemoveAllSets();
+			bRebuildSelectable = true;
+		}
+
 		GetTransactionsAPI()->BeginUndoTransaction(LOCTEXT("ChangeElementMethod", "Change Selection Method"));
 
 		// We have to undo/redo the change to the selection type because if we want to 'undo' this later and restore
@@ -256,6 +268,15 @@ void UGeometrySelectionManager::SetMeshSelectionTypeAndMode(EGeometryElementType
 		SetSelectionElementTypeInternal(NewElementType);
 		SetMeshTopologyModeInternal(NewTopologyMode);
 
+		// ensures that in Vertex mode, only vertices are displayed, while in Face or Edge mode, only faces/edges are displayed
+		if (bRebuildSelectable)
+		{
+			for (int32 k = 0; k < ActiveTargetReferences.Num(); ++k)
+			{
+				CreateOrUpdateAllSets(CachedSelectableRenderElements[k], UnselectedParams);
+			}
+		}
+
 		if (bHasSelection && bConvertSelection && ensure(ActiveTargetReferences.Num() == OldTypeSelections.Num()))
 		{
 			for (int32 TargetIdx = 0; TargetIdx < ActiveTargetReferences.Num(); ++TargetIdx)
@@ -1530,28 +1551,33 @@ void UGeometrySelectionManager::CreateOrUpdateAllSets(const FGeometrySelectionEl
 	{
 		PreviewGeometry->CreateInWorld(World, FTransform::Identity);
 	}
-	
-	PreviewGeometry->CreateOrUpdateLineSet(SelectionParams.Identifiers[1], Elements.Segments.Num(), [&](int32 j, TArray<FRenderableLine>& LinesOut)
+
+	if(SelectionElementType == EGeometryElementType::Edge || SelectionElementType == EGeometryElementType::Face)
+	{
+		PreviewGeometry->CreateOrUpdateLineSet(SelectionParams.Identifiers[1], Elements.Segments.Num(), [&](int32 j, TArray<FRenderableLine>& LinesOut)
 		{
 		const FSegment3d Seg = Elements.Segments[j];
 		LinesOut.Add(FRenderableLine(Seg.StartPoint(), Seg.EndPoint(),  SelectionParams.Color,  SelectionParams.LineThickness, SelectionParams.DepthBias));
 		}, 1);
 	
-	PreviewGeometry->CreateOrUpdateTriangleSet(SelectionParams.Identifiers[2], Elements.Triangles.Num(), [&](int32 k, TArray<FRenderableTriangle>& TrianglesOut)
-		{
-		const FTriangle3d Triangle = Elements.Triangles[k];
-		const FVector3d Normal = Triangle.Normal();
-		FRenderableTriangleVertex A(Triangle.V[0], FVector2D(0,0), Normal, SelectionParams.Color);
-		FRenderableTriangleVertex B(Triangle.V[1], FVector2D(1,0), Normal, SelectionParams.Color);
-		FRenderableTriangleVertex C(Triangle.V[2], FVector2D(1,1), Normal, SelectionParams.Color);
-		TrianglesOut.Add(FRenderableTriangle(SelectionParams.SelectionFillColor, A, B, C));
-		},1 );
-	
-	PreviewGeometry->CreateOrUpdatePointSet(SelectionParams.Identifiers[0], Elements.Points.Num(), [&](int32 k, TArray<FRenderablePoint>& PointsOut)
+		PreviewGeometry->CreateOrUpdateTriangleSet(SelectionParams.Identifiers[2], Elements.Triangles.Num(), [&](int32 k, TArray<FRenderableTriangle>& TrianglesOut)
+			{
+			const FTriangle3d Triangle = Elements.Triangles[k];
+			const FVector3d Normal = Triangle.Normal();
+			FRenderableTriangleVertex A(Triangle.V[0], FVector2D(0,0), Normal, SelectionParams.Color);
+			FRenderableTriangleVertex B(Triangle.V[1], FVector2D(1,0), Normal, SelectionParams.Color);
+			FRenderableTriangleVertex C(Triangle.V[2], FVector2D(1,1), Normal, SelectionParams.Color);
+			TrianglesOut.Add(FRenderableTriangle(SelectionParams.SelectionFillColor, A, B, C));
+			},1 );
+	}
+	else if (SelectionElementType == EGeometryElementType::Vertex)
+	{
+		PreviewGeometry->CreateOrUpdatePointSet(SelectionParams.Identifiers[0], Elements.Points.Num(), [&](int32 k, TArray<FRenderablePoint>& PointsOut)
 		{
 		const FVector3d Point = Elements.Points[k];
 		PointsOut.Add(FRenderablePoint(Point, SelectionParams.Color, SelectionParams.PointSize, SelectionParams.DepthBias));
 		});
+	}
 }
 
 void UGeometrySelectionManager::RebuildSelectionRenderCache()
diff --git a/Engine/Plugins/Runtime/MeshModelingToolset/Source/ModelingComponents/Public/Selection/GeometrySelectionManager.h b/Engine/Plugins/Runtime/MeshModelingToolset/Source/ModelingComponents/Public/Selection/GeometrySelectionManager.h
--- a/Engine/Plugins/Runtime/MeshModelingToolset/Source/ModelingComponents/Public/Selection/GeometrySelectionManager.h
+++ b/Engine/Plugins/Runtime/MeshModelingToolset/Source/ModelingComponents/Public/Selection/GeometrySelectionManager.h
@@ -528,11 +528,11 @@ protected:
 	inline static TStaticArray<FString, 3> SelectedSetIds = {"Selected_Vertices", "Selected_Lines", "Selected_Triangles"};
 
 	UPROPERTY()
-	FMeshElementSelectionParams UnselectedParams {UnselectedSetIds, 5.f, 4.f, 8.f};
+	FMeshElementSelectionParams UnselectedParams {UnselectedSetIds, 5.f, 2.f, 8.f};
 	UPROPERTY()
-	FMeshElementSelectionParams HoverOverSelectedParams {HoverOverSelectedSetIds, 10.f, 8.f, 10.f};
+	FMeshElementSelectionParams HoverOverSelectedParams {HoverOverSelectedSetIds, 10.f, 6.f, 10.f};
 	UPROPERTY()
-	FMeshElementSelectionParams HoverOverUnselectedParams {HoverOverUnselectedSetIds, 10.f, 8.f, 10.f};
+	FMeshElementSelectionParams HoverOverUnselectedParams {HoverOverUnselectedSetIds, 10.f, 6.f, 10.f};
 	UPROPERTY()
 	FMeshElementSelectionParams SelectedParams {SelectedSetIds, 6.f, 6.f, 10.f};
 	
