ModelingComponents: Expose LineSetComponent as a blueprint spawnable component.

diff --git a/Engine/Plugins/Runtime/MeshModelingToolset/Source/ModelingComponents/Private/Drawing/LineSetComponent.cpp b/Engine/Plugins/Runtime/MeshModelingToolset/Source/ModelingComponents/Private/Drawing/LineSetComponent.cpp
--- a/Engine/Plugins/Runtime/MeshModelingToolset/Source/ModelingComponents/Private/Drawing/LineSetComponent.cpp
+++ b/Engine/Plugins/Runtime/MeshModelingToolset/Source/ModelingComponents/Private/Drawing/LineSetComponent.cpp
@@ -258,6 +258,21 @@ int32 ULineSetComponent::AddLineInternal(const FRenderableLine& Line)
 	return ID;
 }
 
+int32 ULineSetComponent::AddLines(const TArray<FVector>& InStart, const TArray<FVector>& InEnd, const FColor& InColor, const float InThickness, const float InDepthBias)
+{
+	if(InStart.Num() != InEnd.Num())
+	{
+		FFrame::KismetExecutionMessage(TEXT("InStart and InEnd arrays are not the same size."), ELogVerbosity::Error);
+		return 0;
+	}
+
+	const int NumLines = InStart.Num();
+	for (int Index = 0; Index < NumLines; ++Index)
+	{
+		AddLine(InStart[Index], InEnd[Index], InColor, InThickness, InDepthBias);
+	}
+	return NumLines;
+}
 
 void ULineSetComponent::InsertLine(const int32 ID, const FRenderableLine& OverlayLine)
 {
diff --git a/Engine/Plugins/Runtime/MeshModelingToolset/Source/ModelingComponents/Public/Drawing/LineSetComponent.h b/Engine/Plugins/Runtime/MeshModelingToolset/Source/ModelingComponents/Public/Drawing/LineSetComponent.h
--- a/Engine/Plugins/Runtime/MeshModelingToolset/Source/ModelingComponents/Public/Drawing/LineSetComponent.h
+++ b/Engine/Plugins/Runtime/MeshModelingToolset/Source/ModelingComponents/Public/Drawing/LineSetComponent.h
@@ -36,7 +36,7 @@ struct FRenderableLine
 	float DepthBias;
 };
 
-UCLASS()
+UCLASS(meta=(BlueprintSpawnableComponent))
 class MODELINGCOMPONENTS_API ULineSetComponent : public UMeshComponent
 {
 	GENERATED_BODY()
@@ -46,9 +46,11 @@ public:
 	ULineSetComponent();
 
 	/** Specify material which handles lines */
+	UFUNCTION(BlueprintCallable, Category="Line Set")
 	void SetLineMaterial(UMaterialInterface* InLineMaterial);
 
 	/** Clear the line set */
+	UFUNCTION(BlueprintCallable, Category="Line Set")
 	void Clear();
 
 	/** Reserve enough memory for up to the given ID (for inserting via ID) */
@@ -64,6 +66,15 @@ public:
 		return AddLine(FRenderableLine(InStart, InEnd, InColor, InThickness, InDepthBias));
 	}
 
+	/** Add lines to be rendered using the component. */
+	UFUNCTION(BlueprintCallable, Category="Line Set")
+	int32 AddLines(
+		const TArray<FVector>& InStart,
+		const TArray<FVector>& InEnd,
+		const FColor& InColor,
+		const float InThickness,
+		const float InDepthBias = 0.0f);
+
 	/** Insert a line with the given ID to the overlay */
 	void InsertLine(const int32 ID, const FRenderableLine& OverlayLine);
 
