Fix crash when saving staticmesh that is compiling

#rnx

diff --git a/Engine/Plugins/Runtime/MeshModelingToolset/Source/ModelingComponentsEditorOnly/Private/AssetUtils/CreateStaticMeshUtil.cpp b/Engine/Plugins/Runtime/MeshModelingToolset/Source/ModelingComponentsEditorOnly/Private/AssetUtils/CreateStaticMeshUtil.cpp
--- a/Engine/Plugins/Runtime/MeshModelingToolset/Source/ModelingComponentsEditorOnly/Private/AssetUtils/CreateStaticMeshUtil.cpp
+++ b/Engine/Plugins/Runtime/MeshModelingToolset/Source/ModelingComponentsEditorOnly/Private/AssetUtils/CreateStaticMeshUtil.cpp
@@ -54,6 +54,9 @@ UE::AssetUtils::ECreateStaticMeshResult UE::AssetUtils::CreateStaticMeshAsset(
 		return ECreateStaticMeshResult::InvalidPackage;
 	}
 
+	//Make sure the package is fully load.
+	UsePackage->FullyLoad();
+
 	// create new UStaticMesh object
 	EObjectFlags UseFlags = EObjectFlags::RF_Public | EObjectFlags::RF_Standalone;
 	UStaticMesh* NewStaticMesh = NewObject<UStaticMesh>(UsePackage, FName(*NewObjectName), UseFlags);
