ModelingTools: Fixed crash in lattice tool when deforming a mesh element selection with Deform Normals enabled.

diff --git a/Engine/Plugins/Runtime/MeshModelingToolset/Source/ModelingOperators/Private/DeformationOps/LatticeDeformerOp.cpp b/Engine/Plugins/Runtime/MeshModelingToolset/Source/ModelingOperators/Private/DeformationOps/LatticeDeformerOp.cpp
--- a/Engine/Plugins/Runtime/MeshModelingToolset/Source/ModelingOperators/Private/DeformationOps/LatticeDeformerOp.cpp
+++ b/Engine/Plugins/Runtime/MeshModelingToolset/Source/ModelingOperators/Private/DeformationOps/LatticeDeformerOp.cpp
@@ -72,54 +72,95 @@ void FLatticeDeformerOp::CalculateResult(FProgressCancel* Progress)
 
 	if (bDeformNormals)
 	{
-		if (ResultMesh->HasAttributes())
+		auto GetDeformedNormals = [this, &ExecutionInfo, &Progress](const FDynamicMesh3& Mesh, TArray<FVector3f>& DeformedNormals)
 		{
-			FDynamicMeshNormalOverlay* NormalOverlay = ResultMesh->Attributes()->PrimaryNormals();
-			check(NormalOverlay != nullptr);
-
-			TArray<FVector3f> DeformedNormals;
-			Lattice->GetRotatedOverlayNormals(LatticeControlPoints,
-											  NormalOverlay,
-											  DeformedNormals,
-											  InterpolationType,
-											  ExecutionInfo,
-											  Progress);
-
-			if (Progress && Progress->Cancelled())
+			if (Mesh.HasAttributes())
 			{
-				return;
+				const FDynamicMeshNormalOverlay* NormalOverlay = Mesh.Attributes()->PrimaryNormals();
+				check(NormalOverlay != nullptr);
+
+				Lattice->GetRotatedOverlayNormals(LatticeControlPoints,
+												  NormalOverlay,
+												  DeformedNormals,
+												  InterpolationType,
+												  ExecutionInfo,
+												  Progress);
 			}
-
-			for (int ElementID : NormalOverlay->ElementIndicesItr())
+			else if (Mesh.HasVertexNormals())
 			{
-				NormalOverlay->SetElement(ElementID, DeformedNormals[ElementID]);
+				TArray<FVector3f> OriginalNormals;
+				OriginalNormals.SetNum(Mesh.MaxVertexID());
+				for (int VertexID : Mesh.VertexIndicesItr())
+				{
+					OriginalNormals[VertexID] = Mesh.GetVertexNormal(VertexID);
+				}
+
+				Lattice->GetRotatedMeshVertexNormals(LatticeControlPoints,
+													 OriginalNormals,
+													 DeformedNormals,
+													 InterpolationType,
+													 ExecutionInfo,
+													 Progress);
 			}
+		};
+
+		const FDynamicMesh3& DeformNormalMesh = bUsingSubmesh ? Submesh->GetSubmesh() : *ResultMesh;
+		TArray<FVector3f> DeformedNormals;
+		GetDeformedNormals(DeformNormalMesh, DeformedNormals);
+
+		if (DeformedNormals.IsEmpty())
+		{
+			return;
+		}
+
+		if (Progress && Progress->Cancelled())
+		{
+			return;
 		}
-		else if (ResultMesh->HasVertexNormals())
+		
+		if (bUsingSubmesh)
 		{
-			TArray<FVector3f> OriginalNormals;
-			OriginalNormals.SetNum(ResultMesh->MaxVertexID());
-			for (int VertexID : ResultMesh->VertexIndicesItr())
+			if (ResultMesh->HasAttributes())
 			{
-				OriginalNormals[VertexID] = ResultMesh->GetVertexNormal(VertexID);
-			}
+				const FDynamicMeshNormalOverlay* SubOverlay = Submesh->GetSubmesh().Attributes()->PrimaryNormals();
+				check(SubOverlay != nullptr);
 
-			TArray<FVector3f> RotatedNormals;
-			Lattice->GetRotatedMeshVertexNormals(LatticeControlPoints,
-												 OriginalNormals,
-												 RotatedNormals,
-												 InterpolationType,
-												 ExecutionInfo,
-												 Progress);
+				FDynamicMeshNormalOverlay* ResultOverlay = ResultMesh->Attributes()->PrimaryNormals();
+				check(ResultOverlay != nullptr);
 
-			if (Progress && Progress->Cancelled())
+				for (const int SubElementID : SubOverlay->ElementIndicesItr())
+				{
+					const int BaseElementID = Submesh->MapNormalToBaseMesh(0, SubElementID);
+					ResultOverlay->SetElement(BaseElementID, DeformedNormals[SubElementID]);
+				}
+			}
+			else if (ResultMesh->HasVertexNormals())
 			{
-				return;
+				for (const int SubVID : Submesh->GetSubmesh().VertexIndicesItr())
+				{
+					const int BaseVertexID = Submesh->MapVertexToBaseMesh(SubVID);
+					ResultMesh->SetVertexNormal(BaseVertexID, DeformedNormals[SubVID]);
+				}
 			}
-
-			for (int vid : ResultMesh->VertexIndicesItr())
+		}
+		else
+		{
+			if (ResultMesh->HasAttributes())
+			{
+				FDynamicMeshNormalOverlay* NormalOverlay = ResultMesh->Attributes()->PrimaryNormals();
+				check(NormalOverlay != nullptr);
+			
+				for (int ElementID : NormalOverlay->ElementIndicesItr())
+				{
+					NormalOverlay->SetElement(ElementID, DeformedNormals[ElementID]);
+				}
+			}
+			else if (ResultMesh->HasVertexNormals())
 			{
-				ResultMesh->SetVertexNormal(vid, RotatedNormals[vid]);
+				for (int vid : ResultMesh->VertexIndicesItr())
+				{
+					ResultMesh->SetVertexNormal(vid, DeformedNormals[vid]);
+				}
 			}
 		}
 	}
diff --git a/Engine/Source/Runtime/GeometryCore/Public/DynamicMesh/MeshIndexMappings.h b/Engine/Source/Runtime/GeometryCore/Public/DynamicMesh/MeshIndexMappings.h
--- a/Engine/Source/Runtime/GeometryCore/Public/DynamicMesh/MeshIndexMappings.h
+++ b/Engine/Source/Runtime/GeometryCore/Public/DynamicMesh/MeshIndexMappings.h
@@ -76,16 +76,19 @@ public:
 	inline bool ContainsGroup(int FromID) const { return GroupMap.ContainsFrom(FromID); }
 
 	FIndexMapi& GetUVMap(int UVLayer) { return UVMaps[UVLayer]; }
+	const FIndexMapi& GetUVMap(int UVLayer) const { return UVMaps[UVLayer]; }
 	void SetUV(int UVLayer, int FromID, int ToID) { UVMaps[UVLayer].Add(FromID, ToID); }
 	int GetNewUV(int UVLayer, int FromID) const { return UVMaps[UVLayer].GetTo(FromID); }
 	inline bool ContainsUV(int UVLayer, int FromID) const { return UVMaps[UVLayer].ContainsFrom(FromID); }
 
 	FIndexMapi& GetNormalMap(int NormalLayer) { return NormalMaps[NormalLayer]; }
+	const FIndexMapi& GetNormalMap(int NormalLayer) const { return NormalMaps[NormalLayer]; }
 	void SetNormal(int NormalLayer, int FromID, int ToID) { NormalMaps[NormalLayer].Add(FromID, ToID); }
 	int GetNewNormal(int NormalLayer, int FromID) const { return NormalMaps[NormalLayer].GetTo(FromID); }
 	inline bool ContainsNormal(int NormalLayer, int FromID) const { return NormalMaps[NormalLayer].ContainsFrom(FromID); }
 
 	FIndexMapi& GetColorMap() { return ColorMap; }
+	const FIndexMapi& GetColorMap() const { return ColorMap; }
 	void SetColor( int FromID, int ToID) { ColorMap.Add(FromID, ToID); }
 	int GetNewColor(int FromID) const { return ColorMap.GetTo(FromID); }
 	inline bool ContainsColor(int FromID) const { return ColorMap.ContainsFrom(FromID); }
diff --git a/Engine/Source/Runtime/GeometryCore/Public/DynamicSubmesh3.h b/Engine/Source/Runtime/GeometryCore/Public/DynamicSubmesh3.h
--- a/Engine/Source/Runtime/GeometryCore/Public/DynamicSubmesh3.h
+++ b/Engine/Source/Runtime/GeometryCore/Public/DynamicSubmesh3.h
@@ -202,6 +202,42 @@ public:
 		}
 	}
 
+	/** @return Normal ID in the submesh, mapped from the base mesh */
+	int MapNormalToSubmesh(int NormalLayer, int BaseNID) const
+	{
+		return Mappings.GetNormalMap(NormalLayer).GetTo(BaseNID);
+	}
+
+	/** @return Normal ID in the base mesh, mapped from the submesh */
+	int MapNormalToBaseMesh(int NormalLayer, int SubNID) const
+	{
+		return Mappings.GetNormalMap(NormalLayer).GetFrom(SubNID);
+	}
+
+	/** @return UV ID in the submesh, mapped from the base mesh */
+	int MapUVToSubmesh(int UVLayer, int BaseUVID) const
+	{
+		return Mappings.GetUVMap(UVLayer).GetTo(BaseUVID);
+	}
+
+	/** @return UV ID in the base mesh, mapped from the submesh */
+	int MapUVToBaseMesh(int UVLayer, int SubUVID) const
+	{
+		return Mappings.GetUVMap(UVLayer).GetFrom(SubUVID);
+	}
+
+	/** @return Vertex Color ID in the submesh, mapped from the base mesh */
+	int MapColorToSubmesh(int BaseCID) const
+	{
+		return Mappings.GetColorMap().GetTo(BaseCID);
+	}
+
+	/** @return Vertex Color ID in the base mesh, mapped from the submesh */
+	int MapColorToBaseMesh(int SubCID) const
+	{
+		return Mappings.GetColorMap().GetFrom(SubCID);
+	}
+
 	/**
 	 * Computes the Submesh object, index mappings corresponding sub to base mesh, and boundary between sub and base mesh
 	 * @param SubTriangles ArrayView of triangle IDs to include in the submesh
