fix materials slot mapping when creating a static mesh w/ unnamed materials slots

diff --git a/Engine/Plugins/Runtime/MeshModelingToolset/Source/ModelingComponentsEditorOnly/Private/AssetUtils/CreateStaticMeshUtil.cpp b/Engine/Plugins/Runtime/MeshModelingToolset/Source/ModelingComponentsEditorOnly/Private/AssetUtils/CreateStaticMeshUtil.cpp
--- a/Engine/Plugins/Runtime/MeshModelingToolset/Source/ModelingComponentsEditorOnly/Private/AssetUtils/CreateStaticMeshUtil.cpp
+++ b/Engine/Plugins/Runtime/MeshModelingToolset/Source/ModelingComponentsEditorOnly/Private/AssetUtils/CreateStaticMeshUtil.cpp
@@ -193,9 +193,13 @@ UE::AssetUtils::ECreateStaticMeshResult UE::AssetUtils::CreateStaticMeshAsset(
 			for (FPolygonGroupID PolygonGroupID : Mesh->PolygonGroups().GetElementIDs())
 			{
 				// Material index is either from the matching material slot name or the section index if that name is not found
-				int32 MaterialIndex = NewStaticMesh->GetStaticMaterials().IndexOfByPredicate(
-					[&MaterialSlotName = MaterialSlotNames[PolygonGroupID]](const FStaticMaterial& StaticMaterial) { return StaticMaterial.MaterialSlotName == MaterialSlotName; }
-				);
+				int32 MaterialIndex = INDEX_NONE;
+				if (MaterialSlotNames[PolygonGroupID] != FName())
+				{
+					MaterialIndex = NewStaticMesh->GetStaticMaterials().IndexOfByPredicate(
+						[&MaterialSlotName = MaterialSlotNames[PolygonGroupID]](const FStaticMaterial& StaticMaterial) { return StaticMaterial.MaterialSlotName == MaterialSlotName; }
+					);
+				}
 				if (MaterialIndex == INDEX_NONE)
 				{
 					MaterialIndex = SectionIndex;
