[Modeling Tools]: Changed the default target requirements for most tools and tool builders to use UDynamicMeshProvider/Commiter. Only specific tools that directly get the MeshDescription in non-optional ways.

diff --git a/Engine/Plugins/Experimental/MeshModelingToolsetExp/Source/MeshModelingToolsEditorOnlyExp/Private/AttributeEditorTool.cpp b/Engine/Plugins/Experimental/MeshModelingToolsetExp/Source/MeshModelingToolsEditorOnlyExp/Private/AttributeEditorTool.cpp
--- a/Engine/Plugins/Experimental/MeshModelingToolsetExp/Source/MeshModelingToolsEditorOnlyExp/Private/AttributeEditorTool.cpp
+++ b/Engine/Plugins/Experimental/MeshModelingToolsetExp/Source/MeshModelingToolsEditorOnlyExp/Private/AttributeEditorTool.cpp
@@ -20,6 +20,7 @@
 #include "Components/PrimitiveComponent.h"
 
 #include "TargetInterfaces/MeshDescriptionCommitter.h"
+#include "TargetInterfaces/MeshDescriptionProvider.h"
 #include "TargetInterfaces/PrimitiveComponentBackedTarget.h"
 #include "ModelingToolTargetUtil.h"
 #include "ToolTargetManager.h"
@@ -46,6 +47,17 @@ bool UAttributeEditorToolBuilder::CanBuildTool(const FToolBuilderState& SceneSta
 			[](UActorComponent& Component) { return !ToolBuilderUtil::IsVolume(Component); }) >= 1;
 }
 
+const FToolTargetTypeRequirements& UAttributeEditorToolBuilder::GetTargetRequirements() const
+{
+	static FToolTargetTypeRequirements TypeRequirements({
+		UMaterialProvider::StaticClass(),
+		UMeshDescriptionProvider::StaticClass(),
+		UMeshDescriptionCommitter::StaticClass(),
+		UPrimitiveComponentBackedTarget::StaticClass()
+		});
+	return TypeRequirements;
+}
+
 
 void UAttributeEditorActionPropertySet::PostAction(EAttributeEditorToolActions Action)
 {
diff --git a/Engine/Plugins/Experimental/MeshModelingToolsetExp/Source/MeshModelingToolsEditorOnlyExp/Private/SkeletalMesh/SkinWeightsBindingTool.cpp b/Engine/Plugins/Experimental/MeshModelingToolsetExp/Source/MeshModelingToolsEditorOnlyExp/Private/SkeletalMesh/SkinWeightsBindingTool.cpp
--- a/Engine/Plugins/Experimental/MeshModelingToolsetExp/Source/MeshModelingToolsEditorOnlyExp/Private/SkeletalMesh/SkinWeightsBindingTool.cpp
+++ b/Engine/Plugins/Experimental/MeshModelingToolsetExp/Source/MeshModelingToolsEditorOnlyExp/Private/SkeletalMesh/SkinWeightsBindingTool.cpp
@@ -9,6 +9,9 @@
 #include "MeshDescriptionToDynamicMesh.h"
 #include "MeshOpPreviewHelpers.h"
 #include "SkeletalMeshAttributes.h"
+#include "TargetInterfaces/MeshDescriptionProvider.h"
+#include "TargetInterfaces/MeshDescriptionCommitter.h"
+#include "TargetInterfaces/PrimitiveComponentBackedTarget.h"
 #include "ToolSetupUtil.h"
 #include "ToolTargetManager.h"
 #include "Components/SkeletalMeshComponent.h"
@@ -40,6 +43,17 @@ UMultiSelectionMeshEditingTool* USkinWeightsBindingToolBuilder::CreateNewTool(co
 	return Tool;
 }
 
+const FToolTargetTypeRequirements& USkinWeightsBindingToolBuilder::GetTargetRequirements() const
+{
+	static FToolTargetTypeRequirements TypeRequirements({
+		UMaterialProvider::StaticClass(),
+		UMeshDescriptionProvider::StaticClass(),
+		UMeshDescriptionCommitter::StaticClass(),
+		UPrimitiveComponentBackedTarget::StaticClass()
+		});
+	return TypeRequirements;
+}
+
 USkinWeightsBindingTool::USkinWeightsBindingTool()
 {
 	Properties = CreateDefaultSubobject<USkinWeightsBindingToolProperties>(TEXT("SkinWeightsBindingProperties"));
diff --git a/Engine/Plugins/Experimental/MeshModelingToolsetExp/Source/MeshModelingToolsEditorOnlyExp/Private/SkeletalMesh/SkinWeightsPaintTool.cpp b/Engine/Plugins/Experimental/MeshModelingToolsetExp/Source/MeshModelingToolsEditorOnlyExp/Private/SkeletalMesh/SkinWeightsPaintTool.cpp
--- a/Engine/Plugins/Experimental/MeshModelingToolsetExp/Source/MeshModelingToolsEditorOnlyExp/Private/SkeletalMesh/SkinWeightsPaintTool.cpp
+++ b/Engine/Plugins/Experimental/MeshModelingToolsetExp/Source/MeshModelingToolsEditorOnlyExp/Private/SkeletalMesh/SkinWeightsPaintTool.cpp
@@ -42,6 +42,8 @@
 #include "AnimationRuntime.h"
 #include "DynamicMesh/DynamicVertexSkinWeightsAttribute.h"
 #include "TargetInterfaces/DynamicMeshCommitter.h"
+#include "TargetInterfaces/MeshDescriptionCommitter.h"
+#include "TargetInterfaces/MeshDescriptionProvider.h"
 #include "Operations/TransferBoneWeights.h"
 
 #include UE_INLINE_GENERATED_CPP_BY_NAME(SkinWeightsPaintTool)
@@ -1012,6 +1014,17 @@ UMeshSurfacePointTool* USkinWeightsPaintToolBuilder::CreateNewTool(const FToolBu
 	return Tool;
 }
 
+const FToolTargetTypeRequirements& USkinWeightsPaintToolBuilder::GetTargetRequirements() const
+{
+	static FToolTargetTypeRequirements TypeRequirements({
+		UMaterialProvider::StaticClass(),
+		UMeshDescriptionProvider::StaticClass(),
+		UMeshDescriptionCommitter::StaticClass(),
+		UPrimitiveComponentBackedTarget::StaticClass()
+		});
+	return TypeRequirements;
+}
+
 void USkinWeightsPaintTool::Init(const FToolBuilderState& InSceneState)
 {
 	const UContextObjectStore* ContextObjectStore = InSceneState.ToolManager->GetContextObjectStore();
diff --git a/Engine/Plugins/Experimental/MeshModelingToolsetExp/Source/MeshModelingToolsEditorOnlyExp/Private/SubdividePolyTool.cpp b/Engine/Plugins/Experimental/MeshModelingToolsetExp/Source/MeshModelingToolsEditorOnlyExp/Private/SubdividePolyTool.cpp
--- a/Engine/Plugins/Experimental/MeshModelingToolsetExp/Source/MeshModelingToolsEditorOnlyExp/Private/SubdividePolyTool.cpp
+++ b/Engine/Plugins/Experimental/MeshModelingToolsetExp/Source/MeshModelingToolsEditorOnlyExp/Private/SubdividePolyTool.cpp
@@ -471,7 +471,7 @@ void USubdividePolyTool::OnShutdown(EToolShutdownType ShutdownType)
 			UDynamicMeshComponent* PreviewDynamicMeshComponent = (UDynamicMeshComponent*)PreviewMesh->GetRootComponent();
 			FDynamicMesh3* DynamicMeshResult = PreviewDynamicMeshComponent->GetRenderMesh();
 
-			UE::ToolTarget::CommitMeshDescriptionUpdateViaDynamicMesh(Target, *DynamicMeshResult, true);
+			UE::ToolTarget::CommitDynamicMeshUpdate(Target, *DynamicMeshResult, true);
 
 			GetToolManager()->EndUndoTransaction();
 		}
diff --git a/Engine/Plugins/Experimental/MeshModelingToolsetExp/Source/MeshModelingToolsEditorOnlyExp/Public/AttributeEditorTool.h b/Engine/Plugins/Experimental/MeshModelingToolsetExp/Source/MeshModelingToolsEditorOnlyExp/Public/AttributeEditorTool.h
--- a/Engine/Plugins/Experimental/MeshModelingToolsetExp/Source/MeshModelingToolsEditorOnlyExp/Public/AttributeEditorTool.h
+++ b/Engine/Plugins/Experimental/MeshModelingToolsetExp/Source/MeshModelingToolsEditorOnlyExp/Public/AttributeEditorTool.h
@@ -25,6 +25,9 @@ class MESHMODELINGTOOLSEDITORONLYEXP_API UAttributeEditorToolBuilder : public UM
 public:
 	virtual UMultiSelectionMeshEditingTool* CreateNewTool(const FToolBuilderState& SceneState) const override;
 	virtual bool CanBuildTool(const FToolBuilderState& SceneState) const override;
+
+protected:
+	virtual const FToolTargetTypeRequirements& GetTargetRequirements() const override;
 };
 
 
diff --git a/Engine/Plugins/Experimental/MeshModelingToolsetExp/Source/MeshModelingToolsEditorOnlyExp/Public/SkeletalMesh/SkinWeightsBindingTool.h b/Engine/Plugins/Experimental/MeshModelingToolsetExp/Source/MeshModelingToolsEditorOnlyExp/Public/SkeletalMesh/SkinWeightsBindingTool.h
--- a/Engine/Plugins/Experimental/MeshModelingToolsetExp/Source/MeshModelingToolsEditorOnlyExp/Public/SkeletalMesh/SkinWeightsBindingTool.h
+++ b/Engine/Plugins/Experimental/MeshModelingToolsetExp/Source/MeshModelingToolsEditorOnlyExp/Public/SkeletalMesh/SkinWeightsBindingTool.h
@@ -34,6 +34,9 @@ class MESHMODELINGTOOLSEDITORONLYEXP_API USkinWeightsBindingToolBuilder : public
 public:
 	virtual bool CanBuildTool(const FToolBuilderState& SceneState) const override;
 	virtual UMultiSelectionMeshEditingTool* CreateNewTool(const FToolBuilderState& SceneState) const override;
+
+protected:
+	virtual const FToolTargetTypeRequirements& GetTargetRequirements() const override;
 };
 
 // A mirror of UE::Geometry::ESkinBindingType
diff --git a/Engine/Plugins/Experimental/MeshModelingToolsetExp/Source/MeshModelingToolsEditorOnlyExp/Public/SkeletalMesh/SkinWeightsPaintTool.h b/Engine/Plugins/Experimental/MeshModelingToolsetExp/Source/MeshModelingToolsEditorOnlyExp/Public/SkeletalMesh/SkinWeightsPaintTool.h
--- a/Engine/Plugins/Experimental/MeshModelingToolsetExp/Source/MeshModelingToolsEditorOnlyExp/Public/SkeletalMesh/SkinWeightsPaintTool.h
+++ b/Engine/Plugins/Experimental/MeshModelingToolsetExp/Source/MeshModelingToolsEditorOnlyExp/Public/SkeletalMesh/SkinWeightsPaintTool.h
@@ -276,6 +276,9 @@ class MESHMODELINGTOOLSEDITORONLYEXP_API USkinWeightsPaintToolBuilder : public U
 
 public:
 	UMeshSurfacePointTool* CreateNewTool(const FToolBuilderState& SceneState) const override;
+
+protected:
+	virtual const FToolTargetTypeRequirements& GetTargetRequirements() const override;
 };
 
 // for saveing/restoring the brush settings separately for each brush mode (Add, Replace, etc...)
diff --git a/Engine/Plugins/Experimental/MeshModelingToolsetExp/Source/MeshModelingToolsExp/Private/BakeTransformTool.cpp b/Engine/Plugins/Experimental/MeshModelingToolsetExp/Source/MeshModelingToolsExp/Private/BakeTransformTool.cpp
--- a/Engine/Plugins/Experimental/MeshModelingToolsetExp/Source/MeshModelingToolsExp/Private/BakeTransformTool.cpp
+++ b/Engine/Plugins/Experimental/MeshModelingToolsetExp/Source/MeshModelingToolsExp/Private/BakeTransformTool.cpp
@@ -40,7 +40,16 @@ UMultiSelectionMeshEditingTool* UBakeTransformToolBuilder::CreateNewTool(const F
 	return NewObject<UBakeTransformTool>(SceneState.ToolManager);
 }
 
-
+const FToolTargetTypeRequirements& UBakeTransformToolBuilder::GetTargetRequirements() const
+{
+	static FToolTargetTypeRequirements TypeRequirements({
+		UMaterialProvider::StaticClass(),
+		UMeshDescriptionProvider::StaticClass(),
+		UMeshDescriptionCommitter::StaticClass(),
+		UPrimitiveComponentBackedTarget::StaticClass()
+		});
+	return TypeRequirements;
+}
 
 
 /*
diff --git a/Engine/Plugins/Experimental/MeshModelingToolsetExp/Source/MeshModelingToolsExp/Private/DisplaceMeshTool.cpp b/Engine/Plugins/Experimental/MeshModelingToolsetExp/Source/MeshModelingToolsExp/Private/DisplaceMeshTool.cpp
--- a/Engine/Plugins/Experimental/MeshModelingToolsetExp/Source/MeshModelingToolsExp/Private/DisplaceMeshTool.cpp
+++ b/Engine/Plugins/Experimental/MeshModelingToolsetExp/Source/MeshModelingToolsExp/Private/DisplaceMeshTool.cpp
@@ -22,6 +22,9 @@
 #include "Operations/SelectiveTessellate.h"
 #include "Materials/MaterialInterface.h"
 #include "Properties/MeshStatisticsProperties.h"
+#include "TargetInterfaces/MeshDescriptionProvider.h"
+#include "TargetInterfaces/MeshDescriptionCommitter.h"
+#include "TargetInterfaces/PrimitiveComponentBackedTarget.h"
 
 // needed to disable normals recalculation on the underlying asset
 #include "AssetUtils/MeshDescriptionUtil.h"
@@ -928,6 +931,18 @@ USingleTargetWithSelectionTool* UDisplaceMeshToolBuilder::CreateNewTool(const FT
 	return NewObject<UDisplaceMeshTool>(SceneState.ToolManager);
 }
 
+const FToolTargetTypeRequirements& UDisplaceMeshToolBuilder::GetTargetRequirements() const
+{
+	static FToolTargetTypeRequirements TypeRequirements({
+		UMaterialProvider::StaticClass(),
+		UMeshDescriptionProvider::StaticClass(),
+		UMeshDescriptionCommitter::StaticClass(),
+		UPrimitiveComponentBackedTarget::StaticClass()
+		});
+	return TypeRequirements;
+}
+
+
 /*
  * Tool
  */
diff --git a/Engine/Plugins/Experimental/MeshModelingToolsetExp/Source/MeshModelingToolsExp/Private/EditPivotTool.cpp b/Engine/Plugins/Experimental/MeshModelingToolsetExp/Source/MeshModelingToolsExp/Private/EditPivotTool.cpp
--- a/Engine/Plugins/Experimental/MeshModelingToolsetExp/Source/MeshModelingToolsExp/Private/EditPivotTool.cpp
+++ b/Engine/Plugins/Experimental/MeshModelingToolsetExp/Source/MeshModelingToolsExp/Private/EditPivotTool.cpp
@@ -71,7 +71,16 @@ void UEditPivotToolBuilder::InitializeNewTool(UMultiSelectionMeshEditingTool* Ne
 	}
 }
 
-
+const FToolTargetTypeRequirements& UEditPivotToolBuilder::GetTargetRequirements() const
+{
+	static FToolTargetTypeRequirements TypeRequirements({
+		UMaterialProvider::StaticClass(),
+		UMeshDescriptionProvider::StaticClass(),
+		UMeshDescriptionCommitter::StaticClass(),
+		UPrimitiveComponentBackedTarget::StaticClass()
+		});
+	return TypeRequirements;
+}
 
 
 void UEditPivotToolActionPropertySet::PostAction(EEditPivotToolActions Action)
diff --git a/Engine/Plugins/Experimental/MeshModelingToolsetExp/Source/MeshModelingToolsExp/Private/HoleFillTool.cpp b/Engine/Plugins/Experimental/MeshModelingToolsetExp/Source/MeshModelingToolsExp/Private/HoleFillTool.cpp
--- a/Engine/Plugins/Experimental/MeshModelingToolsetExp/Source/MeshModelingToolsExp/Private/HoleFillTool.cpp
+++ b/Engine/Plugins/Experimental/MeshModelingToolsetExp/Source/MeshModelingToolsExp/Private/HoleFillTool.cpp
@@ -302,7 +302,7 @@ void UHoleFillTool::OnShutdown(EToolShutdownType ShutdownType)
 		GetToolManager()->BeginUndoTransaction(LOCTEXT("HoleFillToolTransactionName", "Hole Fill Tool"));
 
 		check(Result.Mesh.Get() != nullptr);
-		UE::ToolTarget::CommitMeshDescriptionUpdateViaDynamicMesh(Target, *Result.Mesh.Get(), true);
+		UE::ToolTarget::CommitDynamicMeshUpdate(Target, *Result.Mesh.Get(), true);
 
 		GetToolManager()->EndUndoTransaction();
 	}
diff --git a/Engine/Plugins/Experimental/MeshModelingToolsetExp/Source/MeshModelingToolsExp/Private/LatticeDeformerTool.cpp b/Engine/Plugins/Experimental/MeshModelingToolsetExp/Source/MeshModelingToolsExp/Private/LatticeDeformerTool.cpp
--- a/Engine/Plugins/Experimental/MeshModelingToolsetExp/Source/MeshModelingToolsExp/Private/LatticeDeformerTool.cpp
+++ b/Engine/Plugins/Experimental/MeshModelingToolsetExp/Source/MeshModelingToolsExp/Private/LatticeDeformerTool.cpp
@@ -433,7 +433,7 @@ void ULatticeDeformerTool::OnShutdown(EToolShutdownType ShutdownType)
 			FTransform3d LocalToWorld(TargetComponent->GetWorldTransform());
 			MeshTransforms::ApplyTransformInverse(*DynamicMeshResult, LocalToWorld, true);
 
-			UE::ToolTarget::CommitMeshDescriptionUpdateViaDynamicMesh(Targets[0], *DynamicMeshResult, true);
+			UE::ToolTarget::CommitDynamicMeshUpdate(Targets[0], *DynamicMeshResult, true);
 
 			GetToolManager()->EndUndoTransaction();
 		}
diff --git a/Engine/Plugins/Experimental/MeshModelingToolsetExp/Source/MeshModelingToolsExp/Private/MeshAttributePaintTool.cpp b/Engine/Plugins/Experimental/MeshModelingToolsetExp/Source/MeshModelingToolsExp/Private/MeshAttributePaintTool.cpp
--- a/Engine/Plugins/Experimental/MeshModelingToolsetExp/Source/MeshModelingToolsExp/Private/MeshAttributePaintTool.cpp
+++ b/Engine/Plugins/Experimental/MeshModelingToolsetExp/Source/MeshModelingToolsExp/Private/MeshAttributePaintTool.cpp
@@ -9,6 +9,8 @@
 
 #include "MeshDescription.h"
 
+#include "TargetInterfaces/MeshDescriptionProvider.h"
+#include "TargetInterfaces/MeshDescriptionCommitter.h"
 #include "TargetInterfaces/PrimitiveComponentBackedTarget.h"
 #include "ModelingToolTargetUtil.h"
 #include "DynamicMesh/NonManifoldMappingSupport.h"
@@ -173,7 +175,16 @@ bool UMeshAttributePaintToolBuilder::CanBuildTool(const FToolBuilderState& Scene
 			[](UActorComponent& Component) { return !ToolBuilderUtil::IsVolume(Component); }) >= 1;
 }
 
-
+const FToolTargetTypeRequirements& UMeshAttributePaintToolBuilder::GetTargetRequirements() const
+{
+	static FToolTargetTypeRequirements TypeRequirements({
+		UMaterialProvider::StaticClass(),
+		UMeshDescriptionProvider::StaticClass(),
+		UMeshDescriptionCommitter::StaticClass(),
+		UPrimitiveComponentBackedTarget::StaticClass()
+		});
+	return TypeRequirements;
+}
 
 
 void UMeshAttributePaintTool::SetWorld(UWorld* World)
diff --git a/Engine/Plugins/Experimental/MeshModelingToolsetExp/Source/MeshModelingToolsExp/Private/MeshSpaceDeformerTool.cpp b/Engine/Plugins/Experimental/MeshModelingToolsetExp/Source/MeshModelingToolsExp/Private/MeshSpaceDeformerTool.cpp
--- a/Engine/Plugins/Experimental/MeshModelingToolsetExp/Source/MeshModelingToolsExp/Private/MeshSpaceDeformerTool.cpp
+++ b/Engine/Plugins/Experimental/MeshModelingToolsetExp/Source/MeshModelingToolsExp/Private/MeshSpaceDeformerTool.cpp
@@ -335,7 +335,7 @@ void UMeshSpaceDeformerTool::OnShutdown(EToolShutdownType ShutdownType)
 			FDynamicMesh3* DynamicMeshResult = Result.Mesh.Get();
 			check(DynamicMeshResult != nullptr);
 
-			UE::ToolTarget::CommitMeshDescriptionUpdateViaDynamicMesh(Target, *DynamicMeshResult, true);
+			UE::ToolTarget::CommitDynamicMeshUpdate(Target, *DynamicMeshResult, true);
 
 			GetToolManager()->EndUndoTransaction();
 		}
diff --git a/Engine/Plugins/Experimental/MeshModelingToolsetExp/Source/MeshModelingToolsExp/Private/MeshVertexPaintTool.cpp b/Engine/Plugins/Experimental/MeshModelingToolsetExp/Source/MeshModelingToolsExp/Private/MeshVertexPaintTool.cpp
--- a/Engine/Plugins/Experimental/MeshModelingToolsetExp/Source/MeshModelingToolsExp/Private/MeshVertexPaintTool.cpp
+++ b/Engine/Plugins/Experimental/MeshModelingToolsetExp/Source/MeshModelingToolsExp/Private/MeshVertexPaintTool.cpp
@@ -25,6 +25,9 @@
 #include "DynamicMesh/Operations/SplitAttributeWelder.h"
 #include "MeshDescriptionToDynamicMesh.h"
 #include "DynamicMeshToMeshDescription.h"
+#include "TargetInterfaces/MeshDescriptionProvider.h"
+#include "TargetInterfaces/MeshDescriptionCommitter.h"
+#include "TargetInterfaces/PrimitiveComponentBackedTarget.h"
 
 #include "Changes/MeshVertexChange.h"
 #include "Changes/MeshPolygroupChange.h"
@@ -77,6 +80,16 @@ bool UMeshVertexPaintToolBuilder::CanBuildTool(const FToolBuilderState& SceneSta
 			[](UActorComponent& Component) { return Cast<UMeshComponent>(&Component) != nullptr; }) > 0;
 }
 
+const FToolTargetTypeRequirements& UMeshVertexPaintToolBuilder::GetTargetRequirements() const
+{
+	static FToolTargetTypeRequirements TypeRequirements({
+		UMaterialProvider::StaticClass(),
+		UMeshDescriptionProvider::StaticClass(),
+		UMeshDescriptionCommitter::StaticClass(),
+		UPrimitiveComponentBackedTarget::StaticClass()
+		});
+	return TypeRequirements;
+}
 
 /*
  * Properties
diff --git a/Engine/Plugins/Experimental/MeshModelingToolsetExp/Source/MeshModelingToolsExp/Private/MirrorTool.cpp b/Engine/Plugins/Experimental/MeshModelingToolsetExp/Source/MeshModelingToolsExp/Private/MirrorTool.cpp
--- a/Engine/Plugins/Experimental/MeshModelingToolsetExp/Source/MeshModelingToolsExp/Private/MirrorTool.cpp
+++ b/Engine/Plugins/Experimental/MeshModelingToolsetExp/Source/MeshModelingToolsExp/Private/MirrorTool.cpp
@@ -471,7 +471,7 @@ void UMirrorTool::GenerateAsset(const TArray<FDynamicMeshOpResult>& Results)
 		{
 			NewSelection.Actors.Add(UE::ToolTarget::GetTargetActor(Targets[OrigMeshIdx]));
 
-			UE::ToolTarget::CommitMeshDescriptionUpdateViaDynamicMesh(Targets[OrigMeshIdx], *Mesh, true);
+			UE::ToolTarget::CommitDynamicMeshUpdate(Targets[OrigMeshIdx], *Mesh, true);
 		}
 		else
 		{
diff --git a/Engine/Plugins/Experimental/MeshModelingToolsetExp/Source/MeshModelingToolsExp/Private/PlaneCutTool.cpp b/Engine/Plugins/Experimental/MeshModelingToolsetExp/Source/MeshModelingToolsExp/Private/PlaneCutTool.cpp
--- a/Engine/Plugins/Experimental/MeshModelingToolsetExp/Source/MeshModelingToolsExp/Private/PlaneCutTool.cpp
+++ b/Engine/Plugins/Experimental/MeshModelingToolsetExp/Source/MeshModelingToolsExp/Private/PlaneCutTool.cpp
@@ -30,8 +30,6 @@
 
 #include "ModelingToolTargetUtil.h"
 #include "TargetInterfaces/MaterialProvider.h"
-#include "TargetInterfaces/MeshDescriptionCommitter.h"
-#include "TargetInterfaces/MeshDescriptionProvider.h"
 #include "TargetInterfaces/PrimitiveComponentBackedTarget.h"
 #include "ModelingToolTargetUtil.h"
 
@@ -476,7 +474,7 @@ void UPlaneCutTool::GenerateAsset(const TArray<FDynamicMeshOpResult>& Results)
 			}
 		}
 
-		UE::ToolTarget::CommitMeshDescriptionUpdateViaDynamicMesh(Targets[OrigMeshIdx], *UseMesh, true);
+		UE::ToolTarget::CommitDynamicMeshUpdate(Targets[OrigMeshIdx], *UseMesh, true);
 	}
 
 	if (bNeedToAdd)
diff --git a/Engine/Plugins/Experimental/MeshModelingToolsetExp/Source/MeshModelingToolsExp/Private/SeamSculptTool.cpp b/Engine/Plugins/Experimental/MeshModelingToolsetExp/Source/MeshModelingToolsExp/Private/SeamSculptTool.cpp
--- a/Engine/Plugins/Experimental/MeshModelingToolsetExp/Source/MeshModelingToolsExp/Private/SeamSculptTool.cpp
+++ b/Engine/Plugins/Experimental/MeshModelingToolsetExp/Source/MeshModelingToolsExp/Private/SeamSculptTool.cpp
@@ -129,7 +129,7 @@ void USeamSculptTool::OnShutdown(EToolShutdownType ShutdownType)
 
 		GetToolManager()->BeginUndoTransaction(LOCTEXT("SeamSculptTransactionName", "UV Seam Edit"));
 
-		UE::ToolTarget::CommitMeshDescriptionUpdateViaDynamicMesh(Target, ResultMesh, true);
+		UE::ToolTarget::CommitDynamicMeshUpdate(Target, ResultMesh, true);
 
 		GetToolManager()->EndUndoTransaction();
 	}
diff --git a/Engine/Plugins/Experimental/MeshModelingToolsetExp/Source/MeshModelingToolsExp/Private/SplitMeshesTool.cpp b/Engine/Plugins/Experimental/MeshModelingToolsetExp/Source/MeshModelingToolsExp/Private/SplitMeshesTool.cpp
--- a/Engine/Plugins/Experimental/MeshModelingToolsetExp/Source/MeshModelingToolsExp/Private/SplitMeshesTool.cpp
+++ b/Engine/Plugins/Experimental/MeshModelingToolsetExp/Source/MeshModelingToolsExp/Private/SplitMeshesTool.cpp
@@ -17,7 +17,7 @@
 #include "Selections/GeometrySelectionUtil.h"
 #include "Util/ColorConstants.h"
 
-#include "TargetInterfaces/MeshDescriptionProvider.h"
+#include "TargetInterfaces/DynamicMeshProvider.h"
 #include "TargetInterfaces/PrimitiveComponentBackedTarget.h"
 
 #include UE_INLINE_GENERATED_CPP_BY_NAME(SplitMeshesTool)
@@ -33,7 +33,7 @@ const FToolTargetTypeRequirements& USplitMeshesToolBuilder::GetTargetRequirement
 {
 	static FToolTargetTypeRequirements TypeRequirements({
 		UMaterialProvider::StaticClass(),
-		UMeshDescriptionProvider::StaticClass(),
+		UDynamicMeshProvider::StaticClass(),
 		UPrimitiveComponentBackedTarget::StaticClass()
 		});
 	return TypeRequirements;
diff --git a/Engine/Plugins/Experimental/MeshModelingToolsetExp/Source/MeshModelingToolsExp/Private/TransferMeshTool.cpp b/Engine/Plugins/Experimental/MeshModelingToolsetExp/Source/MeshModelingToolsExp/Private/TransferMeshTool.cpp
--- a/Engine/Plugins/Experimental/MeshModelingToolsetExp/Source/MeshModelingToolsExp/Private/TransferMeshTool.cpp
+++ b/Engine/Plugins/Experimental/MeshModelingToolsetExp/Source/MeshModelingToolsExp/Private/TransferMeshTool.cpp
@@ -9,6 +9,7 @@
 #include "TargetInterfaces/StaticMeshBackedTarget.h"
 #include "TargetInterfaces/MeshDescriptionProvider.h"
 #include "TargetInterfaces/MeshDescriptionCommitter.h"
+#include "TargetInterfaces/PrimitiveComponentBackedTarget.h"
 #include "ModelingToolTargetUtil.h"
 #include "Physics/ComponentCollisionUtil.h"
 #include "ShapeApproximation/SimpleShapeSet3.h"
@@ -34,6 +35,17 @@ UMultiSelectionMeshEditingTool* UTransferMeshToolBuilder::CreateNewTool(const FT
 	return NewObject<UTransferMeshTool>(SceneState.ToolManager);
 }
 
+const FToolTargetTypeRequirements& UTransferMeshToolBuilder::GetTargetRequirements() const
+{
+	static FToolTargetTypeRequirements TypeRequirements({
+		UMaterialProvider::StaticClass(),
+		UMeshDescriptionProvider::StaticClass(),
+		UMeshDescriptionCommitter::StaticClass(),
+		UPrimitiveComponentBackedTarget::StaticClass()
+		});
+	return TypeRequirements;
+}
+
 /*
  * Tool
  */
diff --git a/Engine/Plugins/Experimental/MeshModelingToolsetExp/Source/MeshModelingToolsExp/Public/BakeTransformTool.h b/Engine/Plugins/Experimental/MeshModelingToolsetExp/Source/MeshModelingToolsExp/Public/BakeTransformTool.h
--- a/Engine/Plugins/Experimental/MeshModelingToolsetExp/Source/MeshModelingToolsExp/Public/BakeTransformTool.h
+++ b/Engine/Plugins/Experimental/MeshModelingToolsetExp/Source/MeshModelingToolsExp/Public/BakeTransformTool.h
@@ -26,6 +26,9 @@ class MESHMODELINGTOOLSEXP_API UBakeTransformToolBuilder : public UMultiSelectio
 	GENERATED_BODY()
 public:
 	virtual UMultiSelectionMeshEditingTool* CreateNewTool(const FToolBuilderState& SceneState) const override;
+
+protected:
+	virtual const FToolTargetTypeRequirements& GetTargetRequirements() const override;
 };
 
 
diff --git a/Engine/Plugins/Experimental/MeshModelingToolsetExp/Source/MeshModelingToolsExp/Public/DisplaceMeshTool.h b/Engine/Plugins/Experimental/MeshModelingToolsetExp/Source/MeshModelingToolsExp/Public/DisplaceMeshTool.h
--- a/Engine/Plugins/Experimental/MeshModelingToolsetExp/Source/MeshModelingToolsExp/Public/DisplaceMeshTool.h
+++ b/Engine/Plugins/Experimental/MeshModelingToolsetExp/Source/MeshModelingToolsExp/Public/DisplaceMeshTool.h
@@ -295,6 +295,9 @@ class MESHMODELINGTOOLSEXP_API UDisplaceMeshToolBuilder : public USingleTargetWi
 	GENERATED_BODY()
 public:
 	virtual USingleTargetWithSelectionTool* CreateNewTool(const FToolBuilderState& SceneState) const override final;
+
+protected:
+	virtual const FToolTargetTypeRequirements& GetTargetRequirements() const override;
 };
 
 /**
diff --git a/Engine/Plugins/Experimental/MeshModelingToolsetExp/Source/MeshModelingToolsExp/Public/EditPivotTool.h b/Engine/Plugins/Experimental/MeshModelingToolsetExp/Source/MeshModelingToolsExp/Public/EditPivotTool.h
--- a/Engine/Plugins/Experimental/MeshModelingToolsetExp/Source/MeshModelingToolsExp/Public/EditPivotTool.h
+++ b/Engine/Plugins/Experimental/MeshModelingToolsetExp/Source/MeshModelingToolsExp/Public/EditPivotTool.h
@@ -28,6 +28,9 @@ class MESHMODELINGTOOLSEXP_API UEditPivotToolBuilder : public UMultiSelectionMes
 public:
 	virtual UMultiSelectionMeshEditingTool* CreateNewTool(const FToolBuilderState& SceneState) const override;
 	virtual void InitializeNewTool(UMultiSelectionMeshEditingTool* NewTool, const FToolBuilderState& SceneState) const override;
+
+protected:
+	virtual const FToolTargetTypeRequirements& GetTargetRequirements() const override;
 };
 
 
diff --git a/Engine/Plugins/Experimental/MeshModelingToolsetExp/Source/MeshModelingToolsExp/Public/MeshAttributePaintTool.h b/Engine/Plugins/Experimental/MeshModelingToolsetExp/Source/MeshModelingToolsExp/Public/MeshAttributePaintTool.h
--- a/Engine/Plugins/Experimental/MeshModelingToolsetExp/Source/MeshModelingToolsExp/Public/MeshAttributePaintTool.h
+++ b/Engine/Plugins/Experimental/MeshModelingToolsetExp/Source/MeshModelingToolsExp/Public/MeshAttributePaintTool.h
@@ -93,6 +93,9 @@ public:
 
 	virtual UMeshSurfacePointTool* CreateNewTool(const FToolBuilderState& SceneState) const override;
 	virtual bool CanBuildTool(const FToolBuilderState& SceneState) const override;
+
+protected:
+	virtual const FToolTargetTypeRequirements& GetTargetRequirements() const override;
 };
 
 
diff --git a/Engine/Plugins/Experimental/MeshModelingToolsetExp/Source/MeshModelingToolsExp/Public/MeshVertexPaintTool.h b/Engine/Plugins/Experimental/MeshModelingToolsetExp/Source/MeshModelingToolsExp/Public/MeshVertexPaintTool.h
--- a/Engine/Plugins/Experimental/MeshModelingToolsetExp/Source/MeshModelingToolsExp/Public/MeshVertexPaintTool.h
+++ b/Engine/Plugins/Experimental/MeshModelingToolsetExp/Source/MeshModelingToolsExp/Public/MeshVertexPaintTool.h
@@ -43,6 +43,9 @@ class MESHMODELINGTOOLSEXP_API UMeshVertexPaintToolBuilder : public UMeshSurface
 public:
 	virtual UMeshSurfacePointTool* CreateNewTool(const FToolBuilderState& SceneState) const override;
 	virtual bool CanBuildTool(const FToolBuilderState& SceneState) const override;
+
+protected:
+	virtual const FToolTargetTypeRequirements& GetTargetRequirements() const override;
 };
 
 
diff --git a/Engine/Plugins/Experimental/MeshModelingToolsetExp/Source/MeshModelingToolsExp/Public/TransferMeshTool.h b/Engine/Plugins/Experimental/MeshModelingToolsetExp/Source/MeshModelingToolsExp/Public/TransferMeshTool.h
--- a/Engine/Plugins/Experimental/MeshModelingToolsetExp/Source/MeshModelingToolsExp/Public/TransferMeshTool.h
+++ b/Engine/Plugins/Experimental/MeshModelingToolsetExp/Source/MeshModelingToolsExp/Public/TransferMeshTool.h
@@ -19,6 +19,9 @@ class MESHMODELINGTOOLSEXP_API UTransferMeshToolBuilder : public UMultiSelection
 public:
 	virtual bool CanBuildTool(const FToolBuilderState& SceneState) const override;
 	virtual UMultiSelectionMeshEditingTool* CreateNewTool(const FToolBuilderState& SceneState) const override;
+
+protected:
+	virtual const FToolTargetTypeRequirements& GetTargetRequirements() const override;
 };
 
 
diff --git a/Engine/Plugins/Runtime/MeshModelingToolset/Source/MeshModelingTools/Private/UVLayoutTool.cpp b/Engine/Plugins/Runtime/MeshModelingToolset/Source/MeshModelingTools/Private/UVLayoutTool.cpp
--- a/Engine/Plugins/Runtime/MeshModelingToolset/Source/MeshModelingTools/Private/UVLayoutTool.cpp
+++ b/Engine/Plugins/Runtime/MeshModelingToolset/Source/MeshModelingTools/Private/UVLayoutTool.cpp
@@ -18,6 +18,10 @@
 #include "ModelingToolTargetUtil.h"
 #include "ToolTargetManager.h"
 
+#include "TargetInterfaces/MeshDescriptionProvider.h"
+#include "TargetInterfaces/MeshDescriptionCommitter.h"
+#include "TargetInterfaces/PrimitiveComponentBackedTarget.h"
+
 
 #include UE_INLINE_GENERATED_CPP_BY_NAME(UVLayoutTool)
 
@@ -41,6 +45,17 @@ bool UUVLayoutToolBuilder::CanBuildTool(const FToolBuilderState& SceneState) con
 			[](UActorComponent& Component) { return ToolBuilderUtil::ComponentTypeCouldHaveUVs(Component); }) > 0;
 }
 
+const FToolTargetTypeRequirements& UUVLayoutToolBuilder::GetTargetRequirements() const
+{
+	static FToolTargetTypeRequirements TypeRequirements({
+		UMaterialProvider::StaticClass(),
+		UMeshDescriptionProvider::StaticClass(),
+		UMeshDescriptionCommitter::StaticClass(),
+		UPrimitiveComponentBackedTarget::StaticClass()
+		});
+	return TypeRequirements;
+}
+
 /*
  * Tool
  */
diff --git a/Engine/Plugins/Runtime/MeshModelingToolset/Source/MeshModelingTools/Public/UVLayoutTool.h b/Engine/Plugins/Runtime/MeshModelingToolset/Source/MeshModelingTools/Public/UVLayoutTool.h
--- a/Engine/Plugins/Runtime/MeshModelingToolset/Source/MeshModelingTools/Public/UVLayoutTool.h
+++ b/Engine/Plugins/Runtime/MeshModelingToolset/Source/MeshModelingTools/Public/UVLayoutTool.h
@@ -35,6 +35,9 @@ class MESHMODELINGTOOLS_API UUVLayoutToolBuilder : public UMultiSelectionMeshEdi
 public:
 	virtual UMultiSelectionMeshEditingTool* CreateNewTool(const FToolBuilderState& SceneState) const override;
 	virtual bool CanBuildTool(const FToolBuilderState& SceneState) const override;
+
+protected:
+	virtual const FToolTargetTypeRequirements& GetTargetRequirements() const override;
 };
 
 
diff --git a/Engine/Plugins/Runtime/MeshModelingToolset/Source/ModelingComponents/Private/BaseTools/BaseCreateFromSelectedTool.cpp b/Engine/Plugins/Runtime/MeshModelingToolset/Source/ModelingComponents/Private/BaseTools/BaseCreateFromSelectedTool.cpp
--- a/Engine/Plugins/Runtime/MeshModelingToolset/Source/ModelingComponents/Private/BaseTools/BaseCreateFromSelectedTool.cpp
+++ b/Engine/Plugins/Runtime/MeshModelingToolset/Source/ModelingComponents/Private/BaseTools/BaseCreateFromSelectedTool.cpp
@@ -293,7 +293,7 @@ void UBaseCreateFromSelectedTool::UpdateAsset(const FDynamicMeshOpResult& Result
 	MeshTransforms::ApplyTransform(*Result.Mesh, ResultTransform, true);
 	MeshTransforms::ApplyTransformInverse(*Result.Mesh, TargetToWorld, true);
 
-	UE::ToolTarget::CommitMeshDescriptionUpdateViaDynamicMesh(UpdateTarget, *Result.Mesh, true);
+	UE::ToolTarget::CommitDynamicMeshUpdate(UpdateTarget, *Result.Mesh, true);
 
 	FComponentMaterialSet MaterialSet;
 	MaterialSet.Materials = GetOutputMaterials();
diff --git a/Engine/Plugins/Runtime/MeshModelingToolset/Source/ModelingComponents/Private/BaseTools/BaseMeshProcessingTool.cpp b/Engine/Plugins/Runtime/MeshModelingToolset/Source/ModelingComponents/Private/BaseTools/BaseMeshProcessingTool.cpp
--- a/Engine/Plugins/Runtime/MeshModelingToolset/Source/ModelingComponents/Private/BaseTools/BaseMeshProcessingTool.cpp
+++ b/Engine/Plugins/Runtime/MeshModelingToolset/Source/ModelingComponents/Private/BaseTools/BaseMeshProcessingTool.cpp
@@ -15,8 +15,8 @@
 #include "MeshDescriptionToDynamicMesh.h"
 
 #include "TargetInterfaces/MaterialProvider.h"
-#include "TargetInterfaces/MeshDescriptionCommitter.h"
-#include "TargetInterfaces/MeshDescriptionProvider.h"
+#include "TargetInterfaces/DynamicMeshCommitter.h"
+#include "TargetInterfaces/DynamicMeshProvider.h"
 #include "TargetInterfaces/PrimitiveComponentBackedTarget.h"
 #include "ToolTargetManager.h"
 #include "ModelingToolTargetUtil.h"
@@ -36,8 +36,8 @@ const FToolTargetTypeRequirements& UBaseMeshProcessingToolBuilder::GetTargetRequ
 {
 	static FToolTargetTypeRequirements TypeRequirements({
 		UMaterialProvider::StaticClass(),
-		UMeshDescriptionCommitter::StaticClass(),
-		UMeshDescriptionProvider::StaticClass(),
+		UDynamicMeshCommitter::StaticClass(),
+		UDynamicMeshProvider::StaticClass(),
 		UPrimitiveComponentBackedTarget::StaticClass()
 		});
 	return TypeRequirements;
@@ -230,7 +230,7 @@ void UBaseMeshProcessingTool::Shutdown(EToolShutdownType ShutdownType)
 			}
 
 			bool bTopologyChanged = HasMeshTopologyChanged();
-			UE::ToolTarget::CommitMeshDescriptionUpdateViaDynamicMesh(Target, *DynamicMeshResult, bTopologyChanged);
+			UE::ToolTarget::CommitDynamicMeshUpdate(Target, *DynamicMeshResult, bTopologyChanged);
 
 			GetToolManager()->EndUndoTransaction();
 		}
diff --git a/Engine/Plugins/Runtime/MeshModelingToolset/Source/ModelingComponents/Private/BaseTools/MeshSurfacePointMeshEditingTool.cpp b/Engine/Plugins/Runtime/MeshModelingToolset/Source/ModelingComponents/Private/BaseTools/MeshSurfacePointMeshEditingTool.cpp
--- a/Engine/Plugins/Runtime/MeshModelingToolset/Source/ModelingComponents/Private/BaseTools/MeshSurfacePointMeshEditingTool.cpp
+++ b/Engine/Plugins/Runtime/MeshModelingToolset/Source/ModelingComponents/Private/BaseTools/MeshSurfacePointMeshEditingTool.cpp
@@ -5,8 +5,8 @@
 #include "ToolBuilderUtil.h"
 
 #include "TargetInterfaces/PrimitiveComponentBackedTarget.h"
-#include "TargetInterfaces/MeshDescriptionProvider.h"
-#include "TargetInterfaces/MeshDescriptionCommitter.h"
+#include "TargetInterfaces/DynamicMeshProvider.h"
+#include "TargetInterfaces/DynamicMeshCommitter.h"
 #include "TargetInterfaces/MaterialProvider.h"
 #include "ToolTargetManager.h"
 
@@ -21,8 +21,8 @@ const FToolTargetTypeRequirements& UMeshSurfacePointMeshEditingToolBuilder::GetT
 {
 	static FToolTargetTypeRequirements TypeRequirements({
 		UMaterialProvider::StaticClass(),
-		UMeshDescriptionProvider::StaticClass(),
-		UMeshDescriptionCommitter::StaticClass(),
+		UDynamicMeshProvider::StaticClass(),
+		UDynamicMeshCommitter::StaticClass(),
 		UPrimitiveComponentBackedTarget::StaticClass()
 		});
 	return TypeRequirements;
diff --git a/Engine/Plugins/Runtime/MeshModelingToolset/Source/ModelingComponents/Private/BaseTools/MultiSelectionMeshEditingTool.cpp b/Engine/Plugins/Runtime/MeshModelingToolset/Source/ModelingComponents/Private/BaseTools/MultiSelectionMeshEditingTool.cpp
--- a/Engine/Plugins/Runtime/MeshModelingToolset/Source/ModelingComponents/Private/BaseTools/MultiSelectionMeshEditingTool.cpp
+++ b/Engine/Plugins/Runtime/MeshModelingToolset/Source/ModelingComponents/Private/BaseTools/MultiSelectionMeshEditingTool.cpp
@@ -5,8 +5,8 @@
 #include "Engine/World.h"
 
 #include "TargetInterfaces/MaterialProvider.h"
-#include "TargetInterfaces/MeshDescriptionCommitter.h"
-#include "TargetInterfaces/MeshDescriptionProvider.h"
+#include "TargetInterfaces/DynamicMeshCommitter.h"
+#include "TargetInterfaces/DynamicMeshProvider.h"
 #include "TargetInterfaces/PrimitiveComponentBackedTarget.h"
 #include "TargetInterfaces/AssetBackedTarget.h"
 #include "ToolTargetManager.h"
@@ -20,8 +20,8 @@ const FToolTargetTypeRequirements& UMultiSelectionMeshEditingToolBuilder::GetTar
 {
 	static FToolTargetTypeRequirements TypeRequirements({
 		UMaterialProvider::StaticClass(),
-		UMeshDescriptionCommitter::StaticClass(),
-		UMeshDescriptionProvider::StaticClass(),
+		UDynamicMeshCommitter::StaticClass(),
+		UDynamicMeshProvider::StaticClass(),
 		UPrimitiveComponentBackedTarget::StaticClass()
 		});
 	return TypeRequirements;
diff --git a/Engine/Plugins/Runtime/MeshModelingToolset/Source/ModelingComponents/Private/BaseTools/SingleSelectionMeshEditingTool.cpp b/Engine/Plugins/Runtime/MeshModelingToolset/Source/ModelingComponents/Private/BaseTools/SingleSelectionMeshEditingTool.cpp
--- a/Engine/Plugins/Runtime/MeshModelingToolset/Source/ModelingComponents/Private/BaseTools/SingleSelectionMeshEditingTool.cpp
+++ b/Engine/Plugins/Runtime/MeshModelingToolset/Source/ModelingComponents/Private/BaseTools/SingleSelectionMeshEditingTool.cpp
@@ -5,8 +5,8 @@
 #include "Engine/World.h"
 
 #include "TargetInterfaces/MaterialProvider.h"
-#include "TargetInterfaces/MeshDescriptionCommitter.h"
-#include "TargetInterfaces/MeshDescriptionProvider.h"
+#include "TargetInterfaces/DynamicMeshCommitter.h"
+#include "TargetInterfaces/DynamicMeshProvider.h"
 #include "TargetInterfaces/PrimitiveComponentBackedTarget.h"
 #include "ToolTargetManager.h"
 #include "Selection/StoredMeshSelectionUtil.h"
@@ -20,8 +20,8 @@ const FToolTargetTypeRequirements& USingleSelectionMeshEditingToolBuilder::GetTa
 {
 	static FToolTargetTypeRequirements TypeRequirements({
 		UMaterialProvider::StaticClass(),
-		UMeshDescriptionCommitter::StaticClass(),
-		UMeshDescriptionProvider::StaticClass(),
+		UDynamicMeshCommitter::StaticClass(),
+		UDynamicMeshProvider::StaticClass(),
 		UPrimitiveComponentBackedTarget::StaticClass()
 		});
 	return TypeRequirements;
diff --git a/Engine/Plugins/Runtime/MeshModelingToolset/Source/ModelingComponents/Private/BaseTools/SingleTargetWithSelectionTool.cpp b/Engine/Plugins/Runtime/MeshModelingToolset/Source/ModelingComponents/Private/BaseTools/SingleTargetWithSelectionTool.cpp
--- a/Engine/Plugins/Runtime/MeshModelingToolset/Source/ModelingComponents/Private/BaseTools/SingleTargetWithSelectionTool.cpp
+++ b/Engine/Plugins/Runtime/MeshModelingToolset/Source/ModelingComponents/Private/BaseTools/SingleTargetWithSelectionTool.cpp
@@ -9,8 +9,8 @@
 #include "Drawing/PreviewGeometryActor.h"
 #include "TargetInterfaces/DynamicMeshSource.h"
 #include "TargetInterfaces/MaterialProvider.h"
-#include "TargetInterfaces/MeshDescriptionCommitter.h"
-#include "TargetInterfaces/MeshDescriptionProvider.h"
+#include "TargetInterfaces/DynamicMeshCommitter.h"
+#include "TargetInterfaces/DynamicMeshProvider.h"
 #include "TargetInterfaces/PrimitiveComponentBackedTarget.h"
 #include "ToolTargetManager.h"
 #include "Selection/StoredMeshSelectionUtil.h"
@@ -28,8 +28,8 @@ const FToolTargetTypeRequirements& USingleTargetWithSelectionToolBuilder::GetTar
 {
 	static FToolTargetTypeRequirements TypeRequirements({
 		UMaterialProvider::StaticClass(),
-		UMeshDescriptionCommitter::StaticClass(),
-		UMeshDescriptionProvider::StaticClass(),
+		UDynamicMeshCommitter::StaticClass(),
+		UDynamicMeshProvider::StaticClass(),
 		UPrimitiveComponentBackedTarget::StaticClass()
 		});
 	return TypeRequirements;
