MeshModelingTools: Minor gizmo fixes: turn a check in UInteractiveGizmo into an ensure, avoid an ensure in TransformGizmoDataBinder that fired on unused TRS gizmos, and ensure to avoid a crash in drag alignment mechanic.

diff --git a/Engine/Plugins/Runtime/MeshModelingToolset/Source/ModelingComponents/Private/Mechanics/DragAlignmentMechanic.cpp b/Engine/Plugins/Runtime/MeshModelingToolset/Source/ModelingComponents/Private/Mechanics/DragAlignmentMechanic.cpp
--- a/Engine/Plugins/Runtime/MeshModelingToolset/Source/ModelingComponents/Private/Mechanics/DragAlignmentMechanic.cpp
+++ b/Engine/Plugins/Runtime/MeshModelingToolset/Source/ModelingComponents/Private/Mechanics/DragAlignmentMechanic.cpp
@@ -21,6 +21,12 @@ using namespace UE::Geometry;
 void FDragAlignmentBase::AddToGizmo(UCombinedTransformGizmo* TransformGizmo, const TArray<const UPrimitiveComponent*>* ComponentsToIgnoreInAlignment, 
 	const TArray<const UPrimitiveComponent*>* InvisibleComponentsToIncludeInAlignment)
 {
+	if (!ensureMsgf(TransformGizmo && TransformGizmo->ActiveTarget, TEXT("Drag Alignment AddToGizmo method must "
+		"be called on a gizmo that is already initialized.")))
+	{
+		return;
+	}
+
 	// If we have components to ignore/include, we need a copy of the array so that the alignment
 	// functions can use them later.
 	TSharedPtr<TArray<const UPrimitiveComponent*>> ComponentsToIgnorePersistent;
diff --git a/Engine/Plugins/Runtime/MeshModelingToolset/Source/ModelingComponents/Private/Transforms/TransformGizmoDataBinder.cpp b/Engine/Plugins/Runtime/MeshModelingToolset/Source/ModelingComponents/Private/Transforms/TransformGizmoDataBinder.cpp
--- a/Engine/Plugins/Runtime/MeshModelingToolset/Source/ModelingComponents/Private/Transforms/TransformGizmoDataBinder.cpp
+++ b/Engine/Plugins/Runtime/MeshModelingToolset/Source/ModelingComponents/Private/Transforms/TransformGizmoDataBinder.cpp
@@ -469,8 +469,12 @@ void FTransformGizmoDataBinder::UnbindFromGizmo(UCombinedTransformGizmo* Gizmo,
 	}
 
 	// Note about this ensure: it might seem ok to unbind once the ActiveTarget is no longer around,
-	// but it can be dangerous in case it is reused elsewhere instead of being immediately destroyed.
-	if (ensure(Gizmo->ActiveTarget))
+	//  but it can be dangerous if the proxy is reused elsewhere instead of being immediately destroyed,
+	//  causing us to continue to trigger these callbacks.
+	if (ensure(Gizmo->ActiveTarget 
+		// The second condition here is for cases where we destroyed a gizmo without actually ever calling
+		//  SetActiveTarget on it, meaning that we never fully bound to it in the first place.
+		|| !BoundGizmos.Contains(Gizmo)))
 	{
 		Gizmo->ActiveTarget->OnBeginTransformEdit.RemoveAll(this);
 		Gizmo->ActiveTarget->OnTransformChanged.RemoveAll(this);
diff --git a/Engine/Source/Runtime/InteractiveToolsFramework/Private/InteractiveGizmo.cpp b/Engine/Source/Runtime/InteractiveToolsFramework/Private/InteractiveGizmo.cpp
--- a/Engine/Source/Runtime/InteractiveToolsFramework/Private/InteractiveGizmo.cpp
+++ b/Engine/Source/Runtime/InteractiveToolsFramework/Private/InteractiveGizmo.cpp
@@ -52,7 +52,7 @@ void UInteractiveGizmo::Tick(float DeltaTime)
 UInteractiveGizmoManager* UInteractiveGizmo::GetGizmoManager() const
 {
 	UInteractiveGizmoManager* GizmoManager = Cast<UInteractiveGizmoManager>(GetOuter());
-	check(GizmoManager != nullptr);
+	ensure(GizmoManager != nullptr);
 	return GizmoManager;
 }
 
