ModelingTools: Fixed geometry selection opacity to use the alpha from the ModelingMode color settings

diff --git a/Engine/Plugins/Editor/ModelingToolsEditorMode/Source/ModelingToolsEditorMode/Public/ModelingToolsEditorModeSettings.h b/Engine/Plugins/Editor/ModelingToolsEditorMode/Source/ModelingToolsEditorMode/Public/ModelingToolsEditorModeSettings.h
--- a/Engine/Plugins/Editor/ModelingToolsEditorMode/Source/ModelingToolsEditorMode/Public/ModelingToolsEditorModeSettings.h
+++ b/Engine/Plugins/Editor/ModelingToolsEditorMode/Source/ModelingToolsEditorMode/Public/ModelingToolsEditorModeSettings.h
@@ -355,7 +355,7 @@ public:
 	 * Custom color for when geometry is not a part of current selection when using Mesh Element Selection
 	 */
 	UPROPERTY(Config, EditAnywhere, Category = "Modeling Mode|Mesh Element Selection Mode")
-	FLinearColor UnselectedColor = FLinearColor(.906f, 0.038f, 0.038f); // default red
+	FLinearColor UnselectedColor = FLinearColor(.906f, 0.038f, 0.038f,.4f); // default red
 
 	/**
 	 * Custom color hovering over selected geometry in Mesh Element Selection
@@ -367,13 +367,13 @@ public:
 	* Custom color hovering over unselected geometry in Mesh Element Selection
 	*/
 	UPROPERTY(Config, EditAnywhere, Category = "Modeling Mode|Mesh Element Selection Mode")
-	FLinearColor HoverOverUnselectedColor = FLinearColor(0.027f, 0.656f, 0.027f); // default green
+	FLinearColor HoverOverUnselectedColor = FLinearColor(0.027f, 0.656f, 0.027f, .6f); // default green
 
 	/**
 	* Custom color for when geometry is a part of the current selection when using Mesh Element Selection
 	*/
 	UPROPERTY(Config, EditAnywhere, Category = "Modeling Mode|Mesh Element Selection Mode")
-	FLinearColor GeometrySelectedColor = FLinearColor(0.859f, 0.799f, 0.076f); // default yellow
+	FLinearColor GeometrySelectedColor = FLinearColor(0.859f, 0.799f, 0.076f, .4f); // default yellow
 
 public:
 
diff --git a/Engine/Plugins/Runtime/MeshModelingToolset/Source/ModelingComponents/Private/Selection/GeometrySelectionManager.cpp b/Engine/Plugins/Runtime/MeshModelingToolset/Source/ModelingComponents/Private/Selection/GeometrySelectionManager.cpp
--- a/Engine/Plugins/Runtime/MeshModelingToolset/Source/ModelingComponents/Private/Selection/GeometrySelectionManager.cpp
+++ b/Engine/Plugins/Runtime/MeshModelingToolset/Source/ModelingComponents/Private/Selection/GeometrySelectionManager.cpp
@@ -1710,19 +1710,17 @@ void UGeometrySelectionManager::SetSelectionColors(const FLinearColor Unselected
 	HoverOverUnselectedParams.Color = HoverOverUnselectedCol.ToFColor(true);
 	SelectedParams.Color = GeometrySelectedCol.ToFColor(true);
 
-	const float Opacity = .6f;
-
 	UnselectedParams.SelectionFillColor =
-		ToolSetupUtil::GetCustomTwoSidedDepthOffsetMaterial(this->ToolsContext->ToolManager, UnselectedCol, UnselectedParams.DepthBias, Opacity);
+		ToolSetupUtil::GetCustomTwoSidedDepthOffsetMaterial(this->ToolsContext->ToolManager, UnselectedCol, UnselectedParams.DepthBias, UnselectedCol.A);
 
 	HoverOverSelectedParams.SelectionFillColor =
 		ToolSetupUtil::GetCustomTwoSidedDepthOffsetMaterial(this->ToolsContext->ToolManager, HoverOverSelectedCol, HoverOverSelectedParams.DepthBias);
 
 	HoverOverUnselectedParams.SelectionFillColor =
-		ToolSetupUtil::GetCustomTwoSidedDepthOffsetMaterial(this->ToolsContext->ToolManager, HoverOverUnselectedCol, HoverOverUnselectedParams.DepthBias, Opacity +.2);
+		ToolSetupUtil::GetCustomTwoSidedDepthOffsetMaterial(this->ToolsContext->ToolManager, HoverOverUnselectedCol, HoverOverUnselectedParams.DepthBias, HoverOverUnselectedCol.A);
 
 	SelectedParams.SelectionFillColor =
-		ToolSetupUtil::GetCustomTwoSidedDepthOffsetMaterial(this->ToolsContext->ToolManager, GeometrySelectedCol, SelectedParams.DepthBias, Opacity);
+		ToolSetupUtil::GetCustomTwoSidedDepthOffsetMaterial(this->ToolsContext->ToolManager, GeometrySelectedCol, SelectedParams.DepthBias, GeometrySelectedCol.A);
 
 	// ensures that when color is changed in Editor Preferences, colors are immediately updated in the UI
 	MarkRenderCachesDirty();
