add selection validation to geometry selection manager, so we can recover if objects are deleted from under the selection system
+ fix dynamic mesh selection lambda keeping a raw pointer, and instead keep an identifier w/ a weak object pointer

diff --git a/Engine/Plugins/Editor/ModelingToolsEditorMode/Source/ModelingToolsEditorMode/Private/ModelingToolsEditorMode.cpp b/Engine/Plugins/Editor/ModelingToolsEditorMode/Source/ModelingToolsEditorMode/Private/ModelingToolsEditorMode.cpp
--- a/Engine/Plugins/Editor/ModelingToolsEditorMode/Source/ModelingToolsEditorMode/Private/ModelingToolsEditorMode.cpp
+++ b/Engine/Plugins/Editor/ModelingToolsEditorMode/Source/ModelingToolsEditorMode/Private/ModelingToolsEditorMode.cpp
@@ -278,6 +278,13 @@ void UModelingToolsEditorMode::Tick(FEditorViewportClient* ViewportClient, float
 		FModelingToolsEditorModeToolkit* ModelingToolkit = (FModelingToolsEditorModeToolkit*)Toolkit.Get();
 		ModelingToolkit->ShowRealtimeAndModeWarnings(ViewportClient->IsRealtime() == false);
 	}
+
+	// Recover from invalid selection state, in case objects have been deleted from under the selection manager
+	if (SelectionManager && !SelectionManager->ValidateSelectionState())
+	{
+		SelectionManager->ClearSelection();
+		SelectionManager->ClearActiveTargets();
+	}
 }
 
 // Note: Stylus support is currently non-functioning; the code to enable it is left here as reference in case it is brought back
diff --git a/Engine/Plugins/Runtime/MeshModelingToolset/Source/ModelingComponents/Private/Selection/DynamicMeshSelector.cpp b/Engine/Plugins/Runtime/MeshModelingToolset/Source/ModelingComponents/Private/Selection/DynamicMeshSelector.cpp
--- a/Engine/Plugins/Runtime/MeshModelingToolset/Source/ModelingComponents/Private/Selection/DynamicMeshSelector.cpp
+++ b/Engine/Plugins/Runtime/MeshModelingToolset/Source/ModelingComponents/Private/Selection/DynamicMeshSelector.cpp
@@ -795,7 +795,11 @@ TUniquePtr<IGeometrySelector> FDynamicMeshComponentSelectorFactory::BuildForTarg
 
 	TUniquePtr<FDynamicMeshSelector> Selector = MakeUnique<FDynamicMeshSelector>();
 	Selector->Initialize(TargetIdentifier, Component->GetDynamicMesh(),
-		[Component]() { return IsValid(Component) ? (UE::Geometry::FTransformSRT3d)Component->GetComponentTransform() : FTransformSRT3d::Identity(); }
+		[TargetIdentifier]() -> FTransformSRT3d
+		{ 
+			UDynamicMeshComponent* Component = TargetIdentifier.GetAsComponentType<UDynamicMeshComponent>();
+			return IsValid(Component) ? (FTransformSRT3d)Component->GetComponentTransform() : FTransformSRT3d::Identity(); 
+		}
 	);
 	return Selector;
 }
diff --git a/Engine/Plugins/Runtime/MeshModelingToolset/Source/ModelingComponents/Private/Selection/GeometrySelectionManager.cpp b/Engine/Plugins/Runtime/MeshModelingToolset/Source/ModelingComponents/Private/Selection/GeometrySelectionManager.cpp
--- a/Engine/Plugins/Runtime/MeshModelingToolset/Source/ModelingComponents/Private/Selection/GeometrySelectionManager.cpp
+++ b/Engine/Plugins/Runtime/MeshModelingToolset/Source/ModelingComponents/Private/Selection/GeometrySelectionManager.cpp
@@ -410,7 +410,27 @@ bool UGeometrySelectionManager::HasActiveTargets() const
 	return (ActiveTargetReferences.Num() > 0);
 }
 
-
+bool UGeometrySelectionManager::ValidateSelectionState() const
+{
+	for (const TSharedPtr<FGeometrySelectionTarget>& Target : ActiveTargetReferences)
+	{
+		if (!Target.IsValid())
+		{
+			return false;
+		}
+		// if we have a stale target/selection object, selection state is not valid
+		// Note: it is ok for the object to be explicitly null, just not stale
+		if (Target->SelectionIdentifer.TargetObject.IsStale())
+		{
+			return false;
+		}
+		if (Target->TargetIdentifier.TargetObject.IsStale())
+		{
+			return false;
+		}
+	}
+	return true;
+}
 
 void UGeometrySelectionManager::ClearActiveTargets()
 {
diff --git a/Engine/Plugins/Runtime/MeshModelingToolset/Source/ModelingComponents/Public/Selection/GeometrySelectionManager.h b/Engine/Plugins/Runtime/MeshModelingToolset/Source/ModelingComponents/Public/Selection/GeometrySelectionManager.h
--- a/Engine/Plugins/Runtime/MeshModelingToolset/Source/ModelingComponents/Public/Selection/GeometrySelectionManager.h
+++ b/Engine/Plugins/Runtime/MeshModelingToolset/Source/ModelingComponents/Public/Selection/GeometrySelectionManager.h
@@ -145,6 +145,12 @@ public:
 	 * @return true if there are any active selection targets
 	 */
 	bool HasActiveTargets() const;
+
+	/**
+	 * Attempt to validate the current selection state; can be called to detect if e.g. selected objects have been deleted from under the selection manager.
+	 * @return true if current active selection state appears to be valid (i.e., does not include stale / deleted objects)
+	 */
+	bool ValidateSelectionState() const;
 	
 	/**
 	 * Empty the active selection target set
