Modeling Tools: Fixes a bug which displayed the wrong visuals in Mesh Element Selection in undo/redo

diff --git a/Engine/Plugins/Runtime/MeshModelingToolset/Source/ModelingComponents/Private/Selection/GeometrySelectionManager.cpp b/Engine/Plugins/Runtime/MeshModelingToolset/Source/ModelingComponents/Private/Selection/GeometrySelectionManager.cpp
--- a/Engine/Plugins/Runtime/MeshModelingToolset/Source/ModelingComponents/Private/Selection/GeometrySelectionManager.cpp
+++ b/Engine/Plugins/Runtime/MeshModelingToolset/Source/ModelingComponents/Private/Selection/GeometrySelectionManager.cpp
@@ -88,26 +88,100 @@ public:
 	/** Makes the change to the object */
 	virtual void Apply(UObject* Object) override
 	{
-		if (FromElementType != ToElementType)
+		// do the (default) red selectable lines/verts need to be rebuilt?
+		// - ex: when moving from object mode to tri/vert/edge mode, or moving between vert and edge/face mode,
+		//		 or between triangle and polygroup topology
+		bool bRebuildSelectable = false;
+
+		UGeometrySelectionManager* GeoSelectionManager = CastChecked<UGeometrySelectionManager>(Object);
+
+		// removes existing Line/Point/Triangle Sets when moving between vertex and face/edge modes during redo
+		if (
+			(((ToElementType == EGeometryElementType::Vertex) && (FromElementType == EGeometryElementType::Edge || FromElementType == EGeometryElementType::Face))
+			|| ((ToElementType == EGeometryElementType::Edge || ToElementType == EGeometryElementType::Face) && (FromElementType == EGeometryElementType::Vertex)))
+			&& (ToTopologyMode != UGeometrySelectionManager::EMeshTopologyMode::None)
+		)
 		{
-			CastChecked<UGeometrySelectionManager>(Object)->SetSelectionElementTypeInternal(ToElementType);
+			GeoSelectionManager->RemoveAllSets();
+			bRebuildSelectable = true;
 		}
+
 		if (FromTopologyMode != ToTopologyMode)
 		{
-			CastChecked<UGeometrySelectionManager>(Object)->SetMeshTopologyModeInternal(ToTopologyMode);
+			// when changing to Object mode, lines or verts need to be cleared
+			if (ToTopologyMode == UGeometrySelectionManager::EMeshTopologyMode::None)
+			{
+				GeoSelectionManager->RemoveAllSets();
+			}
+			// in all other cases of changing topology modes, the lines/verts need to be rebuilt
+			// uses a flag to preserve order of Removing Sets->Setting Element Type/Topo Mode -> Rebuild (when applicable for each step)
+			else
+			{
+				bRebuildSelectable = true;
+			}
+			GeoSelectionManager->SetMeshTopologyModeInternal(ToTopologyMode);
+		}
+		
+		if (FromElementType != ToElementType)
+		{
+			GeoSelectionManager->SetSelectionElementTypeInternal(ToElementType);
+		}
+
+		// if applicable, rebuilds lines/verts
+		if (bRebuildSelectable)
+		{
+			GeoSelectionManager->RebuildSelectable();
 		}
+				
 	}
 
 	/** Reverts change to the object */
 	virtual void Revert(UObject* Object) override
 	{
-		if (FromElementType != ToElementType)
+		// do the (default) red selectable lines/verts need to be rebuilt?
+		// - ex: when moving from object mode to tri/vert/edge mode, or moving between vert and edge/face mode
+		//		 or between triangle and polygroup topology
+		bool bRebuildSelectable = false;
+
+		UGeometrySelectionManager* GeoSelectionManager = CastChecked<UGeometrySelectionManager>(Object);
+
+		// removes existing Line/Point/Triangle Sets when moving between vertex and face/edge modes during undo
+		if (
+			(((ToElementType == EGeometryElementType::Vertex) && (FromElementType == EGeometryElementType::Edge || FromElementType == EGeometryElementType::Face))
+			|| ((ToElementType == EGeometryElementType::Edge || ToElementType == EGeometryElementType::Face) && (FromElementType == EGeometryElementType::Vertex)))
+			&& (FromTopologyMode != UGeometrySelectionManager::EMeshTopologyMode::None)
+		)
 		{
-			CastChecked<UGeometrySelectionManager>(Object)->SetSelectionElementTypeInternal(FromElementType);
+			GeoSelectionManager->RemoveAllSets();
+			bRebuildSelectable = true;
+			
 		}
+
 		if (FromTopologyMode != ToTopologyMode)
 		{
-			CastChecked<UGeometrySelectionManager>(Object)->SetMeshTopologyModeInternal(FromTopologyMode);
+			// when changing to Object mode, lines or verts need to be cleared
+			if (FromTopologyMode == UGeometrySelectionManager::EMeshTopologyMode::None)
+			{
+				GeoSelectionManager->RemoveAllSets();
+			}
+			// in all other cases of changing topology modes, the lines/verts need to be rebuilt
+			// uses a flag to preserve order of Removing Sets->Setting Element Type/Topo Mode -> Rebuild (when applicable for each step)
+			else
+			{
+				bRebuildSelectable = true;
+			}
+			GeoSelectionManager->SetMeshTopologyModeInternal(FromTopologyMode);
+		}
+		
+		if (FromElementType != ToElementType)
+		{
+			GeoSelectionManager->SetSelectionElementTypeInternal(FromElementType);
+		}
+
+		// if applicable, rebuilds lines/verts
+		if (bRebuildSelectable)
+		{
+			GeoSelectionManager->RebuildSelectable();
 		}
 	}
 
@@ -216,6 +290,14 @@ void UGeometrySelectionManager::SetMeshTopologyMode(EMeshTopologyMode NewTopolog
 	}
 }
 
+void UGeometrySelectionManager::RebuildSelectable() const
+{
+	for (int32 k = 0; k < ActiveTargetReferences.Num(); ++k)
+	{
+		CreateOrUpdateAllSets(CachedSelectableRenderElements[k], UnselectedParams);
+	}
+}
+
 
 void UGeometrySelectionManager::SetMeshSelectionTypeAndMode(EGeometryElementType NewElementType, EMeshTopologyMode NewTopologyMode, bool bConvertSelection)
 {
@@ -268,13 +350,9 @@ void UGeometrySelectionManager::SetMeshSelectionTypeAndMode(EGeometryElementType
 		SetSelectionElementTypeInternal(NewElementType);
 		SetMeshTopologyModeInternal(NewTopologyMode);
 
-		// ensures that in Vertex mode, only vertices are displayed, while in Face or Edge mode, only faces/edges are displayed
 		if (bRebuildSelectable)
 		{
-			for (int32 k = 0; k < ActiveTargetReferences.Num(); ++k)
-			{
-				CreateOrUpdateAllSets(CachedSelectableRenderElements[k], UnselectedParams);
-			}
+			RebuildSelectable();
 		}
 
 		if (bHasSelection && bConvertSelection && ensure(ActiveTargetReferences.Num() == OldTypeSelections.Num()))
diff --git a/Engine/Plugins/Runtime/MeshModelingToolset/Source/ModelingComponents/Public/Selection/GeometrySelectionManager.h b/Engine/Plugins/Runtime/MeshModelingToolset/Source/ModelingComponents/Public/Selection/GeometrySelectionManager.h
--- a/Engine/Plugins/Runtime/MeshModelingToolset/Source/ModelingComponents/Public/Selection/GeometrySelectionManager.h
+++ b/Engine/Plugins/Runtime/MeshModelingToolset/Source/ModelingComponents/Public/Selection/GeometrySelectionManager.h
@@ -579,6 +579,8 @@ private:
 
 	void RemoveAllSets() const;
 
+	void RebuildSelectable() const;
+
 	// Tracks saved selection state. Useful when the selection is temporarily cleared (e.g., for a tool)
 	struct FSavedSelection
 	{
