Modeling Tools: Fixes bug which was not allowing correct de-selection in Mesh Element Selection - Edge Mode

diff --git a/Engine/Plugins/Runtime/MeshModelingToolset/Source/ModelingComponents/Private/Selection/DynamicMeshSelector.cpp b/Engine/Plugins/Runtime/MeshModelingToolset/Source/ModelingComponents/Private/Selection/DynamicMeshSelector.cpp
--- a/Engine/Plugins/Runtime/MeshModelingToolset/Source/ModelingComponents/Private/Selection/DynamicMeshSelector.cpp
+++ b/Engine/Plugins/Runtime/MeshModelingToolset/Source/ModelingComponents/Private/Selection/DynamicMeshSelector.cpp
@@ -296,28 +296,45 @@ void FBaseDynamicMeshSelector::UpdateSelectionViaRaycast_MeshTopology(
 	// therefore need to include both of the TriEdgeIDs which the edge belongs to to the SelectionEditor
 	if (SelectionEditor.GetElementType() == EGeometryElementType::Edge)
 	{
-		for (const uint64 AddedElement : ResultOut.SelectionDelta.Added)
+		auto SelectionFunc = [&SelectionEditor](const uint32 EncodedID)
 		{
-			FMeshTriEdgeID TriEdgeID(FGeoSelectionID(AddedElement).GeometryID);
-			TargetMesh->ProcessMesh([TriEdgeID, &SelectionEditor](const UE::Geometry::FDynamicMesh3& SourceMesh)
+			SelectionEditor.Select(EncodedID);
+		};
+		auto DeselectionFunc = [&SelectionEditor](const uint32 EncodedID)
+		{
+			SelectionEditor.RemoveFromSelection(EncodedID);
+		};
+
+		auto AffectBothTriEdgeIDs = [this, &SelectionEditor](TArray<uint64> DeltaElements, TFunctionRef<void(const uint32)> SelectionOrDeselectionFunction)
+		{
+			for (const uint64 Element : DeltaElements)
 			{
-				// the added EdgeID
-				const int32 EdgeID = SourceMesh.IsTriangle(TriEdgeID.TriangleID) ? SourceMesh.GetTriEdge(TriEdgeID.TriangleID, TriEdgeID.TriEdgeIndex) : IndexConstants::InvalidID;
-				if (SourceMesh.IsEdge(EdgeID))
+				FMeshTriEdgeID TriEdgeID(FGeoSelectionID(Element).GeometryID);
+				TargetMesh->ProcessMesh([TriEdgeID, &SelectionEditor, &SelectionOrDeselectionFunction](const UE::Geometry::FDynamicMesh3& SourceMesh)
 				{
-					SourceMesh.EnumerateTriEdgeIDsFromEdgeID(EdgeID,
-						[&SelectionEditor, &TriEdgeID](const FMeshTriEdgeID OtherTriEdgeID)
-						{
-							// avoid adding to selection edge already exists in selection
-							if (OtherTriEdgeID.TriangleID == TriEdgeID.TriangleID)
+					// the added or removed EdgeID
+					const int32 EdgeID = SourceMesh.IsTriangle(TriEdgeID.TriangleID) ? SourceMesh.GetTriEdge(TriEdgeID.TriangleID, TriEdgeID.TriEdgeIndex) : IndexConstants::InvalidID;
+					if (SourceMesh.IsEdge(EdgeID))
+					{
+						SourceMesh.EnumerateTriEdgeIDsFromEdgeID(EdgeID,
+							[&SelectionEditor, &TriEdgeID, &SelectionOrDeselectionFunction](const FMeshTriEdgeID OtherTriEdgeID)
 							{
-								return;
-							}
-							SelectionEditor.Select(OtherTriEdgeID.Encoded());
-						});
-				}
-			});
-		}
+								// avoid adding to the selection the edge which already exists in selection
+								// OR removing from the selection the edge which has already been removed
+								if (OtherTriEdgeID.TriangleID == TriEdgeID.TriangleID)
+								{
+									return;
+								}
+								SelectionOrDeselectionFunction(OtherTriEdgeID.Encoded());
+							});
+					}
+				});
+			}
+		};
+
+		// selecting or deselecting (whichever applicable) the secondary TriEdgeID
+		AffectBothTriEdgeIDs(ResultOut.SelectionDelta.Added, SelectionFunc);
+		AffectBothTriEdgeIDs(ResultOut.SelectionDelta.Removed, DeselectionFunc);
 	}
 }
 
