ModelingTools: Fixed InsertEdgeLoops tool in the Select tool bar initializing to the wrong topology type.
Disallow InsertEdgeLoops, InsertEdges and SimplifyByGroups tools in Triangle selection mode.

diff --git a/Engine/Plugins/Runtime/MeshModelingToolset/Source/MeshModelingTools/Private/EditMeshPolygonsTool.cpp b/Engine/Plugins/Runtime/MeshModelingToolset/Source/MeshModelingTools/Private/EditMeshPolygonsTool.cpp
--- a/Engine/Plugins/Runtime/MeshModelingToolset/Source/MeshModelingTools/Private/EditMeshPolygonsTool.cpp
+++ b/Engine/Plugins/Runtime/MeshModelingToolset/Source/MeshModelingTools/Private/EditMeshPolygonsTool.cpp
@@ -143,13 +143,6 @@ bool UEditMeshPolygonsActionModeToolBuilder::CanBuildTool(const FToolBuilderStat
 {
 	if (UEditMeshPolygonsToolBuilder::CanBuildTool(SceneState))
 	{
-		if ( StartupAction == EEditMeshPolygonsToolActions::SimplifyByGroups
-			 || StartupAction == EEditMeshPolygonsToolActions::InsertEdge
-			 || StartupAction == EEditMeshPolygonsToolActions::InsertEdgeLoop )
-		{
-			return true;
-		}
-
 		if ( UGeometrySelectionManager* SelectionManager = SceneState.ToolManager->GetContextObjectStore()->FindContext<UGeometrySelectionManager>() )
 		{
 			EGeometryTopologyType TopologyType = EGeometryTopologyType::Triangle;
@@ -157,6 +150,14 @@ bool UEditMeshPolygonsActionModeToolBuilder::CanBuildTool(const FToolBuilderStat
 			int NumTargets;
 			bool bIsEmpty = false;
 			SelectionManager->GetActiveSelectionInfo(TopologyType, ElementType, NumTargets, bIsEmpty);
+
+			// Default to Polygroup topology type if no topology mode selected. GetActiveSelectionInfo will return Triangle in this case.
+			if (SelectionManager->GetMeshTopologyMode() == UGeometrySelectionManager::EMeshTopologyMode::None)
+			{
+				TopologyType = EGeometryTopologyType::Polygroup;
+			}
+
+			bool bCanBuild = false;
 			if (StartupAction == EEditMeshPolygonsToolActions::Extrude
 				|| StartupAction == EEditMeshPolygonsToolActions::PushPull
 				|| StartupAction == EEditMeshPolygonsToolActions::Offset
@@ -164,16 +165,23 @@ bool UEditMeshPolygonsActionModeToolBuilder::CanBuildTool(const FToolBuilderStat
 				|| StartupAction == EEditMeshPolygonsToolActions::Outset
 				|| StartupAction == EEditMeshPolygonsToolActions::CutFaces)
 			{
-				return (TopologyType == EGeometryTopologyType::Polygroup && ElementType == EGeometryElementType::Face && bIsEmpty == false);
+				bCanBuild = (TopologyType == EGeometryTopologyType::Polygroup && ElementType == EGeometryElementType::Face && bIsEmpty == false);
 			}
 			else if (StartupAction == EEditMeshPolygonsToolActions::BevelAuto)
 			{
-				return (TopologyType == EGeometryTopologyType::Polygroup && ElementType != EGeometryElementType::Vertex && bIsEmpty == false);
+				bCanBuild = (TopologyType == EGeometryTopologyType::Polygroup && ElementType != EGeometryElementType::Vertex && bIsEmpty == false);
 			}
 			else if (StartupAction == EEditMeshPolygonsToolActions::ExtrudeEdges)
 			{
-				return (ElementType == EGeometryElementType::Edge && !bIsEmpty);
+				bCanBuild = (ElementType == EGeometryElementType::Edge && !bIsEmpty);
 			}
+			else if ( StartupAction == EEditMeshPolygonsToolActions::SimplifyByGroups
+			 || StartupAction == EEditMeshPolygonsToolActions::InsertEdge
+			 || StartupAction == EEditMeshPolygonsToolActions::InsertEdgeLoop )
+			{
+				bCanBuild = (TopologyType == EGeometryTopologyType::Polygroup);
+			}
+			return bCanBuild;
 		}
 	}
 	return false;
@@ -189,13 +197,9 @@ void UEditMeshPolygonsActionModeToolBuilder::InitializeNewTool(USingleTargetWith
 	// const method.
 	if (UGeometrySelectionManager* SelectionManager = SceneState.ToolManager->GetContextObjectStore()->FindContext<UGeometrySelectionManager>())
 	{
-		EGeometryTopologyType TopologyType = EGeometryTopologyType::Triangle;
-		EGeometryElementType ElementType = EGeometryElementType::Face;
-		int NumTargets;
-		bool bIsEmpty = false;
-		SelectionManager->GetActiveSelectionInfo(TopologyType, ElementType, NumTargets, bIsEmpty);
-
-		if (TopologyType == EGeometryTopologyType::Triangle)
+		// Note that we don't use GetActiveSelectionInfo here because that defaults to Triangle in the
+		// None/Object selection case and we want this tool to default to polygroup.
+		if (SelectionManager->GetMeshTopologyMode() == UGeometrySelectionManager::EMeshTopologyMode::Triangle)
 		{
 			if (UEditMeshPolygonsTool* EditPolygonsTool = Cast<UEditMeshPolygonsTool>(Tool))
 			{
