MeshModelingTools: Various collapse, weld, and bridge improvements in PolyEd and TriEd:
- Allow edge collapse to work on polygroup edges.
- Allow weld and bridge to work on sequences of polygroup edges.
- Allow all three to start with a vertex selection.
- Make all three set their selection appropriately on completion.
- Allow collapse to be a lot more permissive about what is allowed to be collapsed.
- Prevent WeldEdgeSequence from incorrectly failing due to edge movements that it does during its own weld sequence.

diff --git a/Engine/Plugins/Runtime/GeometryProcessing/Source/DynamicMesh/Private/GroupTopology.cpp b/Engine/Plugins/Runtime/GeometryProcessing/Source/DynamicMesh/Private/GroupTopology.cpp
--- a/Engine/Plugins/Runtime/GeometryProcessing/Source/DynamicMesh/Private/GroupTopology.cpp
+++ b/Engine/Plugins/Runtime/GeometryProcessing/Source/DynamicMesh/Private/GroupTopology.cpp
@@ -490,8 +490,18 @@ void FGroupTopology::FindCornerNbrGroups(const TArray<int>& CornerIDs, TArray<in
 
 
 void FGroupTopology::FindCornerNbrEdges(int CornerID, TArray<int>& EdgesOut) const
+{
+	ForCornerNbrEdges(CornerID, [&EdgesOut](int32 EdgeID)
+	{
+		EdgesOut.Add(EdgeID);
+		return true;
+	});
+}
+
+void UE::Geometry::FGroupTopology::ForCornerNbrEdges(int CornerID, TFunctionRef<bool(int32 EdgeId)>ReturnTrueToContinue) const
 {
 	check(CornerID >= 0 && CornerID < Corners.Num());
+	TSet<int32> ProcessedEdges;
 	for (int GroupID : Corners[CornerID].NeighbourGroupIDs)
 	{
 		const FGroup* Group = FindGroupByID(GroupID);
@@ -500,13 +510,24 @@ void FGroupTopology::FindCornerNbrEdges(int CornerID, TArray<int>& EdgesOut) con
 			for (int32 EdgeID : Boundary.GroupEdges)
 			{
 				const FGroupEdge& Edge = Edges[EdgeID];
-				if (Edge.EndpointCorners.A == CornerID || Edge.EndpointCorners.B == CornerID)
+				if (Edge.EndpointCorners.A != CornerID && Edge.EndpointCorners.B != CornerID)
 				{
-					EdgesOut.AddUnique(EdgeID);
+					continue;
+				}
+
+				bool bAlreadyProcessed = false;
+				ProcessedEdges.Add(EdgeID, &bAlreadyProcessed);
+				if (bAlreadyProcessed)
+				{
+					continue;
+				}
+
+				if (!ReturnTrueToContinue(EdgeID))
+				{
+					return;
 				}
 			}
 		}
-
 	}
 }
 
diff --git a/Engine/Plugins/Runtime/GeometryProcessing/Source/DynamicMesh/Private/Operations/WeldEdgeSequence.cpp b/Engine/Plugins/Runtime/GeometryProcessing/Source/DynamicMesh/Private/Operations/WeldEdgeSequence.cpp
--- a/Engine/Plugins/Runtime/GeometryProcessing/Source/DynamicMesh/Private/Operations/WeldEdgeSequence.cpp
+++ b/Engine/Plugins/Runtime/GeometryProcessing/Source/DynamicMesh/Private/Operations/WeldEdgeSequence.cpp
@@ -47,7 +47,7 @@ FWeldEdgeSequence::EWeldResult FWeldEdgeSequence::CheckInput()
 	// Selected edges must be boundary edges
 	for (int Edge : EdgeSpanToDiscard.Edges)
 	{
-		if (Mesh->IsBoundaryEdge(Edge) == false)
+		if (!ensure(Mesh->IsEdge(Edge)) || !Mesh->IsBoundaryEdge(Edge))
 		{
 			return EWeldResult::Failed_EdgesNotBoundaryEdges;
 		}
@@ -55,7 +55,7 @@ FWeldEdgeSequence::EWeldResult FWeldEdgeSequence::CheckInput()
 
 	for (int Edge : EdgeSpanToKeep.Edges)
 	{
-		if (Mesh->IsBoundaryEdge(Edge) == false)
+		if (!ensure(Mesh->IsEdge(Edge)) || !Mesh->IsBoundaryEdge(Edge))
 		{
 			return EWeldResult::Failed_EdgesNotBoundaryEdges;
 		}
@@ -191,7 +191,13 @@ FWeldEdgeSequence::EWeldResult FWeldEdgeSequence::WeldEdgeSequence()
 		if (Mesh->IsEdge(EdgeA) && Mesh->IsEdge(EdgeB))
 		{
 			FDynamicMesh3::FMergeEdgesInfo MergeInfo;
-			EMeshResult Result = Mesh->MergeEdges(EdgeB, EdgeA, MergeInfo);
+			EMeshResult Result = Mesh->MergeEdges(EdgeB, EdgeA, MergeInfo, 
+				// This is bCheckValidOrientation, which, roughly speaking, checks to see if the direction
+				//  of EdgeB in world space approximately flips in order to match the winding order of its new
+				//  neighbor. We don't want this, in part because our edges temoporarily move anyway as we
+				//  go along and weld them, potentially temporarily flipping. We always weld according to
+				//  winding order, irregardless of original edge position in world space. 
+				false);
 			if (Result != EMeshResult::Ok)
 			{
 				if (Result == EMeshResult::Failed_InvalidNeighbourhood && bAllowFailedMerge == true)
diff --git a/Engine/Plugins/Runtime/GeometryProcessing/Source/DynamicMesh/Public/GroupTopology.h b/Engine/Plugins/Runtime/GeometryProcessing/Source/DynamicMesh/Public/GroupTopology.h
--- a/Engine/Plugins/Runtime/GeometryProcessing/Source/DynamicMesh/Public/GroupTopology.h
+++ b/Engine/Plugins/Runtime/GeometryProcessing/Source/DynamicMesh/Public/GroupTopology.h
@@ -345,6 +345,13 @@ public:
 	/** Add all the groups connected to the given Corners to the GroupsOut list */
 	void FindCornerNbrGroups(const TArray<int>& CornerIDs, TArray<int>& GroupsOut) const;
 
+	/**
+	 * Helper that iterates over the group edges neighboring a particular corner.
+	 * @param ReturnTrueToContinue Function to call on all the group edges, which should return
+	 *  true if it wants iteration to continue, and false if it wants to exit early.
+	 */
+	void ForCornerNbrEdges(int CornerID, TFunctionRef<bool(int32 EdgeID)> ReturnTrueToContinue) const;
+
 	/** Add all the Edges connected to the given Corner to the EdgesOut list. This is somewhat expensive as there is no connection from Corners to Edges, must iterate over adjacent faces */
 	void FindCornerNbrEdges(int CornerID, TArray<int>& EdgesOut) const;
 	/** Add all the Corners connected to the given Corner to the EdgesOut list. This is somewhat expensive as there is no connection from Corners to Corners, must iterate over adjacent faces/edges */
diff --git a/Engine/Plugins/Runtime/GeometryProcessing/Source/DynamicMesh/Public/Operations/WeldEdgeSequence.h b/Engine/Plugins/Runtime/GeometryProcessing/Source/DynamicMesh/Public/Operations/WeldEdgeSequence.h
--- a/Engine/Plugins/Runtime/GeometryProcessing/Source/DynamicMesh/Public/Operations/WeldEdgeSequence.h
+++ b/Engine/Plugins/Runtime/GeometryProcessing/Source/DynamicMesh/Public/Operations/WeldEdgeSequence.h
@@ -11,11 +11,14 @@ namespace Geometry
 class FDynamicMesh3;
 
 /**
-* Weld a pair of group edges.
-* 
-* User can optionally allow triangle deletion which handles cases
-* where the group edges are connected by an edge at the end points.
-*/
+ * Weld a pair of group edges. The two input spans must be boundary spans. Their
+ *  position in the world is irrelevant, as the welding will always be done to
+ *  preserve appropriate triangle winding, i.e. so that the newly welded triangles
+ *  are not flipped relative to their neighbor.
+ * 
+ * User can optionally allow triangle deletion which handles cases
+ *  where the group edges are connected by an edge at the end points.
+ */
 
 class DYNAMICMESH_API FWeldEdgeSequence
 {
diff --git a/Engine/Plugins/Runtime/MeshModelingToolset/Source/MeshModelingTools/Private/EditMeshPolygonsTool.cpp b/Engine/Plugins/Runtime/MeshModelingToolset/Source/MeshModelingTools/Private/EditMeshPolygonsTool.cpp
--- a/Engine/Plugins/Runtime/MeshModelingToolset/Source/MeshModelingTools/Private/EditMeshPolygonsTool.cpp
+++ b/Engine/Plugins/Runtime/MeshModelingToolset/Source/MeshModelingTools/Private/EditMeshPolygonsTool.cpp
@@ -3,6 +3,7 @@
 #include "EditMeshPolygonsTool.h"
 
 #include "Algo/ForEach.h"
+#include "Algo/Reverse.h"
 #include "BaseBehaviors/SingleClickBehavior.h"
 #include "BaseGizmos/CombinedTransformGizmo.h"
 #include "BaseGizmos/TransformGizmoUtil.h"
@@ -116,6 +117,262 @@ namespace EditMeshPolygonsToolLocals
 		TSet<int32> Before;
 		TSet<int32> After;
 	};
+
+	/**
+	 * Creates a group edge selection out of a group corner selection by selecting 
+	 *  those edges whose endpoints are BOTH selected.
+	 */
+	void ConvertCornerSelectionToGroupEdgeSelection(const FGroupTopology* Topology, const TSet<int32>& CornerIDs, TSet<int32>& GroupEdgeIDs)
+	{
+		for (int32 CornerID : CornerIDs)
+		{
+			Topology->ForCornerNbrEdges(CornerID, [CornerID, &CornerIDs, &GroupEdgeIDs, Topology](int32 EdgeID)
+			{
+				if (CornerIDs.Contains(Topology->Edges[EdgeID].EndpointCorners.A)
+					&& CornerIDs.Contains(Topology->Edges[EdgeID].EndpointCorners.B))
+				{
+					GroupEdgeIDs.Add(EdgeID);					
+				}
+				return true;
+			});
+		}
+	}
+
+	// TODO: Note that so far, simply converting our selection sets to arrays via set.Array() has
+	//  been sufficient to get a selection order, but that only works due to TSet storing its
+	//  elements in a sparse array, and seems likely to break depending on how the set is updated.
+	//  For now this is good enough, but we may need to store selection order some other way someday.
+	/** 
+	 * Attempts to link boundary edges together to create either two separate boundaries, or
+	 *   one boundary loop.
+	 * 
+	 * @param GroupEdgesIn Input edges. The order of the array affects which component ends up in 
+	 *  GroupEdgesAOut, and the output of bShouldReverseAForIteration. 
+	 * @param GroupEdgesAOut This will always have the first edge in GroupEdgesIn. If the result was
+	 *  a loop, it will be the only one with edges.
+	 * @param GroupEdgesBOut The other boundary, if result was not a loop.
+	 * @param bShouldReverseAForIteration If iterating pairwise across the two groups, one of the
+	 *  arrays needs reversing, since the boundary orientation will orient the sequences in opposite
+	 *  directions. bShouldReverseAForIteration says that this array should be GroupEdgesAOut rather
+	 *  than GroupEdgesBOut based on the selection order of the longer sequence.
+	 * @return true if the edges were able to be partitioned either into one boundary loop, 
+	 *  or two separate boundaries.
+	 */
+	bool LinkBoundaryGroupEdges(const FGroupTopology* Topology, const FDynamicMesh3* Mesh,
+		const TArray<int32>& GroupEdgesIn, TArray<int32>& GroupEdgesAOut, TArray<int32>& GroupEdgesBOut, 
+		bool& bShouldReverseAForIteration)
+	{
+		if (GroupEdgesIn.IsEmpty() || !Topology || !Mesh)
+		{
+			return false;
+		}
+
+		bShouldReverseAForIteration = false;
+		if (GroupEdgesIn.Num() == 1)
+		{
+			GroupEdgesAOut.Add(GroupEdgesIn[0]);
+			return Topology->IsIsolatedLoop(GroupEdgesAOut[0]);
+		}
+
+		for (int32 GroupEdgeID : GroupEdgesIn)
+		{
+			if (Topology->IsIsolatedLoop(GroupEdgeID) || !Topology->IsBoundaryEdge(GroupEdgeID))
+			{
+				return false;
+			}
+		}
+
+		if (GroupEdgesIn.Num() == 2)
+		{
+			GroupEdgesAOut.Add(GroupEdgesIn[0]);
+			GroupEdgesBOut.Add(GroupEdgesIn[1]);
+			return true;
+		}
+
+		// Build a graph through start/end vids of the edges.
+
+		// GroupID to start vid and end vid pair
+		TMap<int32, FIndex2i> EdgeToStartEnd;
+		// start/end vid to edge id. The bool is true if start.
+		TMap<TPair<int32, bool>, int32> StartEndToEdge;
+		for (int32 GroupEdgeID : GroupEdgesIn)
+		{
+			const FEdgeSpan& Span = Topology->Edges[GroupEdgeID].Span;
+			FIndex2i OrientedEdgeVids = Mesh->GetOrientedBoundaryEdgeV(Span.Edges[0]);
+			bool bReversed = OrientedEdgeVids.A != Span.Vertices[0];
+
+			TPair<int32, bool> KeyForFirst(Span.Vertices[0], !bReversed);
+			TPair<int32, bool> KeyForLast(Span.Vertices.Last(), bReversed);
+			if (StartEndToEdge.Contains(KeyForFirst) || StartEndToEdge.Contains(KeyForLast))
+			{
+				// This means that a vertex was the end or start point for more than one edge,
+				//  i.e. there was a branch. So, there is ambiguity in how to partition.
+				return false;
+			}
+			StartEndToEdge.Add(KeyForFirst, GroupEdgeID);
+			StartEndToEdge.Add(KeyForLast, GroupEdgeID);
+			EdgeToStartEnd.Add(GroupEdgeID, bReversed ? FIndex2i(Span.Vertices.Last(), Span.Vertices[0])
+				: FIndex2i(Span.Vertices[0], Span.Vertices.Last()));
+		}
+
+		TSet<int32> PartitionedEdges;
+		// Helper that gets all the connected edges from a given edge, in order.
+		auto GetEdgeSequence = [&PartitionedEdges, &EdgeToStartEnd, &StartEndToEdge](int32 StartEdge, TArray<int32>& EdgeSequenceOut)
+		{
+			bool bAlreadyProcessed = false;
+			PartitionedEdges.Add(StartEdge, &bAlreadyProcessed);
+			if (bAlreadyProcessed) return;
+
+			// Go backwards and forwards through the graph to get our adjoining edges.
+			// We'll start by going backwards (we'll reverse this output in a bit so that it is in the correct order)
+			FIndex2i StartEnd = EdgeToStartEnd[StartEdge];
+			int32 CurrentEndpoint = StartEnd.A;
+			while (int32* CurrentEdge = StartEndToEdge.Find(TPair<int32, bool>(CurrentEndpoint, false)))
+			{
+				PartitionedEdges.Add(*CurrentEdge, &bAlreadyProcessed);
+				if (bAlreadyProcessed) break;
+
+				EdgeSequenceOut.Add(*CurrentEdge);
+				CurrentEndpoint = EdgeToStartEnd[*CurrentEdge].A;
+			}
+			Algo::Reverse(EdgeSequenceOut);
+
+			// Now that we have the preceding edges, add this one and search forwards
+			EdgeSequenceOut.Add(StartEdge);
+			CurrentEndpoint = StartEnd.B;
+			while (int32* CurrentEdge = StartEndToEdge.Find(TPair<int32, bool>(CurrentEndpoint, true)))
+			{
+				PartitionedEdges.Add(*CurrentEdge, &bAlreadyProcessed);
+				if (bAlreadyProcessed) break;
+
+				EdgeSequenceOut.Add(*CurrentEdge);
+				CurrentEndpoint = EdgeToStartEnd[*CurrentEdge].B;
+			}
+		};
+
+		for (int32 GroupEdgeID : GroupEdgesIn)
+		{
+			if (PartitionedEdges.Contains(GroupEdgeID))
+			{
+				continue;
+			}
+			if (GroupEdgesAOut.IsEmpty())
+			{
+				GetEdgeSequence(GroupEdgeID, GroupEdgesAOut);
+			}
+			else if (GroupEdgesBOut.IsEmpty())
+			{
+				GetEdgeSequence(GroupEdgeID, GroupEdgesBOut);
+			}
+			else
+			{
+				// Had a third connected component
+				return false;
+			}
+		}
+
+		if (GroupEdgesBOut.IsEmpty())
+		{
+			// Make sure that the output result is a loop
+			return !GroupEdgesAOut.IsEmpty()
+				&& EdgeToStartEnd[GroupEdgesAOut[0]].A == EdgeToStartEnd[GroupEdgesAOut.Last()].B;
+		}
+
+		// Figure out the preferred iteration order based on the longer subsequence.
+		// The issue is this: if we have edges A, B, and C in EdgesA and corresponding edges 1, 2, and 3 on the other
+		//  side, then the latter will be ordered 3, 2, 1 in EdgesB according to triangle orientations, and as long
+		//  as we reverse one group or the other, pairwise iteration will give us the correct pairings (through which
+		//  we will iterate either as A1, B2, C3, or as C3, B2, A1, depending on whether we reverse EdgesB or EdgesA,
+		//  respectively). However, what happens if we have a mismatched number of edges, e.g. edge D after C? We will 
+		//  group the extra edge(s) with the last edge in the shorter sequence, so iteration order matters: either we 
+		//  end up grouping CD with 3 if we reverse GroupB, or we end up grouping AB with 1 if we reverse GroupA.
+		// We choose to decide based on which of the ends of the longer sequence was selected last- this is the direction
+		//  that we interpret the user wanting to iterate in. E.g., if user selected D after they selected A, then we decide
+		//  that the iteration order should be A1, B2, CD3.
+		if (ensure(!GroupEdgesAOut.IsEmpty()))
+		{
+			if (GroupEdgesAOut.Num() > GroupEdgesBOut.Num())
+			{
+				// Reverse if the later edge in the sequence was selected earlier the first
+				bShouldReverseAForIteration = GroupEdgesIn.IndexOfByKey(GroupEdgesAOut.Last()) < GroupEdgesIn.IndexOfByKey(GroupEdgesAOut[0]);
+			}
+			else if (GroupEdgesBOut.Num() > GroupEdgesAOut.Num())
+			{
+				// The comparison is backwards here because bShouldReverseAForIteration needs to be the opposite of 
+				//  whether we should reverse EdgesB.
+				bShouldReverseAForIteration = GroupEdgesIn.IndexOfByKey(GroupEdgesBOut[0]) < GroupEdgesIn.IndexOfByKey(GroupEdgesBOut.Last());
+			}
+		}
+		
+		// Make sure that both partitions are not loops
+		return !GroupEdgesAOut.IsEmpty()
+			&& EdgeToStartEnd[GroupEdgesAOut[0]].A != EdgeToStartEnd[GroupEdgesAOut.Last()].B
+			&& EdgeToStartEnd[GroupEdgesBOut[0]].A != EdgeToStartEnd[GroupEdgesBOut.Last()].B;
+	}
+
+	// Helper to share the retriangulation code
+	int32 RetriangulateGroups(FDynamicMesh3* Mesh, FGroupTopology* Topology, TSet<int32> GroupIDs, FDynamicMeshChangeTracker& ChangeTracker)
+	{
+		FDynamicMeshEditor Editor(Mesh);
+		int32 NumCompleted = 0;
+		for (int32 GroupID : GroupIDs)
+		{
+			const TArray<int32>& Triangles = Topology->GetGroupTriangles(GroupID);
+			ChangeTracker.SaveTriangles(Triangles, true);
+			FMeshRegionBoundaryLoops RegionLoops(Mesh, Triangles, true);
+			if (!RegionLoops.bFailed && RegionLoops.Loops.Num() == 1 && Triangles.Num() > 1)
+			{
+				TArray<FMeshRegionBoundaryLoops::VidOverlayMap<FVector2f>> VidUVMaps;
+				if (Mesh->HasAttributes())
+				{
+					const FDynamicMeshAttributeSet* Attributes = Mesh->Attributes();
+					for (int i = 0; i < Attributes->NumUVLayers(); ++i)
+					{
+						VidUVMaps.Emplace();
+						RegionLoops.GetLoopOverlayMap(RegionLoops.Loops[0], *Attributes->GetUVLayer(i), VidUVMaps.Last());
+					}
+				}
+
+				// We don't want to remove isolated vertices while removing triangles because we don't
+				// want to throw away boundary verts. However, this means that we'll have to go back
+				// through these vertices later to throw away isolated internal verts.
+				TArray<int32> OldVertices;
+				UE::Geometry::TriangleToVertexIDs(Mesh, Triangles, OldVertices);
+				Editor.RemoveTriangles(Topology->GetGroupTriangles(GroupID), false);
+
+				RegionLoops.Loops[0].Reverse();
+				FSimpleHoleFiller Filler(Mesh, RegionLoops.Loops[0]);
+				Filler.FillType = FSimpleHoleFiller::EFillType::PolygonEarClipping;
+				Filler.Fill(GroupID);
+
+				// Throw away any of the old verts that are still isolated (they were in the interior of the group)
+				Algo::ForEachIf(OldVertices,
+					[Mesh](int32 Vid)
+					{
+						return !Mesh->IsReferencedVertex(Vid);
+					},
+					[Mesh](int32 Vid)
+					{
+						checkSlow(!Mesh->IsReferencedVertex(Vid));
+						constexpr bool bPreserveManifold = false;
+						Mesh->RemoveVertex(Vid, bPreserveManifold);
+					});
+
+				if (Mesh->HasAttributes())
+				{
+					const FDynamicMeshAttributeSet* Attributes = Mesh->Attributes();
+					for (int i = 0; i < Attributes->NumUVLayers(); ++i)
+					{
+						RegionLoops.UpdateLoopOverlayMapValidity(VidUVMaps[i], *Attributes->GetUVLayer(i));
+					}
+					Filler.UpdateAttributes(VidUVMaps);
+				}
+
+				NumCompleted++;
+			}
+		}
+		return NumCompleted;
+	}//end RetriangulateGroups
 }
 
 /*
@@ -1402,7 +1659,7 @@ void UEditMeshPolygonsTool::OnTick(float DeltaTime)
 			ApplySplitSingleEdge();
 			break;
 		case EEditMeshPolygonsToolActions::CollapseSingleEdge:
-			ApplyCollapseSingleEdge();
+			ApplyCollapseEdge();
 			break;
 		case EEditMeshPolygonsToolActions::FlipSingleEdge:
 			ApplyFlipSingleEdge();
@@ -1800,75 +2057,20 @@ void UEditMeshPolygonsTool::ApplyFlipNormals()
 
 void UEditMeshPolygonsTool::ApplyRetriangulate()
 {
+	using namespace EditMeshPolygonsToolLocals;
 	if (BeginMeshFaceEditChange() == false)
 	{
 		GetToolManager()->DisplayMessage( LOCTEXT("OnRetriangulateFailed", "Cannot Retriangulate Current Selection"), EToolMessageLevel::UserWarning);
 		return;
 	}
 
-	int32 nCompleted = 0;
 	FDynamicMesh3* Mesh = CurrentMesh.Get();
 	FDynamicMeshChangeTracker ChangeTracker(Mesh);
 	ChangeTracker.BeginChange();
-	FDynamicMeshEditor Editor(Mesh);
 	FGroupTopologySelection ActiveSelection = SelectionMechanic->GetActiveSelection();
-	for (int32 GroupID : ActiveSelection.SelectedGroupIDs)
-	{
-		const TArray<int32>& Triangles = Topology->GetGroupTriangles(GroupID);
-		ChangeTracker.SaveTriangles(Triangles, true);
-		FMeshRegionBoundaryLoops RegionLoops(Mesh, Triangles, true);
-		if (!RegionLoops.bFailed && RegionLoops.Loops.Num() == 1 && Triangles.Num() > 1)
-		{
-			TArray<FMeshRegionBoundaryLoops::VidOverlayMap<FVector2f>> VidUVMaps;
-			if (Mesh->HasAttributes())
-			{
-				const FDynamicMeshAttributeSet* Attributes = Mesh->Attributes();
-				for (int i = 0; i < Attributes->NumUVLayers(); ++i)
-				{
-					VidUVMaps.Emplace();
-					RegionLoops.GetLoopOverlayMap(RegionLoops.Loops[0], *Attributes->GetUVLayer(i), VidUVMaps.Last());
-				}
-			}
-
-			// We don't want to remove isolated vertices while removing triangles because we don't
-			// want to throw away boundary verts. However, this means that we'll have to go back
-			// through these vertices later to throw away isolated internal verts.
-			TArray<int32> OldVertices;
-			UE::Geometry::TriangleToVertexIDs(Mesh, Triangles, OldVertices);
-			Editor.RemoveTriangles(Topology->GetGroupTriangles(GroupID), false);
-
-			RegionLoops.Loops[0].Reverse();
-			FSimpleHoleFiller Filler(Mesh, RegionLoops.Loops[0]);
-			Filler.FillType = FSimpleHoleFiller::EFillType::PolygonEarClipping;
-			Filler.Fill(GroupID);
-
-			// Throw away any of the old verts that are still isolated (they were in the interior of the group)
-			Algo::ForEachIf(OldVertices, 
-				[Mesh](int32 Vid) 
-			{ 
-				return !Mesh->IsReferencedVertex(Vid); 
-			},
-				[Mesh](int32 Vid) 
-			{
-				checkSlow(!Mesh->IsReferencedVertex(Vid));
-				constexpr bool bPreserveManifold = false;
-				Mesh->RemoveVertex(Vid, bPreserveManifold);
-			}
-			);
-
-			if (Mesh->HasAttributes())
-			{
-				const FDynamicMeshAttributeSet* Attributes = Mesh->Attributes();
-				for (int i = 0; i < Attributes->NumUVLayers(); ++i)
-				{
-					RegionLoops.UpdateLoopOverlayMapValidity(VidUVMaps[i], *Attributes->GetUVLayer(i));
-				}
-				Filler.UpdateAttributes(VidUVMaps);
-			}
+	
+	int32 nCompleted = RetriangulateGroups(Mesh, Topology.Get(), ActiveSelection.SelectedGroupIDs, ChangeTracker);
 
-			nCompleted++;
-		}
-	}
 	if (nCompleted != ActiveSelection.SelectedGroupIDs.Num())
 	{
 		GetToolManager()->DisplayMessage(LOCTEXT("OnRetriangulateFailures", "Some faces could not be retriangulated"), EToolMessageLevel::UserWarning);
@@ -2038,148 +2240,196 @@ void UEditMeshPolygonsTool::ApplyDuplicate()
 
 
 
-
-void UEditMeshPolygonsTool::ApplyCollapseEdge()
+// Deprecated
+void UEditMeshPolygonsTool::ApplyCollapseSingleEdge()
 {
-	// AAAHHH cannot do because of overlays!
-	return;
-#if 0
-	if (SelectionMechanic->GetActiveSelection().SelectedEdgeIDs.Num() != 1 || BeginMeshEdgeEditChange() == false)
-	{
-		GetToolManager()->DisplayMessage(
-			LOCTEXT("OnEdgeColllapseFailed", "Cannot Collapse current selection"),
-			EToolMessageLevel::UserWarning);
-		return;
-	}
-
-	FDynamicMesh3* Mesh = CurrentMesh.Get();
-
-	FDynamicMeshChangeTracker ChangeTracker(Mesh);
-	ChangeTracker.BeginChange();
-	//const TArray<int32>& EdgeIDs = ActiveEdgeSelection[0].EdgeIDs;
-	//for (int32 eid : EdgeIDs)
-	//{
-	//	if (Mesh->IsEdge(eid))
-	//	{
-	//		FIndex2i EdgeVerts = Mesh->GetEdgeV(eid);
-	//		ChangeTracker.SaveVertexOneRingTriangles(EdgeVerts.A, true);
-	//		ChangeTracker.SaveVertexOneRingTriangles(EdgeVerts.B, true);
-	//		FDynamicMesh3::FEdgeCollapseInfo CollapseInfo;
-	//		Mesh->CollapseEdge()
-	//	}
-	//}
-
-	// emit undo
-	FGroupTopologySelection NewSelection;
-	EmitCurrentMeshChangeAndUpdate(LOCTEXT("PolyMeshEdgeCollapseChange", "Collapse"),
-		ChangeTracker.EndChange(), NewSelection);
-#endif
+	ApplyCollapseEdge();
 }
 
 
 
 void UEditMeshPolygonsTool::ApplyWeldEdges()
 {
-	if (SelectionMechanic->GetActiveSelection().SelectedEdgeIDs.Num() != 2)
+	using namespace EditMeshPolygonsToolLocals;
+
+	FGroupTopologySelection CurrentSelection = SelectionMechanic->GetActiveSelection();
+	
+	TSet<int32> GroupEdges;
+	if (!CurrentSelection.SelectedCornerIDs.IsEmpty())
+	{
+		ConvertCornerSelectionToGroupEdgeSelection(Topology.Get(), CurrentSelection.SelectedCornerIDs, GroupEdges);
+	}
+	else
+	{
+		GroupEdges = CurrentSelection.SelectedEdgeIDs;
+	}
+	
+	if (GroupEdges.Num() < 2)
 	{
 		GetToolManager()->DisplayMessage(
-			LOCTEXT("OnWeldEdgesFailedEdgeCount", "Cannot Weld current selection, selection must be exactly 2 edges."),
+			LOCTEXT("OnWeldEdgesFailedTooFew", "Cannot Weld current selection, selection must be at least 2 edges."),
 			EToolMessageLevel::UserWarning);
 		return;
 	}
 
 	FDynamicMesh3* Mesh = CurrentMesh.Get();
-	FDynamicMesh3 MeshCopy(*Mesh);				// We are going to operate on this copy and, if successful, copy the changes back into Mesh
-	
-	FGroupTopologySelection CurrentSelection = SelectionMechanic->GetActiveSelection();
-	TArray<int32> SelectedEdgeIDs = CurrentSelection.SelectedEdgeIDs.Array();
-	FEdgeSpan& SpanA = Topology->Edges[SelectedEdgeIDs[0]].Span;
-	FEdgeSpan& SpanB = Topology->Edges[SelectedEdgeIDs[1]].Span;
-	
-	if (SpanA.Vertices[0] == SpanA.Vertices.Last() || SpanB.Vertices[0] == SpanB.Vertices.Last())
+	TArray<int32> GroupEdgesA;
+	TArray<int32> GroupEdgesB;
+	bool bShouldReverseA = false;
+	if (!LinkBoundaryGroupEdges(Topology.Get(), Mesh, GroupEdges.Array(), GroupEdgesA, GroupEdgesB, bShouldReverseA)
+		// We don't allow a single loop
+		|| GroupEdgesB.IsEmpty())
 	{
 		GetToolManager()->DisplayMessage(
-			LOCTEXT("OnWeldEdgesFailedEdgesAreLoops", "Cannot Weld current selection, selected edges must not be loops."),
+			LOCTEXT("OnWeldEdgesFailedEdgeCount", "Cannot Weld current selection, selection could not be partitioned into two non-loop open-boundary sequences."),
 			EToolMessageLevel::UserWarning);
-		return;
+			return;
 	}
 
+	// The two sequences are given in boundary orientation, which will be in the opposite direction.
+	//  We reverse one of them so we can do our pairwise welding in the proper order.
+	Algo::Reverse(bShouldReverseA ? GroupEdgesA : GroupEdgesB);
+
+	// At the end, we're going to select the edges we keep.
+	// Note: this requires eids to stay constant across the mesh copy call further below. That seems to be the
+	//  case for dynamic mesh copies, otherwise we'd need to use tri-subidx pairs.
+	TArray<int32> KeptEids;
+	for (int32 GroupEdgeID : GroupEdgesB)
+	{
+		KeptEids.Append(Topology->GetGroupEdgeEdges(GroupEdgeID));
+	}
+
+	FText TransactionName = LOCTEXT("PolyMeshWeldEdgeChange", "Weld Edges");
+	GetToolManager()->BeginUndoTransaction(TransactionName);
 	FDynamicMeshChangeTracker ChangeTracker(Mesh);
 	ChangeTracker.BeginChange();
 
-	// Save one ring tri's for vertices along first edge
-	for (int Vert : SpanA.Vertices)
+	bool bAllSucceeded = true;
+	bool bHaveSeam = false;
+
+	auto WeldGroupEdges = [&ChangeTracker, Mesh, &bAllSucceeded, &bHaveSeam](FEdgeSpan& SpanA, FEdgeSpan& SpanB)
 	{
-		ChangeTracker.SaveVertexOneRingTriangles(Vert, true);
-	}
+		// Save one ring tri's for vertices along both edges. The kept edge is necessary
+		//  because we might be splitting its triangles if needed.
+		for (int32 Vid : SpanA.Vertices)
+		{
+			if (!ensure(Mesh->IsVertex(Vid)))
+			{
+				return false;
+			}
+			Mesh->EnumerateVertexTriangles(Vid, [&ChangeTracker](int32 Tid)
+			{
+				ChangeTracker.SaveTriangle(Tid, true);
+			});
+		}
+		for (int32 Vid : SpanB.Vertices)
+		{
+			if (!ensure(Mesh->IsVertex(Vid)))
+			{
+				return false;
+			}
+			Mesh->EnumerateVertexTriangles(Vid, [&ChangeTracker](int32 Tid)
+			{
+				ChangeTracker.SaveTriangle(Tid, true);
+			});
+		}
+
+		FWeldEdgeSequence EdgeWelder(Mesh, SpanA, SpanB);
+		EdgeWelder.bAllowIntermediateTriangleDeletion = true;
+		EdgeWelder.bAllowFailedMerge = true;
+
+		FWeldEdgeSequence::EWeldResult Result = EdgeWelder.Weld();
+		if (EdgeWelder.UnmergedEdgePairsOut.Num() != 0)
+		{
+			bHaveSeam = true;
+		}
+		if (Result != FWeldEdgeSequence::EWeldResult::Ok)
+		{
+			bAllSucceeded = false;
+			return false;
+		}
+		return true;
+	};
 
-	// Save one ring tri's for vertices along second edge
-	for (int Vert : SpanB.Vertices)
+	// Some number of edges will be welded normally. However if we have a mismatched number,
+	//  the last edge of the smaller sequence will have to be welded to the remainder of edges.
+	int32 NumMatched = GroupEdgesA.Num() == GroupEdgesB.Num() ? GroupEdgesA.Num()
+		: FMath::Min(GroupEdgesA.Num(), GroupEdgesB.Num()) - 1;
+	for (int32 i = 0; i < NumMatched; ++i)
 	{
-		ChangeTracker.SaveVertexOneRingTriangles(Vert, true);
+		FEdgeSpan SpanA;
+		SpanA.InitializeFromEdges(Mesh, Topology->Edges[GroupEdgesA[i]].Span.Edges);
+		FEdgeSpan& SpanB = Topology->Edges[GroupEdgesB[i]].Span;
+		WeldGroupEdges(SpanA, SpanB);
 	}
 
-	FWeldEdgeSequence EdgeWelder(&MeshCopy, SpanA, SpanB);
-	EdgeWelder.bAllowIntermediateTriangleDeletion = true;
-	EdgeWelder.bAllowFailedMerge = true;
-	
-	FWeldEdgeSequence::EWeldResult Result = EdgeWelder.Weld();
-	if (Result != FWeldEdgeSequence::EWeldResult::Ok)
+	// If there was a mismatched number of edges, we have set NumMatched to be one less than
+	//  the shorter sequence so we can weld the last edge to the remainder on the other side.
+	if (NumMatched < GroupEdgesA.Num())
 	{
-		switch (Result)
+		// Assemble our two edge sequences to weld.
+		TArray<int32> SpanAEids;
+		for (int32 i = 0; i < GroupEdgesA.Num() - NumMatched; ++i)
 		{
-		case FWeldEdgeSequence::EWeldResult::Failed_EdgesNotBoundaryEdges:
-			GetToolManager()->DisplayMessage(
-				LOCTEXT("OnWeldEdgesFailedBoundary", "Cannot Weld current selection, selected edges must be boundary edges."),
-				EToolMessageLevel::UserWarning);
-			break;
+			int32 Index = bShouldReverseA ? GroupEdgesA.Num() - 1 - i : NumMatched + i;
+			FEdgeSpan Span = Topology->Edges[GroupEdgesA[Index]].Span;
+			Span.SetCorrectOrientation();
+			SpanAEids.Append(Span.Edges);
+		}
+		TArray<int32> SpanBEids;
+		for (int32 i = 0; i < GroupEdgesB.Num() - NumMatched; ++i)
+		{
+			int32 Index = !bShouldReverseA ? GroupEdgesB.Num() - 1 - i : NumMatched + i;
+			FEdgeSpan Span = Topology->Edges[GroupEdgesB[Index]].Span;
+			Span.SetCorrectOrientation();
+			SpanBEids.Append(Span.Edges);
+		}
 
-		case FWeldEdgeSequence::EWeldResult::Failed_CannotSplitEdge:
-			GetToolManager()->DisplayMessage(
-				LOCTEXT("OnWeldEdgesFailedSplitEdge", "Cannot Weld current selection, failed to insert vertex."), 
-				EToolMessageLevel::UserWarning);
-			break;
+		FEdgeSpan SpanA;
+		SpanA.InitializeFromEdges(Mesh, SpanAEids);
+		FEdgeSpan SpanB;
+		SpanB.InitializeFromEdges(Mesh, SpanBEids);
 
-		case FWeldEdgeSequence::EWeldResult::Failed_TriangleDeletionDisabled:
-			GetToolManager()->DisplayMessage(
-				LOCTEXT("OnWeldEdgesFailedTriDeleteDisabled", "Cannot Weld current selection, deletion of edges connecting selected edges is disabled."), 
-				EToolMessageLevel::UserWarning);
-			break;
+		WeldGroupEdges(SpanA, SpanB);
+	}
 
-		case FWeldEdgeSequence::EWeldResult::Failed_CannotDeleteTriangle:
-			GetToolManager()->DisplayMessage(
-				LOCTEXT("OnWeldEdgesFailedTriDeleteFailed", "Cannot Weld current selection, failed to delete edge connecting selected edges."), 
-				EToolMessageLevel::UserWarning);
-			break;
+	EmitCurrentMeshChangeAndUpdate(TransactionName, ChangeTracker.EndChange(), FGroupTopologySelection());
 
-		case FWeldEdgeSequence::EWeldResult::Failed_Other:
-		default:
-			GetToolManager()->DisplayMessage(
-				LOCTEXT("OnWeldEdgesFailedOther", "Cannot Weld current selection, bad geometry."), 
-				EToolMessageLevel::UserWarning);
-			break;
+	// Now that the topology is updated, set the new selection
+	FGroupTopologySelection NewSelection;
+	TSet<int32> SelectedEids;
+	for (int32 Eid : KeptEids)
+	{
+		if (Mesh->IsEdge(Eid) && !SelectedEids.Contains(Eid))
+		{
+			int32 GroupEdgeID = Topology->FindGroupEdgeID(Eid);
+			if (GroupEdgeID != IndexConstants::InvalidID)
+			{
+				NewSelection.SelectedEdgeIDs.Add(GroupEdgeID);
+				SelectedEids.Append(Topology->GetGroupEdgeEdges(GroupEdgeID));
+			}
 		}
-
-		return;
 	}
-	else
+	// Seems possible to end up with an empty selection if we welded edges of the same group,
+	//  so the new edge is not a group boundary.
+	if (!NewSelection.IsEmpty())
 	{
-		// On success, apply the result by copying over the existing mesh
-		*Mesh = MeshCopy;
+		SelectionMechanic->SetSelection(NewSelection);
+	}
 
-		if (EdgeWelder.UnmergedEdgePairsOut.Num() != 0)
-		{
-			GetToolManager()->DisplayMessage(
-				LOCTEXT("OnWeldEdgesCompletedSeamsRemain", "Warning: welding incomplete because it would create "
-					"invalid geometry (attached non manifold edge or duplicate triangle). Seam still exists at weld "
-					"location. Modify attached triangles and retry, or undo."),
-				EToolMessageLevel::UserWarning);
-		}
+	if (bHaveSeam)
+	{
+		GetToolManager()->DisplayMessage(LOCTEXT("OnWeldEdgesCompletedSeamsRemain", "Warning: welding incomplete because it would create "
+			"invalid geometry (attached non manifold edge or duplicate triangle). Seam still exists at weld "
+			"location. Modify attached triangles and retry, or undo."), EToolMessageLevel::UserWarning);
+	}
+	else if (!bAllSucceeded)
+	{
+		GetToolManager()->DisplayMessage(LOCTEXT("OnWeldEdgesPartialFailure", 
+			"Warning: some edges could not be welded."), EToolMessageLevel::UserWarning);
 	}
 
-	FGroupTopologySelection NewSelection;
-	EmitCurrentMeshChangeAndUpdate(LOCTEXT("PolyMeshWeldEdgeChange", "Weld Edges"),
-		ChangeTracker.EndChange(), NewSelection);
+	GetToolManager()->EndUndoTransaction();
 }
 
 void UEditMeshPolygonsTool::ApplyStraightenEdges()
@@ -2369,110 +2619,202 @@ void UEditMeshPolygonsTool::ApplyFillHole()
 
 void UEditMeshPolygonsTool::ApplyBridgeEdges()
 {
+	using namespace EditMeshPolygonsToolLocals;
+
 	const FText BridgeFailMessage = LOCTEXT("OnEdgeBridgeFailed", "Cannot Bridge current selection");
 
-	if (SelectionMechanic->GetActiveSelection().SelectedEdgeIDs.Num() != 2 || BeginMeshBoundaryEdgeEditChange(false) == false)
+	FGroupTopologySelection CurrentSelection = SelectionMechanic->GetActiveSelection();
+
+	TSet<int32> GroupEdges;
+	if (!CurrentSelection.SelectedCornerIDs.IsEmpty())
 	{
-		GetToolManager()->DisplayMessage(BridgeFailMessage, EToolMessageLevel::UserWarning);
-		return;
+		ConvertCornerSelectionToGroupEdgeSelection(Topology.Get(), CurrentSelection.SelectedCornerIDs, GroupEdges);
+	}
+	else
+	{
+		GroupEdges = CurrentSelection.SelectedEdgeIDs;
 	}
 
 	FDynamicMesh3* Mesh = CurrentMesh.Get();
-	FDynamicMeshChangeTracker ChangeTracker(Mesh);
-	ChangeTracker.BeginChange();
-	FGroupTopologySelection CurrentSelection = SelectionMechanic->GetActiveSelection();
-
-	TArray<int32> LoopVertices;
-	TArray<int32> LoopEdges;
-	TArray<int32> SelectedEdgeIDs = CurrentSelection.SelectedEdgeIDs.Array();
-	
-	// I think doing this will guarantee that every edge in the span stores the vertices corresponding to the connected triangles orientation
-	FEdgeSpan& SpanA = Topology->Edges[SelectedEdgeIDs[0]].Span;
-	FEdgeSpan& SpanB = Topology->Edges[SelectedEdgeIDs[1]].Span;
-	SpanA.SetCorrectOrientation();
-	SpanB.SetCorrectOrientation();
-
-	// Disallow bridging of edge loops for now
-	if (SpanA.Vertices[0] == SpanA.Vertices.Last() || SpanB.Vertices[0] == SpanB.Vertices.Last())
-	{
-		GetToolManager()->DisplayMessage(BridgeFailMessage, EToolMessageLevel::UserWarning);
+	TArray<int32> GroupEdgesA;
+	TArray<int32> GroupEdgesB;
+	bool bShouldReverseA = false;
+	if (!LinkBoundaryGroupEdges(Topology.Get(), Mesh, GroupEdges.Array(), GroupEdgesA, GroupEdgesB, bShouldReverseA))
+	{
+		GetToolManager()->DisplayMessage(LOCTEXT("OnEdgeBridgeFailedInvalidSelection", 
+			"Cannot bridge current selection, selection could not be partitioned into a single hole or two "
+			"non-loop open-boundary sequences."), EToolMessageLevel::UserWarning);
 		return;
 	}
 
-	// Add all vertices from first edge
-	LoopVertices = SpanA.Vertices;
+	TArray<int32> TrianglesToSelect;
 
-	// If first vertex of second edge is not a duplicate of a terminating vertex of the first edge, add vertex
-	if (SpanB.Vertices[0] != SpanA.Vertices[0] && SpanB.Vertices[0] != SpanA.Vertices.Last())
+	auto BridgeEdges = [&](const TArray<int32>& LoopVids)
 	{
-		LoopVertices.Add(SpanB.Vertices[0]);
-	}
+		TArray<int32> LoopEdges;
+		FEdgeLoop::VertexLoopToEdgeLoop(Mesh, LoopVids, LoopEdges);
+		FEdgeLoop Loop(Mesh, LoopVids, LoopEdges);
 
-	// Definitely add the non-terminating vertices of second edge.
-	for (int Vertex = 1; Vertex < SpanB.Vertices.Num() - 1; ++Vertex)
-	{
-		LoopVertices.Add(SpanB.Vertices[Vertex]);
-	}
+		// We could always use the minimal hole filler, but it doesn't quite do what "bridge" would suggest when
+		// the area to be bridged is concave (across two curved-inward edges). Meanwhile simple ear clipping
+		// seems to fail in some common cases for reasons that we should investigate. For now, start with ear
+		// clipping, and revert to minimal if needed.
+		FSimpleHoleFiller SimpleHoleFiller(Mesh, Loop, FSimpleHoleFiller::EFillType::PolygonEarClipping);
+		TArray<int32> NewTriangles;
 
-	// If last vertex of second edge is not a duplicate of a terminating vertex of the first edge, add vertex
-	if (SpanB.Vertices.Last() != SpanA.Vertices[0] && SpanB.Vertices.Last() != SpanA.Vertices.Last())
-	{
-		LoopVertices.Add(SpanB.Vertices.Last());
-	}
+		// Fill the hole
+		if (!SimpleHoleFiller.Fill())
+		{
+			//Ear clipping doesn't add vertices, so don't need to delete isolated verts
+			FDynamicMeshEditor Editor(Mesh);
+			Editor.RemoveTriangles(SimpleHoleFiller.NewTriangles, false);
+
+			FMinimalHoleFiller MinimalHoleFiller(Mesh, Loop);
+
+			if (!MinimalHoleFiller.Fill())
+			{
+				Editor.RemoveTriangles(MinimalHoleFiller.NewTriangles, false);
+				GetToolManager()->DisplayMessage(BridgeFailMessage, EToolMessageLevel::UserWarning);
+				// Even though we've manually 'undone' the changes, this will still change mesh timestamps, so we need to register the mesh update
+				UpdateFromCurrentMesh(false);
+				return;
+			}
+			NewTriangles = MinimalHoleFiller.NewTriangles;
+		}
+		else 
+		{
+			NewTriangles = SimpleHoleFiller.NewTriangles;
+		}
 
-	FEdgeLoop::VertexLoopToEdgeLoop(Mesh, LoopVertices, LoopEdges);
-	FEdgeLoop Loop(Mesh, LoopVertices, LoopEdges);
+		TrianglesToSelect.Append(NewTriangles);
 
-	// We could always use the minimal hole filler, but it doesn't quite do what "bridge" would suggest when
-	// the area to be bridged is concave (across two curved-inward edges). Meanwhile simple ear clipping
-	// seems to fail in some common cases for reasons that we should investigate. For now, start with ear
-	// clipping, and revert to minimal if needed.
-	FSimpleHoleFiller SimpleHoleFiller(Mesh, Loop, FSimpleHoleFiller::EFillType::PolygonEarClipping);
-	TArray<int32> NewTriangles;
+		// Compute normals and UVs
+		if (Mesh->HasAttributes())
+		{
+			TArray<FVector3d> VertexPositions;
+			Loop.GetVertices(VertexPositions);
+			FVector3d PlaneOrigin;
+			FVector3d PlaneNormal;
+			PolygonTriangulation::ComputePolygonPlane<double>(VertexPositions, PlaneNormal, PlaneOrigin);
 
-	// Fill the hole
-	if (!SimpleHoleFiller.Fill())
-	{
-		//Ear clipping doesn't add vertices, so don't need to delete isolated verts
-		FDynamicMeshEditor Editor(Mesh);
-		Editor.RemoveTriangles(SimpleHoleFiller.NewTriangles, false);
+			FDynamicMeshEditor Editor(Mesh);
+			FFrame3d ProjectionFrame(PlaneOrigin, PlaneNormal);
+			Editor.SetTriangleNormals(NewTriangles);
+			Editor.SetTriangleUVsFromProjection(NewTriangles, ProjectionFrame, UVScaleFactor);
+		}
+	};
 
-		FMinimalHoleFiller MinimalHoleFiller(Mesh, Loop);
+	FDynamicMeshChangeTracker ChangeTracker(Mesh);
+	ChangeTracker.BeginChange();
 
-		if (!MinimalHoleFiller.Fill())
+	if (GroupEdgesB.IsEmpty())
+	{
+		// This means that there was just one big hole. Concatenate everything into one
+		//  big loop to triangulate.
+		TArray<int32> LoopVertices;
+		for (int32 i = 0; i < GroupEdgesA.Num(); ++i)
+		{
+			FEdgeSpan Span = Topology->Edges[GroupEdgesA[i]].Span;
+			Span.SetCorrectOrientation();
+			if (LoopVertices.Num() > 0 && LoopVertices.Last() == Span.Vertices[0])
+			{
+				LoopVertices.Pop();
+			}
+			LoopVertices.Append(Span.Vertices);
+		}
+		if (!ensure(LoopVertices.Num() > 0 && LoopVertices.Last() == LoopVertices[0]))
 		{
-			Editor.RemoveTriangles(MinimalHoleFiller.NewTriangles, false);
 			GetToolManager()->DisplayMessage(BridgeFailMessage, EToolMessageLevel::UserWarning);
-			// Even though we've manually 'undone' the changes, this will still change mesh timestamps, so we need to register the mesh update
-			UpdateFromCurrentMesh(false);
 			return;
 		}
-		else
+		LoopVertices.Pop();
+
+		BridgeEdges(LoopVertices);
+	}
+	else
+	{
+		// We will bridge as many pairs as we can, and then do one big "triangular" bridge for
+		//  the unmatched ones.
+
+		// The two sequences are given in boundary orientation, which will be in the opposite direction.
+		//  So, we need to process one of the sequences in reverse order.
+		Algo::Reverse(bShouldReverseA ? GroupEdgesA : GroupEdgesB);
+
+		// Keeps track of the endpoints so we can use it for the final triangular bridge
+		FIndex2i LastVids(IndexConstants::InvalidID, IndexConstants::InvalidID);
+
+		int32 NumMatched = FMath::Min(GroupEdgesA.Num(), GroupEdgesB.Num());
+		for (int32 i = 0; i < NumMatched; ++i)
 		{
-			NewTriangles = MinimalHoleFiller.NewTriangles;
+			FEdgeSpan SpanA = Topology->Edges[GroupEdgesA[i]].Span;
+			SpanA.SetCorrectOrientation();
+			FEdgeSpan SpanB = Topology->Edges[GroupEdgesB[i]].Span;
+			SpanB.SetCorrectOrientation();
+
+			TArray<int32> LoopVertices;
+			LoopVertices.Append(SpanA.Vertices);
+			LoopVertices.Append(SpanB.Vertices);
+
+			LastVids.A = SpanA.Vertices[bShouldReverseA ? 0 : SpanA.Vertices.Num() - 1];
+			LastVids.B = SpanB.Vertices[!bShouldReverseA ? 0 : SpanB.Vertices.Num() - 1];
+
+			BridgeEdges(LoopVertices);
 		}
-	}
-	else {
-		NewTriangles = SimpleHoleFiller.NewTriangles;
+
+		if (GroupEdgesA.Num() != GroupEdgesB.Num())
+		{
+			bool bAIsLonger = GroupEdgesA.Num() > GroupEdgesB.Num();
+			TArray<int32>& LongerSequence = bAIsLonger ? GroupEdgesA : GroupEdgesB;
+			TArray<int32> RemainingEdges;
+			for (int32 i = NumMatched; i < LongerSequence.Num(); ++i)
+			{
+				RemainingEdges.Add(LongerSequence[i]);
+			}
+			if (bAIsLonger == bShouldReverseA)
+			{
+				// If the remaining edges are the ones we reversed for the pairwise iteration, we
+				//  need to reverse them back so that they are in proper boundary ordering.
+				Algo::Reverse(RemainingEdges);
+			}
+
+			int32 OtherLastVid = bAIsLonger ? LastVids.B : LastVids.A;
+			if (ensure(OtherLastVid != IndexConstants::InvalidID))
+			{
+				TArray<int32> LoopVertices;
+				for (int32 i = 0; i < RemainingEdges.Num(); ++i)
+				{
+					FEdgeSpan Span = Topology->Edges[RemainingEdges[i]].Span;
+					Span.SetCorrectOrientation();
+					if (LoopVertices.Num() > 0 && ensure(LoopVertices.Last() == Span.Vertices[0]))
+					{
+						LoopVertices.Pop();
+					}
+					LoopVertices.Append(Span.Vertices);
+				}
+				LoopVertices.Add(OtherLastVid);
+
+				BridgeEdges(LoopVertices);
+			}
+		}//end if have unmatched edges
 	}
 
-	// Compute normals and UVs
-	if (Mesh->HasAttributes())
-	{
-		TArray<FVector3d> VertexPositions;
-		Loop.GetVertices(VertexPositions);
-		FVector3d PlaneOrigin;
-		FVector3d PlaneNormal;
-		PolygonTriangulation::ComputePolygonPlane<double>(VertexPositions, PlaneNormal, PlaneOrigin);
+	FText TransactionName = LOCTEXT("PolyMeshBridgeEdgeChange", "Bridge Edge");
+	GetToolManager()->BeginUndoTransaction(TransactionName);
 
-		FDynamicMeshEditor Editor(Mesh);
-		FFrame3d ProjectionFrame(PlaneOrigin, PlaneNormal);
-		Editor.SetTriangleNormals(NewTriangles);
-		Editor.SetTriangleUVsFromProjection(NewTriangles, ProjectionFrame, UVScaleFactor);
+	EmitCurrentMeshChangeAndUpdate(TransactionName, ChangeTracker.EndChange(), FGroupTopologySelection());
+
+	// Now that topology is updated, set the new selection
+	FGroupTopologySelection NewSelection;
+	for (int32 Tid : TrianglesToSelect)
+	{
+		NewSelection.SelectedGroupIDs.Add(Mesh->GetTriangleGroup(Tid));
+	}
+	if (ensure(!NewSelection.IsEmpty()))
+	{
+		SelectionMechanic->SetSelection(NewSelection);
 	}
 
-	EmitCurrentMeshChangeAndUpdate(LOCTEXT("PolyMeshBridgeEdgeChange", "Bridge Edge"),
-		ChangeTracker.EndChange(), CurrentSelection);
+	GetToolManager()->EndUndoTransaction();
 }
 
 
@@ -2548,50 +2890,167 @@ void UEditMeshPolygonsTool::ApplyFlipSingleEdge()
 	}
 
 	EmitCurrentMeshChangeAndUpdate(LOCTEXT("PolyMeshFlipChange", "Flip Edges"),
-		ChangeTracker.EndChange(), NewSelection);
+		ChangeTracker.EndChange(), FGroupTopologySelection());
 }
 
-void UEditMeshPolygonsTool::ApplyCollapseSingleEdge()
+void UEditMeshPolygonsTool::ApplyCollapseEdge()
 {
-	if (BeginMeshEdgeEditChange() == false)
+	using namespace EditMeshPolygonsToolLocals;
+
+	FText PartialCollapseFailure = LOCTEXT("OnCollapseFailures", "Some edges could not be collapsed");
+
+	FDynamicMesh3::FEdgeCollapseOptions CollapseOptions;
+	CollapseOptions.bAllowHoleCollapse = true;
+	CollapseOptions.bAllowCollapsingInternalEdgeWithExternalVertices = true;
+	CollapseOptions.bAllowTetrahedronCollapse = true;
+
+	FDynamicMesh3* Mesh = CurrentMesh.Get();
+	FGroupTopologySelection ActiveSelection = SelectionMechanic->GetActiveSelection();
+	if (ActiveSelection.IsEmpty())
 	{
-		GetToolManager()->DisplayMessage(LOCTEXT("OnCollapseFailedMessage", "Cannot Collapse Current Selection"), EToolMessageLevel::UserWarning);
+		GetToolManager()->DisplayMessage(LOCTEXT("OnCollapseFailedMessage", "Cannot collapse empty selection"), EToolMessageLevel::UserWarning);
 		return;
 	}
 
-	FDynamicMesh3* Mesh = CurrentMesh.Get();
-	FGroupTopologySelection ActiveSelection = SelectionMechanic->GetActiveSelection();
+	FText TransactionLabel = LOCTEXT("PolyMeshCollapseChange", "Collapse Edges");
+	GetToolManager()->BeginUndoTransaction(TransactionLabel);
+
 	FDynamicMeshChangeTracker ChangeTracker(Mesh);
 	ChangeTracker.BeginChange();
-	TSet<int32> ValidEdgeIDs;
-	for (FSelectedEdge& Edge : ActiveEdgeSelection)
+
+	TSet<int32> GroupEdgesToCollapse = ActiveSelection.SelectedEdgeIDs;
+	if (!ActiveSelection.SelectedGroupIDs.IsEmpty())
 	{
-		int32 eid = Edge.EdgeIDs[0];
-		if (Mesh->IsEdge(eid))
+		// Retriangulating can be thought of as equivalent to collapsing all the interior
+		//  edges, except without the risk of failing because of some kind of awful topology
+		//  on the inside of the group.
+		int32 NumCompleted = RetriangulateGroups(Mesh, Topology.Get(), ActiveSelection.SelectedGroupIDs, ChangeTracker);
+		if (NumCompleted != ActiveSelection.SelectedGroupIDs.Num())
+		{
+			GetToolManager()->DisplayMessage(PartialCollapseFailure, EToolMessageLevel::UserWarning);
+			// continue on and try to collapse the boundary
+		}
+
+		for (int32 GroupID : ActiveSelection.SelectedGroupIDs)
+		{
+			Topology->ForGroupEdges(GroupID, [&GroupEdgesToCollapse](const FGroupTopology::FGroupEdge&, int32 GroupEdgeID)
+			{
+				GroupEdgesToCollapse.Add(GroupEdgeID);
+			});
+		}
+	}
+	else if (!ActiveSelection.SelectedCornerIDs.IsEmpty())
+	{
+		ConvertCornerSelectionToGroupEdgeSelection(Topology.Get(), ActiveSelection.SelectedCornerIDs, GroupEdgesToCollapse);
+	}
+
+	TSet<int32> EidsToCollapse;
+	for (int32 GroupEdgeID : GroupEdgesToCollapse)
+	{
+		EidsToCollapse.Append(Topology->GetGroupEdgeEdges(GroupEdgeID));
+	}
+
+	// Partition our edges into connected components so that we can collapse into their
+	//  individual centroids.
+	TArray<TSet<int32>> EidComponents;
+	TSet<int32> PartitionedEids;
+	TArray<int32> TempQueue;
+	for (int32 Eid : EidsToCollapse)
+	{
+		if (PartitionedEids.Contains(Eid))
 		{
-			ValidEdgeIDs.Add(eid);
+			continue;
 		}
+
+		TSet<int32>& ComponentEids = EidComponents.Emplace_GetRef();
+		ComponentEids.Add(Eid);
+		FMeshConnectedComponents::GrowToConnectedEdges(*Mesh, { Eid }, ComponentEids, &TempQueue,
+			[&EidsToCollapse](int32 CurrentEid, int32 NeighborEid)
+		{
+			return EidsToCollapse.Contains(NeighborEid);
+		});
+		PartitionedEids.Append(ComponentEids);
 	}
-	TSet<int32> DoneEdgeIDs;
-	for (int32 eid : ValidEdgeIDs)
+
+	// Now process our components.
+	bool bAllCollapsesSuccessful = true;
+	TSet<int32> NewSelectionVids;
+	for (TSet<int32>& Component : EidComponents)
 	{
-		if (DoneEdgeIDs.Contains(eid) == false && Mesh->IsEdge(eid))
+		FVector3d Centroid = FVector3d::Zero();
+		for (int32 Eid : Component)
+		{
+			Centroid += Mesh->GetEdgePoint(Eid, 0.5);
+		}
+		Centroid /= Component.Num();
+
+		// Unfiltered because vids will disappear in subsequent collapses
+		TSet<int32> UnfilteredVidsToMove;
+		
+		for (int32 Eid : Component)
 		{
-			FIndex2i ev = Mesh->GetEdgeV(eid);
-			ChangeTracker.SaveVertexOneRingTriangles(ev.A, true);
-			ChangeTracker.SaveVertexOneRingTriangles(ev.B, true);
+			// Some edges might be collapsed away by other collapses
+			if (!Mesh->IsEdge(Eid))
+			{
+				continue;
+			}
+
+			FIndex2i EdgeVids = Mesh->GetEdgeV(Eid);
+			ChangeTracker.SaveVertexOneRingTriangles(EdgeVids.A, true);
+			ChangeTracker.SaveVertexOneRingTriangles(EdgeVids.B, true);
 			FDynamicMesh3::FEdgeCollapseInfo CollapseInfo;
-			if (Mesh->CollapseEdge(ev.A, ev.B, CollapseInfo) == EMeshResult::Ok)
+			EMeshResult Result = Mesh->CollapseEdge(EdgeVids.A, EdgeVids.B, CollapseOptions, CollapseInfo);
+			
+			if (Result == EMeshResult::Ok)
+			{
+				UnfilteredVidsToMove.Add(CollapseInfo.KeptVertex);
+			}
+			else
 			{
-				DoneEdgeIDs.Add(eid);
-				DoneEdgeIDs.Add(CollapseInfo.RemovedEdges.A);
-				DoneEdgeIDs.Add(CollapseInfo.RemovedEdges.B);
+				bAllCollapsesSuccessful = false;
+			}
+		}
+
+		for (int32 Vid : UnfilteredVidsToMove)
+		{
+			if (Mesh->IsVertex(Vid))
+			{
+				Mesh->SetVertex(Vid, Centroid);
+				NewSelectionVids.Add(Vid);
 			}
 		}
 	}
 
-	EmitCurrentMeshChangeAndUpdate(LOCTEXT("PolyMeshCollapseChange", "Collapse Edges"), 
-		ChangeTracker.EndChange(), FGroupTopologySelection());
+	if (!bAllCollapsesSuccessful)
+	{
+		GetToolManager()->DisplayMessage(PartialCollapseFailure, EToolMessageLevel::UserWarning);
+	}
+
+	EmitCurrentMeshChangeAndUpdate(TransactionLabel, ChangeTracker.EndChange(), FGroupTopologySelection());
+
+	// Now that the topology is updated, we can get the new corner id's to
+	//  set the new selection.
+	FGroupTopologySelection NewSelection;
+	for (int32 Vid : NewSelectionVids)
+	{
+		if (!ensure(Mesh->IsVertex(Vid)))
+		{
+			continue;
+		}
+		int32 CornerID = Topology->GetCornerIDFromVertexID(Vid);
+		if (CornerID != IndexConstants::InvalidID)
+		{
+			NewSelection.SelectedCornerIDs.Add(CornerID);
+		}
+	}
+	// Seems possible to end up with an empty selection if we collapsed a triangle hole in a group,
+	//  so the new vertex is not part of a group boundary.
+	if (!NewSelection.IsEmpty())
+	{
+		SelectionMechanic->SetSelection(NewSelection);
+	}
+
+	GetToolManager()->EndUndoTransaction();
 }
 
 void UEditMeshPolygonsTool::ApplySplitSingleEdge()
diff --git a/Engine/Plugins/Runtime/MeshModelingToolset/Source/MeshModelingTools/Public/EditMeshPolygonsTool.h b/Engine/Plugins/Runtime/MeshModelingToolset/Source/MeshModelingTools/Public/EditMeshPolygonsTool.h
--- a/Engine/Plugins/Runtime/MeshModelingToolset/Source/MeshModelingTools/Public/EditMeshPolygonsTool.h
+++ b/Engine/Plugins/Runtime/MeshModelingToolset/Source/MeshModelingTools/Public/EditMeshPolygonsTool.h
@@ -433,6 +433,10 @@ public:
 	/** Delete selected edge, implicitly merging any connected faces */
 	UFUNCTION(CallInEditor, Category = EdgeEdits, meta = (DisplayName = "Delete Edges", DisplayPriority = 8))
 	void DeleteEdge() { PostAction(EEditMeshPolygonsToolActions::Delete); }
+
+	/** Collapse the selected edges, deleting the attached triangles and merging their vertices into one */
+	UFUNCTION(CallInEditor, Category = EdgeEdits, meta = (DisplayName = "Collapse", DisplayPriority = 5))
+	void Collapse() { PostAction(EEditMeshPolygonsToolActions::CollapseEdge); }
 };
 
 
@@ -459,7 +463,7 @@ public:
 
 	/** Collapse the selected edges, deleting the attached triangles and merging its two vertices into one */
 	UFUNCTION(CallInEditor, Category = EdgeEdits, meta = (DisplayName = "Collapse", DisplayPriority = 5))
-	void Collapse() { PostAction(EEditMeshPolygonsToolActions::CollapseSingleEdge); }
+	void Collapse() { PostAction(EEditMeshPolygonsToolActions::CollapseEdge); }
 
 	/** Flip the selected (non-border, non-seam) edges, replacing them with new edges in the crossing direction */
 	UFUNCTION(CallInEditor, Category = EdgeEdits, meta = (DisplayName = "Flip", DisplayPriority = 6))
@@ -684,6 +688,7 @@ protected:
 	void ApplySimplifyAlongEdges();
 
 	void ApplyFlipSingleEdge();
+	UE_DEPRECATED(5.5, "Use ApplyCollapseEdge instead.")
 	void ApplyCollapseSingleEdge();
 	void ApplySplitSingleEdge();
 
