[GeometryScript] Fix static mesh to dynamic mesh conversion in cooked editor

* In cooked editor, meshes don't have Source models, like at runtime, and we cannot check Build scale settings.
* If we have no source model, we should follow the runtime path, and assume build scale is already baked in render data

#rnx

diff --git a/Engine/Plugins/Runtime/MeshModelingToolset/Source/ModelingComponents/Private/ConversionUtils/SceneComponentToDynamicMesh.cpp b/Engine/Plugins/Runtime/MeshModelingToolset/Source/ModelingComponents/Private/ConversionUtils/SceneComponentToDynamicMesh.cpp
--- a/Engine/Plugins/Runtime/MeshModelingToolset/Source/ModelingComponents/Private/ConversionUtils/SceneComponentToDynamicMesh.cpp
+++ b/Engine/Plugins/Runtime/MeshModelingToolset/Source/ModelingComponents/Private/ConversionUtils/SceneComponentToDynamicMesh.cpp
@@ -330,19 +330,22 @@ namespace Private::ConversionHelper
 
 		FStaticMeshLODResourcesToDynamicMesh::ConversionOptions ConvertOptions;
 #if WITH_EDITOR
-		if (AssetOptions.bUseBuildScale)
+		const bool bIsSourceModelValid = FromStaticMeshAsset->IsSourceModelValid(UseLODIndex);
+		if (AssetOptions.bUseBuildScale && bIsSourceModelValid)
 		{
 			// respect BuildScale build setting
 			const FMeshBuildSettings& LODBuildSettings = FromStaticMeshAsset->GetSourceModel(UseLODIndex).BuildSettings;
 			ConvertOptions.BuildScale = (FVector3d)LODBuildSettings.BuildScale3D;
 		}
+		// In case of cooked editor, Source model won't be valid, so it will follow the same rules as the runtime path.
+		else if (!AssetOptions.bUseBuildScale && !bIsSourceModelValid)
 #else
 		if (!AssetOptions.bUseBuildScale)
+#endif
 		{
 			OutErrorMessage = LOCTEXT("CopyMeshFromStaticMesh_BuildScaleAlreadyBaked", "Requested mesh without BuildScale, but BuildScale is already baked into the RenderData.");
 			return false;
 		}
-#endif
 
 		FStaticMeshLODResourcesToDynamicMesh Converter;
 		if (bRequestInstanceVertexColors && StaticMeshComponent && StaticMeshComponent->LODData.IsValidIndex(UseLODIndex))
