MeshModelingTools: PolyEd/TriEd additions:
- Add an option to weld centered
- Add ability to weld a pair of vertices
- Fix selection after split edge not working correctly
- Fix behavior in some pathologic cases (welding sides of ribbon)
- Allow weld and collapse to collapse away isolated triangles

diff --git a/Engine/Plugins/Runtime/GeometryProcessing/Source/DynamicMesh/Private/Operations/WeldEdgeSequence.cpp b/Engine/Plugins/Runtime/GeometryProcessing/Source/DynamicMesh/Private/Operations/WeldEdgeSequence.cpp
--- a/Engine/Plugins/Runtime/GeometryProcessing/Source/DynamicMesh/Private/Operations/WeldEdgeSequence.cpp
+++ b/Engine/Plugins/Runtime/GeometryProcessing/Source/DynamicMesh/Private/Operations/WeldEdgeSequence.cpp
@@ -25,11 +25,6 @@ FWeldEdgeSequence::EWeldResult FWeldEdgeSequence::Weld()
 		return Result;
 	}
 
-	if ((Result = CheckForAndCollapseSideTriangles()) != EWeldResult::Ok)
-	{
-		return Result;
-	}
-
 	if ((Result = WeldEdgeSequence()) != EWeldResult::Ok)
 	{
 		return Result;
@@ -70,14 +65,21 @@ FWeldEdgeSequence::EWeldResult FWeldEdgeSequence::CheckInput()
 }
 
 FWeldEdgeSequence::EWeldResult FWeldEdgeSequence::SplitSmallerSpan()
+{
+	return SplitEdgesToEqualizeSpanLengths(*Mesh, EdgeSpanToKeep, EdgeSpanToDiscard);
+}
+
+FWeldEdgeSequence::EWeldResult UE::Geometry::FWeldEdgeSequence::SplitEdgesToEqualizeSpanLengths(FDynamicMesh3& Mesh, FEdgeSpan& Span1, FEdgeSpan& Span2)
 {
 	// For each new vertex that must be created:
 	// The longest simple edge is found and split
 	// The newly generated vertex is inserted into the span
 	// The newly generated edge is inserted into the span
+	// TODO: Could improve this by sorting lengths and keeping those updated as we split (i.e. using
+	//  a priority queue)
 
-	FEdgeSpan& SpanToSplit = (EdgeSpanToDiscard.Vertices.Num() < EdgeSpanToKeep.Vertices.Num()) ? EdgeSpanToDiscard : EdgeSpanToKeep;
-	int TotalSplits = FMath::Abs(EdgeSpanToDiscard.Vertices.Num() - EdgeSpanToKeep.Vertices.Num());
+	FEdgeSpan& SpanToSplit = (Span1.Vertices.Num() < Span2.Vertices.Num()) ? Span1 : Span2;
+	int TotalSplits = FMath::Abs(Span1.Vertices.Num() - Span2.Vertices.Num());
 	for (int SplitCount = 0; SplitCount < TotalSplits; ++SplitCount)
 	{
 		double MaxLength = 0.0;
@@ -87,8 +89,8 @@ FWeldEdgeSequence::EWeldResult FWeldEdgeSequence::SplitSmallerSpan()
 		// Find longest edge and store length, ID, and index
 		for (int EdgeIndex = 0; EdgeIndex < SpanToSplit.Edges.Num(); ++EdgeIndex)
 		{
-			FVector3d VertA = Mesh->GetVertex(Mesh->GetEdge(SpanToSplit.Edges[EdgeIndex]).Vert.A);
-			FVector3d VertB = Mesh->GetVertex(Mesh->GetEdge(SpanToSplit.Edges[EdgeIndex]).Vert.B);
+			FVector3d VertA = Mesh.GetVertex(Mesh.GetEdge(SpanToSplit.Edges[EdgeIndex]).Vert.A);
+			FVector3d VertB = Mesh.GetVertex(Mesh.GetEdge(SpanToSplit.Edges[EdgeIndex]).Vert.B);
 			double EdgeLength = DistanceSquared(VertA, VertB);
 
 			if (MaxLength < EdgeLength)
@@ -101,7 +103,7 @@ FWeldEdgeSequence::EWeldResult FWeldEdgeSequence::SplitSmallerSpan()
 
 		// Split longest edge
 		FDynamicMesh3::FEdgeSplitInfo SplitInfo;
-		EMeshResult Result = Mesh->SplitEdge(LongestEID, SplitInfo);
+		EMeshResult Result = Mesh.SplitEdge(LongestEID, SplitInfo);
 		if (Result != EMeshResult::Ok)
 		{
 			return EWeldResult::Failed_CannotSplitEdge;
@@ -126,6 +128,8 @@ FWeldEdgeSequence::EWeldResult FWeldEdgeSequence::SplitSmallerSpan()
 	return EWeldResult::Ok;
 }
 
+// No longer used, as it fails to consider the case of intervening triangles in the middle
+//  of the sequence, if the edges are the boundary of a band.
 FWeldEdgeSequence::EWeldResult FWeldEdgeSequence::CheckForAndCollapseSideTriangles()
 {
 	// Checks for and deletes edge between VertA and VertB
@@ -176,52 +180,152 @@ FWeldEdgeSequence::EWeldResult FWeldEdgeSequence::CheckForAndCollapseSideTriangl
 
 FWeldEdgeSequence::EWeldResult FWeldEdgeSequence::WeldEdgeSequence()
 {
-	if (!ensure(EdgeSpanToDiscard.Edges.Num() == EdgeSpanToKeep.Edges.Num()))
+	if (!ensure(EdgeSpanToDiscard.Edges.Num() == EdgeSpanToKeep.Edges.Num()
+		&& EdgeSpanToDiscard.Vertices.Num() == EdgeSpanToDiscard.Edges.Num() + 1
+		&& EdgeSpanToKeep.Vertices.Num() == EdgeSpanToKeep.Edges.Num() + 1))
 	{
 		return EWeldResult::Failed_Other;
 	}
 
-	for (int EdgePairIndex = 0; EdgePairIndex < EdgeSpanToDiscard.Edges.Num(); ++EdgePairIndex)
+	// There are certain pathological cases in which one edge weld could delete one of the
+	//  next edge paired verts before we can use its location and attribute values for interpolation.
+	//  For example in the following diagram, welding ab to de will delete the triangle bce, but we
+	//  still need to update vertex f:
+	//    a_b_c
+	//     \|/
+	//    d_e_f
+	//    |\|\|
+	//
+	// This is only possible if there is an edge between b and e (so that the triangle can be deleted), and if
+	//  there is not another triangle holding on to c. We can handle this case by collapsing a bit out of order-
+	//  if we know there's an edge at this pair but there is not an edge at the next pair, we can do the next
+	//  pair first, and we know the same issue won't occur there. On the flip side if there is an edge at the
+	//  next pair too, then the next vert can't be destroyed by collapsing this one.
+	
+	// For this and other edge cases, it is safer to do the welding vert by vert instead of edge by edge.
+	//  As another example, in this diagram, after merging ab to de, bc no longer exists, but c still needs
+	//  welding to f:
+	//    a_b_c
+	//    |\|/|
+	//    d_e_f
+
+	auto ProcessVidPair = [this](int32 KeepVid, int32 DiscardVid, int32 KeepVertIndex) -> EWeldResult
 	{
-		int EdgeA = EdgeSpanToDiscard.Edges[EdgeSpanToDiscard.Edges.Num() - 1 - EdgePairIndex];
-		int EdgeB = EdgeSpanToKeep.Edges[EdgePairIndex];
+		FDynamicMesh3::FMergeVerticesInfo MergeInfo;
+		EMeshResult Result = Mesh->MergeVertices(KeepVid, DiscardVid, InterpolationT, MergeInfo);
 
-		// We need to check this because MergeEdges() can implicitly
-		// merge adjacent edges when adjacent edges share a vertex.
-		if (Mesh->IsEdge(EdgeA) && Mesh->IsEdge(EdgeB))
+		if (Result == EMeshResult::Failed_CollapseTriangle
+			|| Result == EMeshResult::Failed_CollapseQuad
+			|| Result == EMeshResult::Failed_FoundDuplicateTriangle)
 		{
-			FDynamicMesh3::FMergeEdgesInfo MergeInfo;
-			EMeshResult Result = Mesh->MergeEdges(EdgeB, EdgeA, MergeInfo, 
-				// This is bCheckValidOrientation, which, roughly speaking, checks to see if the direction
-				//  of EdgeB in world space approximately flips in order to match the winding order of its new
-				//  neighbor. We don't want this, in part because our edges temoporarily move anyway as we
-				//  go along and weld them, potentially temporarily flipping. We always weld according to
-				//  winding order, irregardless of original edge position in world space. 
-				false);
-			if (Result != EMeshResult::Ok)
+			// Currently collapse doesn't allow us to collapse away an isolated triangle, quad,
+			//  or double sided triangle. We can deal with this case, however, simply by deleting
+			//  them.
+			// TODO: Should maybe have an option for this in CollapseEdge/MergeVertices
+			int32 Eid = Mesh->FindEdge(KeepVid, DiscardVid);
+			if (!ensure(Mesh->IsEdge(Eid)))
+			{
+				return EWeldResult::Failed_Other;
+			}
+			FIndex2i TidsToDelete = Mesh->GetEdgeT(Eid);
+			Mesh->RemoveTriangle(TidsToDelete.A);
+			if (TidsToDelete.B != IndexConstants::InvalidID)
 			{
-				if (Result == EMeshResult::Failed_InvalidNeighbourhood && bAllowFailedMerge == true)
+				Mesh->RemoveTriangle(TidsToDelete.B);
+			}
+		}
+		else if (Result == EMeshResult::Failed_InvalidNeighbourhood && bAllowFailedMerge == true)
+		{
+			// If we're allowed to, we just place the edges together without welding.
+			FVector3d Destination = Lerp(Mesh->GetVertex(KeepVid), Mesh->GetVertex(DiscardVid), InterpolationT);
+			Mesh->SetVertex(KeepVid, Destination);
+			Mesh->SetVertex(DiscardVid, Destination);
+
+			// Maybe it's unfortunate that we have to output unmerged edges instead of vertices, but
+			//  theoretically the edges on either side were not successfully welded.
+			auto AddEdgeAtKeepIndex = [this](int32 KeepEidIndex)
+			{
+				int32 DiscardEidIndex = EdgeSpanToDiscard.Edges.Num() - 1 - KeepEidIndex;
+				int32 KeepEid = EdgeSpanToKeep.Edges[KeepEidIndex];
+				int32 DiscardEid = EdgeSpanToDiscard.Edges[DiscardEidIndex];
+				if (Mesh->IsEdge(KeepEid) && Mesh->IsEdge(DiscardEid))
 				{
-					// Get properly oriented vertex pairs for edges which can't be merged
-					FIndex2i EdgeAVerts = Mesh->GetEdgeV(EdgeA);
-					FIndex2i EdgeBVerts = Mesh->GetEdgeV(EdgeB);
-					IndexUtil::OrientTriEdge(EdgeAVerts.A, EdgeAVerts.B, Mesh->GetTriangle(Mesh->GetEdgeT(EdgeA).A));
-					IndexUtil::OrientTriEdge(EdgeBVerts.A, EdgeBVerts.B, Mesh->GetTriangle(Mesh->GetEdgeT(EdgeB).A));
-
-					// Move verts from EdgeA to EdgeB, leaving an invisible seam
-					Mesh->SetVertex(EdgeAVerts.A, Mesh->GetVertex(EdgeBVerts.B));
-					Mesh->SetVertex(EdgeAVerts.B, Mesh->GetVertex(EdgeBVerts.A));
-
-					// Add eids to UnmergedEdgePairsOut so user can easily identify offending edges
-					UnmergedEdgePairsOut.Add(TPair<int, int>(EdgeA, EdgeB));
+					UnmergedEdgePairsOut.Add(TPair<int, int>(KeepEid, DiscardEid));
 				}
-				else
+			};
+			if (KeepVertIndex > 0)
+			{
+				AddEdgeAtKeepIndex(KeepVertIndex - 1);
+			}
+			if (KeepVertIndex < EdgeSpanToKeep.Edges.Num())
+			{
+				AddEdgeAtKeepIndex(KeepVertIndex);
+			}
+		}
+		else if (Result != EMeshResult::Ok)
+		{
+			return EWeldResult::Failed_Other;
+		}
+
+		return EWeldResult::Ok;
+	};
+
+	for (int KeepVertIndex = 0; KeepVertIndex < EdgeSpanToKeep.Vertices.Num(); ++KeepVertIndex)
+	{
+		// The spans are oriented opposite directions, so iterate in opposite order
+		int32 DiscardVertIndex = EdgeSpanToDiscard.Vertices.Num() - 1 - KeepVertIndex;
+		int32 KeepVid = EdgeSpanToKeep.Vertices[KeepVertIndex];
+		int32 DiscardVid = EdgeSpanToDiscard.Vertices[DiscardVertIndex];
+
+		if (KeepVid == DiscardVid)
+		{
+			continue;
+		}
+
+		if (!ensure(Mesh->IsVertex(KeepVid) && Mesh->IsVertex(DiscardVid)))
+		{
+			// This shouldn't happen due to our out-of-order collapse strategy, see above
+			continue;
+		}
+
+		// See above for why we consider processing the next vid first
+		bool bProcessedNext = false;
+		int32 InterveningEdge = Mesh->FindEdge(KeepVid, DiscardVid);
+		if (InterveningEdge != IndexConstants::InvalidID 
+			&& KeepVertIndex < EdgeSpanToKeep.Vertices.Num()-1)
+		{
+			int32 NextKeepVid = EdgeSpanToKeep.Vertices[KeepVertIndex + 1];
+			int32 NextDiscardVid = EdgeSpanToDiscard.Vertices[DiscardVertIndex - 1];
+			if (NextKeepVid == NextDiscardVid)
+			{
+				// Consider ourselves to have dealt with the next vertex
+				bProcessedNext = true;
+			}
+			else if (ensure(Mesh->IsVertex(NextKeepVid) && Mesh->IsVertex(NextDiscardVid))
+				&& Mesh->FindEdge(NextKeepVid, NextDiscardVid) == IndexConstants::InvalidID)
+			{
+				// This is a safe vert pair to collapse, since it doesn't have intervening edges.
+				//  And if it did, it wouldn't be in danger of losing a vert due to our collapse.
+				EWeldResult Result = ProcessVidPair(NextKeepVid, NextDiscardVid, KeepVertIndex + 1);
+				if (Result != EWeldResult::Ok)
 				{
-					return EWeldResult::Failed_Other;
+					return Result;
 				}
+				bProcessedNext = true;
 			}
 		}
-	}
+		
+		EWeldResult Result = ProcessVidPair(KeepVid, DiscardVid, KeepVertIndex);
+		if (Result != EWeldResult::Ok)
+		{
+			return Result;
+		}
+
+		if (bProcessedNext)
+		{
+			++KeepVertIndex;
+		}
+	}//end iterating through vertices
 
 	return EWeldResult::Ok;
 }
diff --git a/Engine/Plugins/Runtime/GeometryProcessing/Source/DynamicMesh/Public/Operations/WeldEdgeSequence.h b/Engine/Plugins/Runtime/GeometryProcessing/Source/DynamicMesh/Public/Operations/WeldEdgeSequence.h
--- a/Engine/Plugins/Runtime/GeometryProcessing/Source/DynamicMesh/Public/Operations/WeldEdgeSequence.h
+++ b/Engine/Plugins/Runtime/GeometryProcessing/Source/DynamicMesh/Public/Operations/WeldEdgeSequence.h
@@ -45,13 +45,15 @@ public:
 	FEdgeSpan EdgeSpanToKeep;		// This is the updated edge span which can be used once Weld() is called
 
 	// Whether triangle deletion is allowed in order to merge edges which are connected by a different edge
-	// This specifically determines how CheckForAndCollapseSideTriangles() behaves
 	bool bAllowIntermediateTriangleDeletion = false;
 
 	// When true, failed calls to MergeEdges() will be handled by moving the edges without merging such
 	// that the final result appears to be welded but has invisible seam(s) instead of just failing.
 	bool bAllowFailedMerge = false;
 
+	// When vertices are welded, each kept vertex will be placed at Lerp(KeepPos, RemovePos, InterpolationT)
+	double InterpolationT = 0;
+
 	// This is populated with pairs of eids which were not able to be merged.
 	// Only valid when bAllowFailedMerge is true
 	TArray<TPair<int, int>> UnmergedEdgePairsOut;
@@ -74,6 +76,12 @@ public:
 	*/
 	EWeldResult Weld();
 
+	/**
+	 * Helper that splits the edges in the shorter span until the spans have the same number
+	 *  of edges. Weld() will automatically do this, but this is public in case the user wants
+	 *  to equalize spans that get concatenated together before all being welded at once.
+	 */
+	static EWeldResult SplitEdgesToEqualizeSpanLengths(FDynamicMesh3& Mesh, FEdgeSpan& Span1, FEdgeSpan& Span2);
 protected:
 	/**
 	* Verifies validity of input edges by ensuring they are
@@ -91,13 +99,8 @@ protected:
 	*/
 	EWeldResult SplitSmallerSpan();
 
-	/**
-	* Checks for edges between terminating vertices of the input spans and attempts to delete them if possible
-	* by deleting the adjacent triangle. These cases would prevent MergeEdges() from succeeding because it won't
-	* implicitly collapse/delete triangles which connect the two edges being merged.
-	* 
-	* @return EWeldResult::OK on success
-	*/
+
+	UE_DEPRECATED(5.5, "Side triangles are handled appropriately in WeldEdgeSequence")
 	EWeldResult CheckForAndCollapseSideTriangles();
 
 	/**
diff --git a/Engine/Plugins/Runtime/MeshModelingToolset/Source/MeshModelingTools/Private/EditMeshPolygonsTool.cpp b/Engine/Plugins/Runtime/MeshModelingToolset/Source/MeshModelingTools/Private/EditMeshPolygonsTool.cpp
--- a/Engine/Plugins/Runtime/MeshModelingToolset/Source/MeshModelingTools/Private/EditMeshPolygonsTool.cpp
+++ b/Engine/Plugins/Runtime/MeshModelingToolset/Source/MeshModelingTools/Private/EditMeshPolygonsTool.cpp
@@ -74,6 +74,15 @@ namespace EditMeshPolygonsToolLocals
 		"Hold Ctrl while translating or (in local mode) rotating to align to scene. Shift and Ctrl change marquee select "
 		"behavior. Ctrl+R toggles Gizmo Orientation Lock.");
 
+	FText WeldIncompleteMessage = LOCTEXT("OnWeldEdgesCompletedSeamsRemain", "Warning: welding incomplete because it would create "
+		"invalid geometry (attached non manifold edge or duplicate triangle). Seam still exists at weld "
+		"location. Modify attached triangles and retry, or undo.");
+
+	FText PartialCollapseFailureMessage = LOCTEXT("OnCollapseFailures", "Some edges could not be collapsed, "
+		"likely because adjoining edges would then have non manifold geometry (more than two faces), or "
+		"the mesh would end up empty.");
+	FText CollapseEdgeTransactionLabel = LOCTEXT("PolyMeshCollapseChange", "Collapse Edges");
+
 	FString GetPropertyCacheIdentifier(bool bTriangleMode)
 	{
 		return bTriangleMode ? TEXT("TriEditTool") : TEXT("PolyEditTool");
@@ -86,6 +95,14 @@ namespace EditMeshPolygonsToolLocals
 			"more than this number of edges to be rendered in PolyEd or TriEd force the tools to "
 			"be disabled to avoid hanging the editor."));
 
+	bool bAllowBowtieWeldAtInternalVertex = false;
+	static FAutoConsoleVariableRef CVarAllowWeldInternalBowtie(
+		TEXT("modeling.PolyEdit.AllowWeldInternalBowtie"),
+		bAllowBowtieWeldAtInternalVertex,
+		TEXT("If true, \"Weld\" and \"Weld To\" operations on a pair of vertices will allow the creation "
+			"of non-boundary bowties. If false, then the vertices in these situations will not be welded, "
+			"and will instead be moved to the destination."));
+
 	// Allows undo/redo of addition of extra corners in the group topology based on user angle thresholds.
 	// Used after user-triggered topology corner changes where the mesh was not actually edited.
 	class FExtraCornerChange : public FToolCommandChange
@@ -373,6 +390,34 @@ namespace EditMeshPolygonsToolLocals
 		}
 		return NumCompleted;
 	}//end RetriangulateGroups
+
+	// Helper that removes the triangles around an edge as long as they are not the last
+	//  ones in the mesh. Used to allow collapses of isolated triangles and quads, which
+	//  are not currently permitted by CollapseEdge.
+	// TODO: We should probably have a permissiveness option that does allow this in
+	//  CollapseEdge, though it should be noted that the kept vert may end up deleted
+	//  in that case.
+	bool RemoveEdgeTrisIfNotLast(FDynamicMesh3& Mesh, int32 Eid)
+	{
+		if (!Mesh.IsEdge(Eid))
+		{
+			return false;
+		}
+
+		FIndex2i EdgeTids = Mesh.GetEdgeT(Eid);
+
+		if (Mesh.TriangleCount() > 2
+			|| (Mesh.TriangleCount() > 1 && EdgeTids.B == IndexConstants::InvalidID))
+		{
+			Mesh.RemoveTriangle(EdgeTids.A);
+			if (EdgeTids.B != IndexConstants::InvalidID)
+			{
+				Mesh.RemoveTriangle(EdgeTids.B);
+			}
+			return true;
+		}
+		return false;
+	}
 }
 
 /*
@@ -1628,6 +1673,9 @@ void UEditMeshPolygonsTool::OnTick(float DeltaTime)
 		case EEditMeshPolygonsToolActions::WeldEdges:
 			ApplyWeldEdges();
 			break;
+		case EEditMeshPolygonsToolActions::WeldEdgesCentered:
+			ApplyWeldEdges(0.5);
+			break;
 		case EEditMeshPolygonsToolActions::StraightenEdge:
 			ApplyStraightenEdges();
 			break;
@@ -2249,6 +2297,11 @@ void UEditMeshPolygonsTool::ApplyCollapseSingleEdge()
 
 
 void UEditMeshPolygonsTool::ApplyWeldEdges()
+{
+	ApplyWeldEdges(0);
+}
+
+void UEditMeshPolygonsTool::ApplyWeldEdges(double InterpolationT)
 {
 	using namespace EditMeshPolygonsToolLocals;
 
@@ -2257,7 +2310,20 @@ void UEditMeshPolygonsTool::ApplyWeldEdges()
 	TSet<int32> GroupEdges;
 	if (!CurrentSelection.SelectedCornerIDs.IsEmpty())
 	{
+		if (CurrentSelection.SelectedCornerIDs.Num() == 2)
+		{
+			ApplyWeldVertices(InterpolationT);
+			return;
+		}
 		ConvertCornerSelectionToGroupEdgeSelection(Topology.Get(), CurrentSelection.SelectedCornerIDs, GroupEdges);
+		if (GroupEdges.Num() < 2)
+		{
+			GetToolManager()->DisplayMessage(
+				LOCTEXT("OnWeldVerticesFailedInvalidCount", "Cannot Weld current selection, "
+					"selection must be either 2 vertices or convertible to at least 2 edges."),
+				EToolMessageLevel::UserWarning);
+			return;
+		}
 	}
 	else
 	{
@@ -2290,24 +2356,26 @@ void UEditMeshPolygonsTool::ApplyWeldEdges()
 	//  We reverse one of them so we can do our pairwise welding in the proper order.
 	Algo::Reverse(bShouldReverseA ? GroupEdgesA : GroupEdgesB);
 
-	// At the end, we're going to select the edges we keep.
-	// Note: this requires eids to stay constant across the mesh copy call further below. That seems to be the
-	//  case for dynamic mesh copies, otherwise we'd need to use tri-subidx pairs.
-	TArray<int32> KeptEids;
-	for (int32 GroupEdgeID : GroupEdgesB)
-	{
-		KeptEids.Append(Topology->GetGroupEdgeEdges(GroupEdgeID));
-	}
-
-	FText TransactionName = LOCTEXT("PolyMeshWeldEdgeChange", "Weld Edges");
-	GetToolManager()->BeginUndoTransaction(TransactionName);
 	FDynamicMeshChangeTracker ChangeTracker(Mesh);
 	ChangeTracker.BeginChange();
 
 	bool bAllSucceeded = true;
 	bool bHaveSeam = false;
 
-	auto WeldGroupEdges = [&ChangeTracker, Mesh, &bAllSucceeded, &bHaveSeam](FEdgeSpan& SpanA, FEdgeSpan& SpanB)
+	// Conceptually, we weld pairwise across group edges until we reach the last group edge
+	//  of the shorter sequence, and then weld that edge to remaining concatenated group edges
+	//  in the longer sequence. 
+	// However there are some pathological cases where welding of one edge could remove an edge
+	//  of an adjacent group edge, and these are best handled inside FWeldEdgeSequence if it 
+	//  knows all of the edges it needs to weld. So, we want to pass FWeldEdgeSequence the 
+	//  concatenated sequences, but we need to do the equalizing splits on a per-group-edge
+	//  basis so that we can make sure that group corners still get welded to other group corners.
+	
+	TArray<int32> ConcatenatedKeptEids;
+	TArray<int32> ConcatenatedDiscardEids;
+
+	auto PrepGroupEdgePair = [&ChangeTracker, Mesh, &bAllSucceeded, 
+		&ConcatenatedKeptEids, &ConcatenatedDiscardEids, bShouldReverseA](FEdgeSpan& SpanA, FEdgeSpan& SpanB)
 	{
 		// Save one ring tri's for vertices along both edges. The kept edge is necessary
 		//  because we might be splitting its triangles if needed.
@@ -2334,15 +2402,22 @@ void UEditMeshPolygonsTool::ApplyWeldEdges()
 			});
 		}
 
-		FWeldEdgeSequence EdgeWelder(Mesh, SpanA, SpanB);
-		EdgeWelder.bAllowIntermediateTriangleDeletion = true;
-		EdgeWelder.bAllowFailedMerge = true;
+		SpanA.SetCorrectOrientation();
+		SpanB.SetCorrectOrientation();
 
-		FWeldEdgeSequence::EWeldResult Result = EdgeWelder.Weld();
-		if (EdgeWelder.UnmergedEdgePairsOut.Num() != 0)
+		FWeldEdgeSequence::EWeldResult Result = FWeldEdgeSequence::SplitEdgesToEqualizeSpanLengths(*Mesh, SpanA, SpanB);
+		
+		if (bShouldReverseA)
+		{
+			ConcatenatedDiscardEids.Insert(SpanA.Edges, 0);
+			ConcatenatedKeptEids.Append(SpanB.Edges);
+		}
+		else
 		{
-			bHaveSeam = true;
+			ConcatenatedDiscardEids.Append(SpanA.Edges);
+			ConcatenatedKeptEids.Insert(SpanB.Edges, 0);
 		}
+
 		if (Result != FWeldEdgeSequence::EWeldResult::Ok)
 		{
 			bAllSucceeded = false;
@@ -2351,16 +2426,13 @@ void UEditMeshPolygonsTool::ApplyWeldEdges()
 		return true;
 	};
 
-	// Some number of edges will be welded normally. However if we have a mismatched number,
-	//  the last edge of the smaller sequence will have to be welded to the remainder of edges.
 	int32 NumMatched = GroupEdgesA.Num() == GroupEdgesB.Num() ? GroupEdgesA.Num()
 		: FMath::Min(GroupEdgesA.Num(), GroupEdgesB.Num()) - 1;
 	for (int32 i = 0; i < NumMatched; ++i)
 	{
-		FEdgeSpan SpanA;
-		SpanA.InitializeFromEdges(Mesh, Topology->Edges[GroupEdgesA[i]].Span.Edges);
+		FEdgeSpan& SpanA = Topology->Edges[GroupEdgesA[i]].Span;
 		FEdgeSpan& SpanB = Topology->Edges[GroupEdgesB[i]].Span;
-		WeldGroupEdges(SpanA, SpanB);
+		PrepGroupEdgePair(SpanA, SpanB);
 	}
 
 	// If there was a mismatched number of edges, we have set NumMatched to be one less than
@@ -2390,15 +2462,56 @@ void UEditMeshPolygonsTool::ApplyWeldEdges()
 		FEdgeSpan SpanB;
 		SpanB.InitializeFromEdges(Mesh, SpanBEids);
 
-		WeldGroupEdges(SpanA, SpanB);
+		PrepGroupEdgePair(SpanA, SpanB);
+	}
+
+	FEdgeSpan ConcatenatedKeptSpan(Mesh), ConcatenatedDiscardSpan(Mesh);
+	ConcatenatedKeptSpan.InitializeFromEdges(ConcatenatedKeptEids);
+	ConcatenatedDiscardSpan.InitializeFromEdges(ConcatenatedDiscardEids);
+
+	FWeldEdgeSequence EdgeWelder(Mesh, ConcatenatedDiscardSpan, ConcatenatedKeptSpan);
+	EdgeWelder.bAllowIntermediateTriangleDeletion = true;
+	EdgeWelder.bAllowFailedMerge = true;
+	EdgeWelder.InterpolationT = InterpolationT;
+
+	FWeldEdgeSequence::EWeldResult Result = EdgeWelder.Weld();
+	if (EdgeWelder.UnmergedEdgePairsOut.Num() != 0)
+	{
+		bHaveSeam = true;
+	}
+	if (Result != FWeldEdgeSequence::EWeldResult::Ok)
+	{
+		bAllSucceeded = false;
 	}
 
+	if (CurrentMesh->TriangleCount() == 0)
+	{
+		GetToolManager()->DisplayMessage(LOCTEXT("WeldEdgesWouldDeleteAll",
+			"Could not weld current selection because doing so would discard entire mesh."), 
+			EToolMessageLevel::UserWarning);
+
+		// Use our change tracker to undo what we've done
+		ChangeTracker.EndChange()->Apply(CurrentMesh.Get(), /*bRevert*/ true);
+		// Update so spatial doesn't complain about mismatched changestamps. 
+		UpdateFromCurrentMesh(false);
+		
+		// The topology didn't actually change, but unfortunately the eids it stores in its spans
+		//  are now invalid, and we need to update those. We do this update ourselves (rather than
+		//  passing true to UpdateFromCurrentMesh above) so that we can keep the same extra corners
+		//  and therefore same selection.
+		TSet<int32> ExtraCorners = Topology->GetCurrentExtraCornerVids();
+		Topology->RebuildTopologyWithSpecificExtraCorners(ExtraCorners);
+		return;
+	}
+
+	FText TransactionName = LOCTEXT("PolyMeshWeldEdgeChange", "Weld Edges");
+	GetToolManager()->BeginUndoTransaction(TransactionName);
 	EmitCurrentMeshChangeAndUpdate(TransactionName, ChangeTracker.EndChange(), FGroupTopologySelection());
 
 	// Now that the topology is updated, set the new selection
 	FGroupTopologySelection NewSelection;
 	TSet<int32> SelectedEids;
-	for (int32 Eid : KeptEids)
+	for (int32 Eid : ConcatenatedKeptEids)
 	{
 		if (Mesh->IsEdge(Eid) && !SelectedEids.Contains(Eid))
 		{
@@ -2411,7 +2524,7 @@ void UEditMeshPolygonsTool::ApplyWeldEdges()
 		}
 	}
 	// Seems possible to end up with an empty selection if we welded edges of the same group,
-	//  so the new edge is not a group boundary.
+	//  so the new edge is not a group boundary, or if we ended up collapsing things.
 	if (!NewSelection.IsEmpty())
 	{
 		SelectionMechanic->SetSelection(NewSelection);
@@ -2419,9 +2532,7 @@ void UEditMeshPolygonsTool::ApplyWeldEdges()
 
 	if (bHaveSeam)
 	{
-		GetToolManager()->DisplayMessage(LOCTEXT("OnWeldEdgesCompletedSeamsRemain", "Warning: welding incomplete because it would create "
-			"invalid geometry (attached non manifold edge or duplicate triangle). Seam still exists at weld "
-			"location. Modify attached triangles and retry, or undo."), EToolMessageLevel::UserWarning);
+		GetToolManager()->DisplayMessage(WeldIncompleteMessage, EToolMessageLevel::UserWarning);
 	}
 	else if (!bAllSucceeded)
 	{
@@ -2432,6 +2543,129 @@ void UEditMeshPolygonsTool::ApplyWeldEdges()
 	GetToolManager()->EndUndoTransaction();
 }
 
+void UEditMeshPolygonsTool::ApplyWeldVertices(double InterpolationT)
+{
+	using namespace EditMeshPolygonsToolLocals;
+
+	FGroupTopologySelection CurrentSelection = SelectionMechanic->GetActiveSelection();
+
+	TArray<int32> CornerIDs = CurrentSelection.SelectedCornerIDs.Array();
+	if (CornerIDs.Num() != 2)
+	{
+		return;
+	}
+
+	FDynamicMeshChangeTracker ChangeTracker(CurrentMesh.Get());
+	ChangeTracker.BeginChange();
+
+	// See if there's a group edge between the two selected corners. If there is, the
+	//  user was probably expecting to collapse the group edge.
+	TSet<int32> GroupEdges;
+	ConvertCornerSelectionToGroupEdgeSelection(Topology.Get(), CurrentSelection.SelectedCornerIDs, GroupEdges);
+	if (GroupEdges.Num() != 0)
+	{
+		GetToolManager()->BeginUndoTransaction(CollapseEdgeTransactionLabel);
+		CollapseGroupEdges(GroupEdges, ChangeTracker);
+		GetToolManager()->EndUndoTransaction();
+		return;
+	}
+	// Othewise do the operation
+	
+	int32 KeptVid = Topology->GetCornerVertexID(CornerIDs[1]);
+	int32 DiscardedVid = Topology->GetCornerVertexID(CornerIDs[0]);
+
+	CurrentMesh->EnumerateVertexTriangles(DiscardedVid, [&ChangeTracker](int32 Tid)
+	{
+		ChangeTracker.SaveTriangle(Tid, true);
+	});
+	CurrentMesh->EnumerateVertexTriangles(KeptVid, [&ChangeTracker](int32 Tid)
+	{
+		ChangeTracker.SaveTriangle(Tid, true);
+	});
+
+	// Helper used when we can't weld, but choose to move the verts to the destination instead
+	auto MoveToDestination = [this, KeptVid, DiscardedVid, InterpolationT]()
+	{
+		FVector3d Destination = Lerp(CurrentMesh->GetVertex(KeptVid), CurrentMesh->GetVertex(DiscardedVid), InterpolationT);
+		CurrentMesh->SetVertex(DiscardedVid, Destination);
+		CurrentMesh->SetVertex(KeptVid, Destination);
+	};
+
+	FDynamicMesh3::FMergeVerticesInfo MergeInfo;
+	FDynamicMesh3::FMergeVerticesOptions Options;
+	Options.bAllowNonBoundaryBowtieCreation = bAllowBowtieWeldAtInternalVertex;
+	EMeshResult Result = CurrentMesh->MergeVertices(KeptVid, DiscardedVid, InterpolationT, Options, MergeInfo);
+
+	if (Result == EMeshResult::Failed_CollapseTriangle
+		|| Result == EMeshResult::Failed_CollapseQuad
+		|| Result == EMeshResult::Failed_FoundDuplicateTriangle)
+	{
+		bool bSuccessful = RemoveEdgeTrisIfNotLast(*CurrentMesh, CurrentMesh->FindEdge(KeptVid, DiscardedVid));
+		if (!bSuccessful)
+		{
+			GetToolManager()->DisplayMessage(LOCTEXT("WeldVerticesCannotDeleteAll",
+				"Could not weld vertices because it would delete remainder of mesh."), EToolMessageLevel::UserWarning);
+			return;
+		}
+	}
+	// Align with behavior in weld and collapse when we're unable to weld due to topology.
+	else if (Result == EMeshResult::Failed_InvalidNeighbourhood)
+	{
+		if (CurrentMesh->FindEdge(KeptVid, DiscardedVid) != IndexConstants::InvalidID)
+		{
+			// Collapse case: refuse to collapse
+			GetToolManager()->DisplayMessage(LOCTEXT("WeldVerticesCollapseInvalidTopology",
+				"Could not weld vertices because the collapse would create an edge with more than "
+				"two triangles (non-manifold geometry)."), EToolMessageLevel::UserWarning);
+			return;
+		}
+		else
+		{
+			// Weld case: move to destination and complain
+			MoveToDestination();
+			GetToolManager()->DisplayMessage(WeldIncompleteMessage, EToolMessageLevel::UserWarning);
+		}
+	}
+	else if (Result == EMeshResult::Failed_WouldCreateBowtie)
+	{
+		MoveToDestination();
+		GetToolManager()->DisplayMessage(LOCTEXT("WeldVerticesDisallowInternalBowtie",
+			"Could not weld vertices because it would create a non-boundary edge bowtie. Vertices "
+			"were moved to their destination without actually welding. Set "
+			"modeling.PolyEdit.AllowWeldInternalBowtie to true to allow a true weld."), 
+			EToolMessageLevel::UserWarning);
+	}
+	else if (Result == EMeshResult::Failed_NotABoundaryEdge)
+	{
+		// This happens if a user is trying to weld internal edges by successive internal vertices.
+		//  Handle this the same way as the other weld failure.
+		MoveToDestination();
+		GetToolManager()->DisplayMessage(WeldIncompleteMessage, EToolMessageLevel::UserWarning);
+	}
+	else if (!ensure(Result == EMeshResult::Ok))
+	{
+		GetToolManager()->DisplayMessage(LOCTEXT("WeldVerticesGenericFailure",
+			"Could not weld vertices."), EToolMessageLevel::UserWarning);
+		return;
+	}
+
+
+	FText TransactionName = LOCTEXT("PolyMeshWeldEdgeChange", "Weld Vertices");
+	GetToolManager()->BeginUndoTransaction(TransactionName);
+	EmitCurrentMeshChangeAndUpdate(TransactionName, ChangeTracker.EndChange(), FGroupTopologySelection());
+
+	// Now that the topology is updated, set the new selection
+	int32 RemainingCornerID = Topology->GetCornerIDFromVertexID(KeptVid);
+	if (RemainingCornerID != IndexConstants::InvalidID)
+	{
+		FGroupTopologySelection NewSelection;
+		NewSelection.SelectedCornerIDs.Add(RemainingCornerID);
+		SelectionMechanic->SetSelection(NewSelection);
+	}
+
+	GetToolManager()->EndUndoTransaction();
+}
+
 void UEditMeshPolygonsTool::ApplyStraightenEdges()
 {
 	if (BeginMeshEdgeEditChange() == false)
@@ -2897,13 +3131,6 @@ void UEditMeshPolygonsTool::ApplyCollapseEdge()
 {
 	using namespace EditMeshPolygonsToolLocals;
 
-	FText PartialCollapseFailure = LOCTEXT("OnCollapseFailures", "Some edges could not be collapsed");
-
-	FDynamicMesh3::FCollapseEdgeOptions CollapseOptions;
-	CollapseOptions.bAllowHoleCollapse = true;
-	CollapseOptions.bAllowCollapsingInternalEdgeWithBoundaryVertices = true;
-	CollapseOptions.bAllowTetrahedronCollapse = true;
-
 	FDynamicMesh3* Mesh = CurrentMesh.Get();
 	FGroupTopologySelection ActiveSelection = SelectionMechanic->GetActiveSelection();
 	if (ActiveSelection.IsEmpty())
@@ -2912,8 +3139,7 @@ void UEditMeshPolygonsTool::ApplyCollapseEdge()
 		return;
 	}
 
-	FText TransactionLabel = LOCTEXT("PolyMeshCollapseChange", "Collapse Edges");
-	GetToolManager()->BeginUndoTransaction(TransactionLabel);
+	GetToolManager()->BeginUndoTransaction(CollapseEdgeTransactionLabel);
 
 	FDynamicMeshChangeTracker ChangeTracker(Mesh);
 	ChangeTracker.BeginChange();
@@ -2927,9 +3153,12 @@ void UEditMeshPolygonsTool::ApplyCollapseEdge()
 		int32 NumCompleted = RetriangulateGroups(Mesh, Topology.Get(), ActiveSelection.SelectedGroupIDs, ChangeTracker);
 		if (NumCompleted != ActiveSelection.SelectedGroupIDs.Num())
 		{
-			GetToolManager()->DisplayMessage(PartialCollapseFailure, EToolMessageLevel::UserWarning);
+			GetToolManager()->DisplayMessage(PartialCollapseFailureMessage, EToolMessageLevel::UserWarning);
 			// continue on and try to collapse the boundary
 		}
+		// After retriangulation, our topology object may have incorrect eids (if the group was on the boundary,
+		//  so the edges got removed during triangle deletion and then recreated). So we have to update it.
+		Topology->RebuildTopology();
 
 		for (int32 GroupID : ActiveSelection.SelectedGroupIDs)
 		{
@@ -2944,6 +3173,23 @@ void UEditMeshPolygonsTool::ApplyCollapseEdge()
 		ConvertCornerSelectionToGroupEdgeSelection(Topology.Get(), ActiveSelection.SelectedCornerIDs, GroupEdgesToCollapse);
 	}
 
+	CollapseGroupEdges(GroupEdgesToCollapse, ChangeTracker);
+
+	GetToolManager()->EndUndoTransaction();
+}
+
+void UEditMeshPolygonsTool::CollapseGroupEdges(TSet<int32>& GroupEdgesToCollapse, 
+	UE::Geometry::FDynamicMeshChangeTracker& ChangeTracker)
+{
+	using namespace EditMeshPolygonsToolLocals;
+
+	FDynamicMesh3* Mesh = CurrentMesh.Get();
+
+	FDynamicMesh3::FCollapseEdgeOptions CollapseOptions;
+	CollapseOptions.bAllowHoleCollapse = true;
+	CollapseOptions.bAllowCollapsingInternalEdgeWithBoundaryVertices = true;
+	CollapseOptions.bAllowTetrahedronCollapse = true;
+
 	TSet<int32> EidsToCollapse;
 	for (int32 GroupEdgeID : GroupEdgesToCollapse)
 	{
@@ -2966,9 +3212,9 @@ void UEditMeshPolygonsTool::ApplyCollapseEdge()
 		ComponentEids.Add(Eid);
 		FMeshConnectedComponents::GrowToConnectedEdges(*Mesh, { Eid }, ComponentEids, &TempQueue,
 			[&EidsToCollapse](int32 CurrentEid, int32 NeighborEid)
-		{
-			return EidsToCollapse.Contains(NeighborEid);
-		});
+			{
+				return EidsToCollapse.Contains(NeighborEid);
+			});
 		PartitionedEids.Append(ComponentEids);
 	}
 
@@ -2986,7 +3232,7 @@ void UEditMeshPolygonsTool::ApplyCollapseEdge()
 
 		// Unfiltered because vids will disappear in subsequent collapses
 		TSet<int32> UnfilteredVidsToMove;
-		
+
 		for (int32 Eid : Component)
 		{
 			// Some edges might be collapsed away by other collapses
@@ -3000,8 +3246,20 @@ void UEditMeshPolygonsTool::ApplyCollapseEdge()
 			ChangeTracker.SaveVertexOneRingTriangles(EdgeVids.B, true);
 			FDynamicMesh3::FEdgeCollapseInfo CollapseInfo;
 			EMeshResult Result = Mesh->CollapseEdge(EdgeVids.A, EdgeVids.B, CollapseOptions, CollapseInfo);
-			
-			if (Result == EMeshResult::Ok)
+
+			// Certain collapses of isolated triangles/quads are not currently allowed by CollapseEdge,
+			//  but we allow them if the user asks for them.
+			if (Result == EMeshResult::Failed_CollapseTriangle
+				|| Result == EMeshResult::Failed_CollapseQuad
+				|| Result == EMeshResult::Failed_FoundDuplicateTriangle)
+			{
+				bAllCollapsesSuccessful = RemoveEdgeTrisIfNotLast(*CurrentMesh, Eid) && bAllCollapsesSuccessful;
+			}
+			// We could also check for EMeshResult::InvalidTopology and do the "move with seam"
+			//  approach we do for welding, but it seems like it would be harder to notice this
+			//  for collapses because the degenerate triangles are harder to find than open boundaries.
+			//  So for now we won't fake a collapse in that case.
+			else if (Result == EMeshResult::Ok)
 			{
 				UnfilteredVidsToMove.Add(CollapseInfo.KeptVertex);
 			}
@@ -3023,17 +3281,19 @@ void UEditMeshPolygonsTool::ApplyCollapseEdge()
 
 	if (!bAllCollapsesSuccessful)
 	{
-		GetToolManager()->DisplayMessage(PartialCollapseFailure, EToolMessageLevel::UserWarning);
+		GetToolManager()->DisplayMessage(PartialCollapseFailureMessage, EToolMessageLevel::UserWarning);
 	}
 
-	EmitCurrentMeshChangeAndUpdate(TransactionLabel, ChangeTracker.EndChange(), FGroupTopologySelection());
+	EmitCurrentMeshChangeAndUpdate(CollapseEdgeTransactionLabel, ChangeTracker.EndChange(), FGroupTopologySelection());
 
 	// Now that the topology is updated, we can get the new corner id's to
 	//  set the new selection.
 	FGroupTopologySelection NewSelection;
 	for (int32 Vid : NewSelectionVids)
 	{
-		if (!ensure(Mesh->IsVertex(Vid)))
+		// Even though we filtered each component, it's possible for one component's collapses to indirectly
+		//  destroy verts in another, hence the check here.
+		if (!Mesh->IsVertex(Vid))
 		{
 			continue;
 		}
@@ -3049,8 +3309,6 @@ void UEditMeshPolygonsTool::ApplyCollapseEdge()
 	{
 		SelectionMechanic->SetSelection(NewSelection);
 	}
-
-	GetToolManager()->EndUndoTransaction();
 }
 
 void UEditMeshPolygonsTool::ApplySplitSingleEdge()
@@ -3084,7 +3342,7 @@ void UEditMeshPolygonsTool::ApplySplitSingleEdge()
 				NewSelection.SelectedGroupIDs.Add(SplitInfo.NewTriangles.A);
 				if (SplitInfo.NewTriangles.B != FDynamicMesh3::InvalidID)
 				{
-					NewSelection.SelectedGroupIDs.Add(SplitInfo.NewTriangles.A);
+					NewSelection.SelectedGroupIDs.Add(SplitInfo.NewTriangles.B);
 				}
 			}
 		}
diff --git a/Engine/Plugins/Runtime/MeshModelingToolset/Source/MeshModelingTools/Public/EditMeshPolygonsTool.h b/Engine/Plugins/Runtime/MeshModelingToolset/Source/MeshModelingTools/Public/EditMeshPolygonsTool.h
--- a/Engine/Plugins/Runtime/MeshModelingToolset/Source/MeshModelingTools/Public/EditMeshPolygonsTool.h
+++ b/Engine/Plugins/Runtime/MeshModelingToolset/Source/MeshModelingTools/Public/EditMeshPolygonsTool.h
@@ -20,6 +20,7 @@
 
 PREDECLARE_GEOMETRY(class FGroupTopology);
 PREDECLARE_GEOMETRY(struct FGroupTopologySelection);
+PREDECLARE_GEOMETRY(class FDynamicMeshChangeTracker);
 
 struct FSlateBrush;
 class UCombinedTransformGizmo;
@@ -81,7 +82,7 @@ public:
 	bool bShowSelectableCorners = true;
 
 	/** When true, allows the transform gizmo to be rendered */
-	UPROPERTY(EditAnywhere, Category = Options)
+	UPROPERTY(EditAnywhere, Category = Gizmo)
 	bool bGizmoVisible = true;
 
 	/** Determines whether, on selection changes, the gizmo's rotation is taken from the object transform, or from the geometry
@@ -129,6 +130,7 @@ enum class EEditMeshPolygonsToolActions
 
 	CollapseEdge,
 	WeldEdges,
+	WeldEdgesCentered,
 	StraightenEdge,
 	FillHole,
 	BridgeEdges,
@@ -402,8 +404,12 @@ class MESHMODELINGTOOLS_API UEditMeshPolygonsToolEdgeActions : public UEditMeshP
 {
 	GENERATED_BODY()
 public:
-	/** Merge selected edges, moving the first edge to the second */
-	UFUNCTION(CallInEditor, Category = EdgeEdits, meta = (DisplayName = "Weld", DisplayPriority = 1))
+	/** Merge selected boundary edges, centering the result */
+	UFUNCTION(CallInEditor, Category = EdgeEdits, meta = (DisplayName = "Weld", DisplayPriority = 0))
+	void WeldCentered() { PostAction(EEditMeshPolygonsToolActions::WeldEdgesCentered); }
+
+	/** Merge selected boundary edges, moving the first edge to the second */
+	UFUNCTION(CallInEditor, Category = EdgeEdits, meta = (DisplayName = "Weld To", DisplayPriority = 1))
 	void Weld() { PostAction(EEditMeshPolygonsToolActions::WeldEdges); }
 
 	/** Make each selected polygroup edge follow a straight path between its endpoints */
@@ -445,8 +451,12 @@ class MESHMODELINGTOOLS_API UEditMeshPolygonsToolEdgeActions_Triangles : public
 {
 	GENERATED_BODY()
 public:
-	/** Merge selected edges, moving the first edge to the second */
-	UFUNCTION(CallInEditor, Category = EdgeEdits, meta = (DisplayName = "Weld", DisplayPriority = 1))
+	/** Merge selected boundary edges, centering the result */
+	UFUNCTION(CallInEditor, Category = EdgeEdits, meta = (DisplayName = "Weld", DisplayPriority = 0))
+	void WeldCentered() { PostAction(EEditMeshPolygonsToolActions::WeldEdgesCentered); }
+
+	/** Merge selected boundary edges, moving the first edge to the second */
+	UFUNCTION(CallInEditor, Category = EdgeEdits, meta = (DisplayName = "Weld To", DisplayPriority = 1))
 	void Weld() { PostAction(EEditMeshPolygonsToolActions::WeldEdges); }
 
 	/** Fill the adjacent hole for any selected boundary edges */
@@ -746,6 +756,12 @@ protected:
 	friend class UEditMeshPolygonsActionModeToolBuilder;
 	TUniqueFunction<void(UEditMeshPolygonsTool*)> PostSetupFunction;
 	void SetToSelectionModeInterface();
+
+private:
+	void ApplyWeldEdges(double InterpolationT);
+	void ApplyWeldVertices(double InterpolationT);
+	void CollapseGroupEdges(TSet<int32>& GroupEdgesToCollapse, 
+		UE::Geometry::FDynamicMeshChangeTracker& ChangeTracker);
 };
 
 
