MeshModelingTools: Fix LOCTEXT conflict in PolyEd.

diff --git a/Engine/Plugins/Runtime/MeshModelingToolset/Source/MeshModelingTools/Private/EditMeshPolygonsTool.cpp b/Engine/Plugins/Runtime/MeshModelingToolset/Source/MeshModelingTools/Private/EditMeshPolygonsTool.cpp
--- a/Engine/Plugins/Runtime/MeshModelingToolset/Source/MeshModelingTools/Private/EditMeshPolygonsTool.cpp
+++ b/Engine/Plugins/Runtime/MeshModelingToolset/Source/MeshModelingTools/Private/EditMeshPolygonsTool.cpp
@@ -2650,7 +2650,7 @@ void UEditMeshPolygonsTool::ApplyWeldVertices(double InterpolationT)
 	}
 
 
-	FText TransactionName = LOCTEXT("PolyMeshWeldEdgeChange", "Weld Vertices");
+	FText TransactionName = LOCTEXT("PolyMeshWeldVerticesChange", "Weld Vertices");
 	GetToolManager()->BeginUndoTransaction(TransactionName);
 	EmitCurrentMeshChangeAndUpdate(TransactionName, ChangeTracker.EndChange(), FGroupTopologySelection());
 
