Modeling Tools: Update MeshElementSelection deletion of polygroup edge to behave like same function in PolyEdit by merging adjoining groups

diff --git a/Engine/Plugins/Runtime/MeshModelingToolset/Source/MeshModelingTools/Private/Commands/DeleteGeometrySelectionCommand.cpp b/Engine/Plugins/Runtime/MeshModelingToolset/Source/MeshModelingTools/Private/Commands/DeleteGeometrySelectionCommand.cpp
--- a/Engine/Plugins/Runtime/MeshModelingToolset/Source/MeshModelingTools/Private/Commands/DeleteGeometrySelectionCommand.cpp
+++ b/Engine/Plugins/Runtime/MeshModelingToolset/Source/MeshModelingTools/Private/Commands/DeleteGeometrySelectionCommand.cpp
@@ -6,9 +6,11 @@
 #include "DynamicMesh/DynamicMesh3.h"
 #include "DynamicMesh/DynamicMeshChangeTracker.h"
 #include "DynamicMeshEditor.h"
+#include "FaceGroupUtil.h"
 #include "Changes/MeshChange.h"
 #include "Selections/GeometrySelectionUtil.h"
 #include "Selection/DynamicMeshSelector.h"
+#include "Selections/MeshConnectedComponents.h"
 
 #include UE_INLINE_GENERATED_CPP_BY_NAME(DeleteGeometrySelectionCommand)
 
@@ -62,30 +64,78 @@ void UDeleteGeometrySelectionCommand::ExecuteCommandForSelection(UGeometrySelect
 		return;
 	}
 
-	bool bTrackChanges = SelectionArgs->HasTransactionsAPI();
+	const bool bTrackChanges = SelectionArgs->HasTransactionsAPI();
 	TUniquePtr<FDynamicMeshChange> DynamicMeshChange;	// only initialized if bTrackChanges == true
 
 	// apply the delete operation
 	MeshObject->EditMesh([&](FDynamicMesh3& EditMesh)
 	{
-		// build list of triangles from whatever the selection contains
-		TSet<int32> TriangleList;
-		//UE::Geometry::FPolygroupSet UsePolygroupSet = ...;		// need to support this eventually
-		UE::Geometry::EnumerateSelectionTriangles(*Selection, EditMesh,
-			[&](int32 TriangleID) { TriangleList.Add(TriangleID); });
-
-		// mark triangles for change
 		FDynamicMeshChangeTracker ChangeTracker(&EditMesh);
-		if (bTrackChanges)
+
+		// handles the case of deleting polygroup edge by merging adjoining groups to match behavior in PolyEdit
+		if (SelectionArgs->TopologyMode == UE::Geometry::EGeometryTopologyType::Polygroup
+			&& SelectionArgs->ElementType == UE::Geometry::EGeometryElementType::Edge)
 		{
-			ChangeTracker.BeginChange();
-			ChangeTracker.SaveTriangles(TriangleList, true);
+			FMeshConnectedComponents Components(&EditMesh);
+
+			// retrieve all selected edges
+			TSet<int32> EdgeIDs;
+			UE::Geometry::EnumeratePolygroupSelectionEdges(*Selection, EditMesh, FPolygroupSet(&EditMesh), 
+				[&](const int32 EdgeID) { EdgeIDs.Add(EdgeID); });
+
+			// similar but simplified version of work done in EnumeratePolygroupSelectionTriangles
+			// retrieves the TriIDs of the triangles adjacent to all edges in the PolyEdge as they will all be merged into one Polygroup
+			TSet<int32> SeedTriangleIDs;
+			for (int32 Edge : EdgeIDs.Array())
+			{
+					FIndex2i AdjacentTriangles = EditMesh.GetEdgeT(Edge);
+					EdgeIDs.Add(Edge);
+					SeedTriangleIDs.Add(AdjacentTriangles.A);
+					if (AdjacentTriangles.B != FDynamicMesh3::InvalidID)
+					{
+						SeedTriangleIDs.Add(AdjacentTriangles.B);
+					}
+			}
+
+			// retrieve the rest of the connected components which will be in the merged PolyGroup
+			Components.FindTrianglesConnectedToSeeds(SeedTriangleIDs.Array(), [&EditMesh, &EdgeIDs](const int32 Tri0, const int32 Tri1)
+			{
+				return EditMesh.GetTriangleGroup(Tri0) == EditMesh.GetTriangleGroup(Tri1) || EdgeIDs.Contains(EditMesh.FindEdgeFromTriPair(Tri0,Tri1));
+			});
+
+			// mark triangles for change
+			if (bTrackChanges)
+			{
+				ChangeTracker.BeginChange();
+	
+				for (FMeshConnectedComponents::FComponent& Component : Components.Components)
+				{
+					ChangeTracker.SaveTriangles(Component.Indices, true);
+					int32 NewGroupID = EditMesh.GetTriangleGroup(Component.Indices[0]);
+					FaceGroupUtil::SetGroupID(EditMesh, Component.Indices, NewGroupID);
+				}
+			}
 		}
-
-		// actually delete them
-		FDynamicMeshEditor Editor(&EditMesh);
-		Editor.RemoveTriangles(TriangleList.Array(), true);
-
+		else
+		{
+			// build list of triangles from whatever the selection contains
+			TSet<int32> TriangleList;
+			//UE::Geometry::FPolygroupSet UsePolygroupSet = ...;		// need to support this eventually
+			UE::Geometry::EnumerateSelectionTriangles(*Selection, EditMesh,
+			[&](const int32 TriangleID) { TriangleList.Add(TriangleID); });
+
+			// mark triangles for change
+			if (bTrackChanges)
+			{
+				ChangeTracker.BeginChange();
+				ChangeTracker.SaveTriangles(TriangleList, true);
+			}
+
+			// actually delete them
+			FDynamicMeshEditor Editor(&EditMesh);
+			Editor.RemoveTriangles(TriangleList.Array(), true);
+		}
+		
 		// extract the change record
 		if (bTrackChanges)
 		{
diff --git a/Engine/Plugins/Runtime/MeshModelingToolset/Source/ModelingComponents/Private/Selection/GeometrySelectionManager.cpp b/Engine/Plugins/Runtime/MeshModelingToolset/Source/ModelingComponents/Private/Selection/GeometrySelectionManager.cpp
--- a/Engine/Plugins/Runtime/MeshModelingToolset/Source/ModelingComponents/Private/Selection/GeometrySelectionManager.cpp
+++ b/Engine/Plugins/Runtime/MeshModelingToolset/Source/ModelingComponents/Private/Selection/GeometrySelectionManager.cpp
@@ -1443,6 +1443,8 @@ void UGeometrySelectionManager::ExecuteSelectionCommand(UGeometrySelectionEditCo
 		// q: we could clear the selection here, and pass the Handle a copy. Perhaps safer?
 		UInteractiveCommandResult* ResultPtr = nullptr;
 		SelectionArguments->SelectionHandle = FGeometrySelectionHandle{ Target->Selector->GetIdentifier(), &Target->Selection, Target->Selector.Get() };
+		SelectionArguments->ElementType = SelectionElementType;
+		SelectionArguments->TopologyMode = GetSelectionTopologyType();
 		SelectionArguments->SetTransactionsAPI(TransactionsAPI);
 		Command->ExecuteCommand(SelectionArguments, &ResultPtr);
 
diff --git a/Engine/Plugins/Runtime/MeshModelingToolset/Source/ModelingComponents/Public/Selection/SelectionEditInteractiveCommand.h b/Engine/Plugins/Runtime/MeshModelingToolset/Source/ModelingComponents/Public/Selection/SelectionEditInteractiveCommand.h
--- a/Engine/Plugins/Runtime/MeshModelingToolset/Source/ModelingComponents/Public/Selection/SelectionEditInteractiveCommand.h
+++ b/Engine/Plugins/Runtime/MeshModelingToolset/Source/ModelingComponents/Public/Selection/SelectionEditInteractiveCommand.h
@@ -15,6 +15,8 @@ class MODELINGCOMPONENTS_API UGeometrySelectionEditCommandArguments : public UIn
 	GENERATED_BODY()
 public:
 	FGeometrySelectionHandle SelectionHandle;
+	UE::Geometry::EGeometryElementType ElementType;
+	UE::Geometry::EGeometryTopologyType TopologyMode;
 
 	bool IsSelectionEmpty() const
 	{
