GeometryScript: Added new function to perform texel density UV scaling operations, allowing mesh UV islands to be scaled to match a consistent texel density over the whole mesh.

diff --git a/Engine/Plugins/Editor/UVEditor/Source/UVEditor/Private/UVEditorDistortionVisualization.cpp b/Engine/Plugins/Editor/UVEditor/Source/UVEditor/Private/UVEditorDistortionVisualization.cpp
--- a/Engine/Plugins/Editor/UVEditor/Source/UVEditor/Private/UVEditorDistortionVisualization.cpp
+++ b/Engine/Plugins/Editor/UVEditor/Source/UVEditor/Private/UVEditorDistortionVisualization.cpp
@@ -17,7 +17,7 @@
 #include "Math/UnitConversion.h"
 #include "Parameterization/MeshUDIMClassifier.h"
 #include "UDIMUtilities.h"
-#include "Math/UVMetrics.h"
+#include "Parameterization/UVMetrics.h"
 
 #include UE_INLINE_GENERATED_CPP_BY_NAME(UVEditorDistortionVisualization)
 
diff --git a/Engine/Plugins/Editor/UVEditor/Source/UVEditor/UVEditor.Build.cs b/Engine/Plugins/Editor/UVEditor/Source/UVEditor/UVEditor.Build.cs
--- a/Engine/Plugins/Editor/UVEditor/Source/UVEditor/UVEditor.Build.cs
+++ b/Engine/Plugins/Editor/UVEditor/Source/UVEditor/UVEditor.Build.cs
@@ -63,7 +63,8 @@ public class UVEditor : ModuleRules
 				"UVEditorTools",
 				"UVEditorToolsEditorOnly",
 				"WorkspaceMenuStructure",
-				"TextureUtilitiesCommon"
+				"TextureUtilitiesCommon",
+				"GeometryAlgorithms"
 				// ... add private dependencies that you statically link with here ...	
 			}
 			);
diff --git a/Engine/Plugins/Editor/UVEditor/Source/UVEditorTools/Private/Math/UVMetrics.cpp b/Engine/Plugins/Editor/UVEditor/Source/UVEditorTools/Private/Math/UVMetrics.cpp
deleted file mode 100644
diff --git a/Engine/Plugins/Editor/UVEditor/Source/UVEditorTools/Private/Operators/UVEditorTexelDensityOp.cpp b/Engine/Plugins/Editor/UVEditor/Source/UVEditorTools/Private/Operators/UVEditorTexelDensityOp.cpp
deleted file mode 100644
diff --git a/Engine/Plugins/Editor/UVEditor/Source/UVEditorTools/Private/UVEditorTexelDensityTool.cpp b/Engine/Plugins/Editor/UVEditor/Source/UVEditorTools/Private/UVEditorTexelDensityTool.cpp
--- a/Engine/Plugins/Editor/UVEditor/Source/UVEditorTools/Private/UVEditorTexelDensityTool.cpp
+++ b/Engine/Plugins/Editor/UVEditor/Source/UVEditorTools/Private/UVEditorTexelDensityTool.cpp
@@ -4,7 +4,7 @@
 
 #include "ContextObjectStore.h"
 #include "MeshOpPreviewHelpers.h" // UMeshOpPreviewWithBackgroundCompute
-#include "Operators/UVEditorTexelDensityOp.h"
+#include "ParameterizationOps/TexelDensityOp.h"
 #include "InputBehaviorSet.h"
 #include "InputRouter.h"
 #include "BaseBehaviors/SingleClickBehavior.h"
@@ -14,7 +14,7 @@
 #include "ToolSetupUtil.h"
 #include "UVEditorUXSettings.h"
 #include "ContextObjects/UVToolContextObjects.h"
-#include "Math/UVMetrics.h"
+#include "Parameterization/UVMetrics.h"
 
 #include UE_INLINE_GENERATED_CPP_BY_NAME(UVEditorTexelDensityTool)
 
diff --git a/Engine/Plugins/Editor/UVEditor/Source/UVEditorTools/Public/Math/UVMetrics.h b/Engine/Plugins/Editor/UVEditor/Source/UVEditorTools/Public/Math/UVMetrics.h
--- a/Engine/Plugins/Editor/UVEditor/Source/UVEditorTools/Public/Math/UVMetrics.h
+++ b/Engine/Plugins/Editor/UVEditor/Source/UVEditorTools/Public/Math/UVMetrics.h
@@ -1,18 +1,8 @@
 // Copyright Epic Games, Inc. All Rights Reserved.
 
 #pragma once
-#include "CoreMinimal.h"
 
-namespace UE::Geometry
-{
-	class FDynamicMesh3;
+#include "Misc/CoreMiscDefines.h"
+UE_DEPRECATED_HEADER(5.5, "Use Parameterization/UVMetrics.h from GeometryAlgorithms instead of UVMetrics.h")
 
-	class UVEDITORTOOLS_API FUVMetrics
-	{
-	public:
-		static double ReedBeta(const UE::Geometry::FDynamicMesh3& Mesh, int32 UVChannel, int32 Tid);
-		static double Sander(const UE::Geometry::FDynamicMesh3& Mesh, int32 UVChannel, int32 Tid, bool bUseL2);
-		static double TexelDensity(const UE::Geometry::FDynamicMesh3& Mesh, int32 UVChannel, int32 Tid, int32 MapSize);
-	};
-
-}
\ No newline at end of file
+#include "Parameterization/UVMetrics.h"
\ No newline at end of file
diff --git a/Engine/Plugins/Editor/UVEditor/Source/UVEditorTools/Public/Operators/UVEditorTexelDensityOp.h b/Engine/Plugins/Editor/UVEditor/Source/UVEditorTools/Public/Operators/UVEditorTexelDensityOp.h
--- a/Engine/Plugins/Editor/UVEditor/Source/UVEditorTools/Public/Operators/UVEditorTexelDensityOp.h
+++ b/Engine/Plugins/Editor/UVEditor/Source/UVEditorTools/Public/Operators/UVEditorTexelDensityOp.h
@@ -2,121 +2,7 @@
 
 #pragma once
 
-#include "CoreMinimal.h"
-#include "Util/ProgressCancel.h"
-#include "ModelingOperators.h"
-#include "InteractiveTool.h"
-#include "DynamicMesh/DynamicMeshAttributeSet.h"
+#include "Misc/CoreMiscDefines.h"
+UE_DEPRECATED_HEADER(5.5, "Use ParameterizationOps/TexelDensityOp.h from ModelingOperators instead of UVEditorTexelDensityOp.h")
 
-#include "UVEditorTexelDensityOp.generated.h"
-
-UENUM()
-enum class ETexelDensityToolMode
-{
-	ApplyToIslands,
-	ApplyToWhole,
-	Normalize
-};
-
-UCLASS()
-class UVEDITORTOOLS_API UUVEditorTexelDensitySettings : public UInteractiveToolPropertySet
-{
-	GENERATED_BODY()
-public:
-
-	UFUNCTION()
-	virtual bool InSamplingMode() const;
-
-	UPROPERTY(EditAnywhere, Category = "TexelDensity", meta = (DisplayName = "Scale Mode", EditCondition = "!InSamplingMode" ))
-	ETexelDensityToolMode TexelDensityMode = ETexelDensityToolMode::ApplyToIslands;
-
-	UPROPERTY(EditAnywhere, Category = "TexelDensity", meta = (DisplayName = "World Units", EditCondition = "TexelDensityMode != ETexelDensityToolMode::Normalize"))
-	float TargetWorldUnits = 100;
-
-	UPROPERTY(EditAnywhere, Category = "TexelDensity", meta = (DisplayName = "Pixels", EditCondition = "TexelDensityMode != ETexelDensityToolMode::Normalize && !InSamplingMode"))
-	float TargetPixelCount = 1024;
-
-	UPROPERTY(EditAnywhere, Category = "TexelDensity", meta = (DisplayName = "Texture Dimensions", EditCondition = "TexelDensityMode != ETexelDensityToolMode::Normalize && !InSamplingMode", LinearDeltaSensitivity = 1000, Delta = 64, SliderExponent = 1, UIMin = 2, UIMax = 16384))
-	float TextureResolution = 1024;
-
-	UPROPERTY(EditAnywhere, Category = "TexelDensity", meta = (DisplayName = "Use UDIMs", EditCondition = "TexelDensityMode != ETexelDensityToolMode::Normalize && !InSamplingMode", TransientToolProperty))
-	bool bEnableUDIMLayout = false;
-};
-
-
-namespace UE
-{
-	namespace Geometry
-	{
-		class FDynamicMesh3;
-		class FDynamicMeshUVPacker;
-
-		enum class EUVTexelDensityOpModes
-		{
-			ScaleGlobal = 0,
-			ScaleIslands = 1,
-			Normalize = 2
-		};
-
-		class UVEDITORTOOLS_API FUVEditorTexelDensityOp : public FDynamicMeshOperator
-		{
-		public:
-			virtual ~FUVEditorTexelDensityOp() {}
-
-			// inputs
-			TSharedPtr<FDynamicMesh3, ESPMode::ThreadSafe> OriginalMesh;
-
-			EUVTexelDensityOpModes TexelDensityMode = EUVTexelDensityOpModes::ScaleGlobal;
-
-			int UVLayerIndex = 0;
-			int TextureResolution = 1024;
-			int TargetWorldSpaceMeasurement = 100;
-			int TargetPixelCountMeasurement = 1024;
-
-			bool bMaintainOriginatingUDIM = false;
-			TOptional<TSet<int32>> Selection;
-			TOptional<TMap<int32, int32>> TextureResolutionPerUDIM;
-			FVector2f UVTranslation = FVector2f::Zero();
-
-			void SetTransform(const FTransformSRT3d& Transform);
-
-			//
-			// FDynamicMeshOperator implementation
-			// 
-
-			virtual void CalculateResult(FProgressCancel* Progress) override;
-
-		protected:
-			void ExecutePacker(FDynamicMeshUVPacker& Packer);			
-			void ScaleMeshSubRegionByDensity(FDynamicMeshUVOverlay* UVLayer, const TArray<int32>& Tids, TSet<int32>& UVElements, int32 TileResolution);
-
-
-		};
-
-	} // end namespace UE::Geometry
-} // end namespace UE
-
-/**
- * Can be hooked up to a UMeshOpPreviewWithBackgroundCompute to perform UV Layout operations.
- *
- * Inherits from UObject so that it can hold a strong pointer to the settings UObject, which
- * needs to be a UObject to be displayed in the details panel.
- */
-UCLASS()
-class UVEDITORTOOLS_API UUVTexelDensityOperatorFactory : public UObject, public UE::Geometry::IDynamicMeshOperatorFactory
-{
-	GENERATED_BODY()
-
-public:
-	// IDynamicMeshOperatorFactory API
-	virtual TUniquePtr<UE::Geometry::FDynamicMeshOperator> MakeNewOperator() override;
-
-	UPROPERTY()
-	TObjectPtr<UUVEditorTexelDensitySettings> Settings;
-
-	TOptional<TSet<int32>> Selection;
-	TSharedPtr<UE::Geometry::FDynamicMesh3, ESPMode::ThreadSafe> OriginalMesh;
-	TUniqueFunction<int32()> GetSelectedUVChannel = []() { return 0; };
-	FTransform TargetTransform;
-	TOptional<TMap<int32, int32>> TextureResolutionPerUDIM;
-};
\ No newline at end of file
+#include "ParameterizationOps/TexelDensityOp.h"
\ No newline at end of file
diff --git a/Engine/Plugins/Editor/UVEditor/Source/UVEditorTools/Public/UVEditorTexelDensityTool.h b/Engine/Plugins/Editor/UVEditor/Source/UVEditorTools/Public/UVEditorTexelDensityTool.h
--- a/Engine/Plugins/Editor/UVEditor/Source/UVEditorTools/Public/UVEditorTexelDensityTool.h
+++ b/Engine/Plugins/Editor/UVEditor/Source/UVEditorTools/Public/UVEditorTexelDensityTool.h
@@ -10,7 +10,7 @@
 #include "Drawing/UVLayoutPreview.h"
 #include "UVEditorToolAnalyticsUtils.h"
 #include "Selection/UVToolSelectionAPI.h"
-#include "Operators/UVEditorTexelDensityOp.h"
+#include "ParameterizationOps/TexelDensityOp.h"
 #include "DynamicMesh/DynamicMeshAABBTree3.h"
 
 #include "UVEditorTexelDensityTool.generated.h"
diff --git a/Engine/Plugins/Editor/UVEditor/Source/UVEditorTools/UVEditorTools.Build.cs b/Engine/Plugins/Editor/UVEditor/Source/UVEditorTools/UVEditorTools.Build.cs
--- a/Engine/Plugins/Editor/UVEditor/Source/UVEditorTools/UVEditorTools.Build.cs
+++ b/Engine/Plugins/Editor/UVEditor/Source/UVEditorTools/UVEditorTools.Build.cs
@@ -48,6 +48,7 @@ public class UVEditorTools : ModuleRules
 				"DynamicMesh",
 				"GeometryCore",
 				"InputCore", // EKey
+				"GeometryAlgorithms",
 				"ModelingComponents",
 				"ModelingOperators",
 				"TextureUtilitiesCommon"
diff --git a/Engine/Plugins/Runtime/GeometryProcessing/Source/GeometryAlgorithms/Private/Parameterization/UVMetrics.cpp b/Engine/Plugins/Runtime/GeometryProcessing/Source/GeometryAlgorithms/Private/Parameterization/UVMetrics.cpp
new file mode 100644
--- /dev/null
+++ b/Engine/Plugins/Runtime/GeometryProcessing/Source/GeometryAlgorithms/Private/Parameterization/UVMetrics.cpp
@@ -0,0 +1,151 @@
+// Copyright Epic Games, Inc. All Rights Reserved.
+
+#include "Parameterization/UVMetrics.h"
+#include "DynamicMesh/DynamicMesh3.h"
+#include "DynamicMesh/DynamicMeshAttributeSet.h"
+
+using namespace UE::Geometry;
+
+namespace UVMetricsLocals
+{
+	double Triangle2DArea(FVector2f p1, FVector2f p2, FVector2f p3)
+	{
+		return FMath::Abs(((p2.X - p1.X) * (p3.Y - p1.Y) - (p3.X - p1.X) * (p2.Y - p1.Y)) / 2.0);
+	}
+
+	double Triangle3DArea(FVector3d q1, FVector3d q2, FVector3d q3)
+	{
+		FVector3d Q1 = FVector3d(q2.X - q1.X, q2.Y - q1.Y, q2.Z - q1.Z);
+		FVector3d Q2 = FVector3d(q3.X - q1.X, q3.Y - q1.Y, q3.Z - q1.Z);
+
+		FVector3d U = Q1.Cross(Q2);
+
+		return U.Length() / 2.0;
+	}
+
+	double Triangle2DPerimeter(FVector2f p1, FVector2f p2, FVector2f p3)
+	{
+		FVector2f Edge1 = p2 - p1;
+		FVector2f Edge2 = p3 - p1;
+		FVector2f Edge3 = p3 - p2;
+		return Edge1.Length() + Edge2.Length() + Edge3.Length();
+
+	}
+
+	double Triangle3DPerimeter(FVector3d q1, FVector3d q2, FVector3d q3)
+	{
+		FVector3d Edge1 = q2 - q1;
+		FVector3d Edge2 = q3 - q1;
+		FVector3d Edge3 = q3 - q2;
+		return Edge1.Length() + Edge2.Length() + Edge3.Length();
+	}
+
+	FVector2f GetTriangleSingularValues(FVector3d V1, FVector3d V2, FVector3d V3, FVector2f UV1, FVector2f UV2, FVector2f UV3)
+	{
+		FVector2f SingularValues;
+		FVector3d Ss, St;
+		double A;
+
+		A = Triangle2DArea(UV1, UV2, UV3);
+		Ss = (V1 * (UV2.Y - UV3.Y) + V2 * (UV3.Y - UV1.Y) + V3 * (UV1.Y - UV2.Y)) / A;
+		St = (V1 * (UV3.X - UV2.X) + V2 * (UV1.X - UV3.X) + V3 * (UV2.X - UV1.X)) / A;
+
+		double a, b, c;
+		a = Ss.Dot(Ss);
+		b = Ss.Dot(St);
+		c = St.Dot(St);
+
+		double Discriminant = FMath::Sqrt(FMath::Square(a - c) + 4 * FMath::Square(b));
+
+		SingularValues.X = (float)(FMath::Sqrt(0.5 * ((a + c) + Discriminant)));
+		SingularValues.Y = (float)(FMath::Sqrt(0.5 * ((a + c) - Discriminant)));
+
+		return SingularValues;
+	}
+}
+
+
+double FUVMetrics::ReedBeta(const FDynamicMesh3& Mesh, int32 UVChannel, int32 Tid)
+{
+	// Inspiration for the math used to compute the eccentricity comes from here:
+	// https://www.reedbeta.com/blog/conformal-texture-mapping/
+
+	FVector3d V1, V2, V3;
+	FVector2f UV1, UV2, UV3;
+	Mesh.GetTriVertices(Tid, V1, V2, V3);
+	Mesh.Attributes()->GetUVLayer(UVChannel)->GetTriElements(Tid, UV1, UV2, UV3);
+
+	FVector2f SingularValues = UVMetricsLocals::GetTriangleSingularValues(V1, V2, V3, UV1, UV2, UV3);
+	if (TMathUtil<double>::Abs(SingularValues.X) < TMathUtil<double>::Epsilon)
+	{
+		return 0;
+	}
+	else
+	{
+		double Eccentricity = FMath::Sqrt(1 - FMath::Square(SingularValues.Y / SingularValues.X));
+		return Eccentricity;
+	}
+}
+
+double FUVMetrics::Sander(const FDynamicMesh3& Mesh, int32 UVChannel, int32 Tid, bool bUseL2)
+{
+	// These metrics are defined by this paper:
+	// "Texture Mapping Progressive Meshes"
+	// https://dl.acm.org/doi/10.1145/383259.383307
+	// https://hhoppe.com/tmpm.pdf
+
+	FVector3d q1, q2, q3;
+	FVector2f p1, p2, p3;
+	Mesh.GetTriVertices(Tid, q1, q2, q3);
+	Mesh.Attributes()->GetUVLayer(UVChannel)->GetTriElements(Tid, p1, p2, p3);
+
+	double GammaMax, GammaMin, L2_Norm, L_Inf;
+	double L2_Norm_Normalized, L_Inf_Normalized;
+	{
+		FVector3d Ss, St;
+		double A, Aq;
+		A = UVMetricsLocals::Triangle2DArea(p1, p2, p3);
+		Aq = UVMetricsLocals::Triangle3DArea(q1, q2, q3);
+		Ss = (q1 * (p2.Y - p3.Y) + q2 * (p3.Y - p1.Y) + q3 * (p1.Y - p2.Y)) / (2 * A);
+		St = (q1 * (p3.X - p2.X) + q2 * (p1.X - p3.X) + q3 * (p2.X - p1.X)) / (2 * A);
+
+		double a, b, c;
+		a = Ss.Dot(Ss);
+		b = Ss.Dot(St);
+		c = St.Dot(St);
+
+		double D = FMath::Sqrt(FMath::Square(a - c) + 4 * FMath::Square(b));
+
+		GammaMax = FMath::Sqrt(0.5 * ((a + c) + D));
+		GammaMin = FMath::Sqrt(0.5 * ((a + c) - D));
+
+		L2_Norm = FMath::Sqrt(0.5 * (a + c));
+		L_Inf = GammaMax;
+
+		L2_Norm_Normalized = L2_Norm / FMath::Sqrt(Aq / A);
+		L_Inf_Normalized = L_Inf / FMath::Sqrt(Aq / A);
+	}
+	if (bUseL2)
+	{
+		return L2_Norm_Normalized;
+	}
+	else
+	{
+		return L_Inf_Normalized;
+	}
+}
+
+double FUVMetrics::TexelDensity(const FDynamicMesh3& Mesh, int32 UVChannel, int32 Tid, int32 MapSize)
+{
+	// Computes a metric describing the difference between the density of a single UV triangle relative to the desired density
+	FVector3d q1, q2, q3;
+	FVector2f p1, p2, p3;
+	Mesh.GetTriVertices(Tid, q1, q2, q3);
+	Mesh.Attributes()->GetUVLayer(UVChannel)->GetTriElements(Tid, p1, p2, p3);
+
+	double UVArea = UVMetricsLocals::Triangle2DArea(p1, p2, p3);
+	double WorldSpaceArea = UVMetricsLocals::Triangle3DArea(q1, q2, q3);
+	double Density = FMath::Sqrt(UVArea / WorldSpaceArea) * MapSize; // Since we're comparing area ratios here, we want to take the square root to get the linear scaling factor back out.
+																		// Technically Mapsize is Mapsize^2 in terms of pixels per unit area, so we implicitly take it's square root here.
+	return Density;
+}
diff --git a/Engine/Plugins/Runtime/GeometryProcessing/Source/GeometryAlgorithms/Public/Parameterization/UVMetrics.h b/Engine/Plugins/Runtime/GeometryProcessing/Source/GeometryAlgorithms/Public/Parameterization/UVMetrics.h
new file mode 100644
--- /dev/null
+++ b/Engine/Plugins/Runtime/GeometryProcessing/Source/GeometryAlgorithms/Public/Parameterization/UVMetrics.h
@@ -0,0 +1,19 @@
+// Copyright Epic Games, Inc. All Rights Reserved.
+
+#pragma once
+#include "CoreMinimal.h"
+
+namespace UE {
+namespace Geometry {
+
+	class FDynamicMesh3;
+
+	class GEOMETRYALGORITHMS_API FUVMetrics
+	{
+		public:
+			static double ReedBeta(const UE::Geometry::FDynamicMesh3& Mesh, int32 UVChannel, int32 Tid);
+			static double Sander(const UE::Geometry::FDynamicMesh3& Mesh, int32 UVChannel, int32 Tid, bool bUseL2);
+			static double TexelDensity(const UE::Geometry::FDynamicMesh3& Mesh, int32 UVChannel, int32 Tid, int32 MapSize);		
+	};
+}
+}
\ No newline at end of file
diff --git a/Engine/Plugins/Runtime/GeometryScripting/Source/GeometryScriptingCore/Private/MeshUVFunctions.cpp b/Engine/Plugins/Runtime/GeometryScripting/Source/GeometryScriptingCore/Private/MeshUVFunctions.cpp
--- a/Engine/Plugins/Runtime/GeometryScripting/Source/GeometryScriptingCore/Private/MeshUVFunctions.cpp
+++ b/Engine/Plugins/Runtime/GeometryScripting/Source/GeometryScriptingCore/Private/MeshUVFunctions.cpp
@@ -12,6 +12,7 @@
 #include "Parameterization/MeshLocalParam.h"
 #include "TriangleTypes.h"
 #include "XAtlasWrapper.h"
+#include "ParameterizationOps/TexelDensityOp.h"
 
 #include "Async/ParallelFor.h"
 #include "UDynamicMesh.h"
@@ -600,8 +601,84 @@ UDynamicMesh* UGeometryScriptLibrary_MeshUVFunctions::SetMeshUVsFromCylinderProj
 	return TargetMesh;
 }
 
+UDynamicMesh*
+UGeometryScriptLibrary_MeshUVFunctions::ApplyTexelDensityUVScaling(
+	UDynamicMesh* TargetMesh,
+	int UVSetIndex,
+	FGeometryScriptUVTexelDensityOptions Options,
+	FGeometryScriptMeshSelection Selection,
+	UGeometryScriptDebug* Debug)
+{
+	if (TargetMesh == nullptr)
+	{
+		UE::Geometry::AppendError(Debug, EGeometryScriptErrorType::InvalidInputs, LOCTEXT("ApplyTexelDensityUVScaling_InvalidInput", "ApplyTexelDensityUVScaling: TargetMesh is Null"));
+		return TargetMesh;
+	}
+
+
+	TargetMesh->EditMesh([&](FDynamicMesh3& EditMesh) 
+	{
+
+		TSharedPtr<FDynamicMesh3> SourceMesh = MakeShared<FDynamicMesh3>(MoveTemp(EditMesh));
+
+		FUVEditorTexelDensityOp Op;
+		
+		Op.OriginalMesh = SourceMesh;
+
+		switch (Options.TexelDensityMode)
+		{
+		case EGeometryScriptTexelDensityMode::ApplyToIslands:
+			Op.TexelDensityMode = EUVTexelDensityOpModes::ScaleIslands;
+			break;
+		case EGeometryScriptTexelDensityMode::ApplyToWhole:
+			Op.TexelDensityMode = EUVTexelDensityOpModes::ScaleGlobal;
+			break;
+		case EGeometryScriptTexelDensityMode::Normalize:
+			Op.TexelDensityMode = EUVTexelDensityOpModes::Normalize;
+			break;
+		default:
+			ensure(false);
+			break;
+		}
 
 
+		TArray<int32> TriangleSelection;
+		Selection.ConvertToMeshIndexArray(*SourceMesh, TriangleSelection, EGeometryScriptIndexType::Triangle);
+
+		Op.TextureResolution = Options.TextureResolution;
+		Op.TargetWorldSpaceMeasurement = Options.TargetWorldUnits;
+		Op.TargetPixelCountMeasurement = Options.TargetPixelCount;
+
+		Op.UVLayerIndex = UVSetIndex;
+		Op.TextureResolution = Options.TextureResolution;
+		Op.SetTransform(FTransformSRT3d::Identity());
+		Op.bMaintainOriginatingUDIM = Options.bEnableUDIMLayout;
+		if (TriangleSelection.Num() > 0)
+		{
+			Op.Selection = TSet<int32>(TriangleSelection);
+		}
+		if (Options.UDIMResolutions.Num() > 0)
+		{
+			Op.TextureResolutionPerUDIM = Options.UDIMResolutions;
+		}
+
+		Op.CalculateResult(nullptr);
+		if (Op.GetResultInfo().Result == EGeometryResultType::Success)
+		{
+			TUniquePtr<FDynamicMesh3> ResultMesh = Op.ExtractResult();
+			EditMesh = MoveTemp(*ResultMesh);
+		}
+		else
+		{
+			EditMesh = MoveTemp(*SourceMesh);
+			UE::Geometry::AppendError(Debug, EGeometryScriptErrorType::InvalidInputs, LOCTEXT("ApplyTexelDensityUVScaling_ComputeError", "ApplyTexelDensityUVScaling: Error computing result, returning input mesh"));
+		}
+
+	}, EDynamicMeshChangeType::AttributeEdit, EDynamicMeshAttributeChangeFlags::UVs, false);
+
+	return TargetMesh;
+}
+
 UDynamicMesh* UGeometryScriptLibrary_MeshUVFunctions::RecomputeMeshUVs( 
 	UDynamicMesh* TargetMesh, 
 	int UVSetIndex,
diff --git a/Engine/Plugins/Runtime/GeometryScripting/Source/GeometryScriptingCore/Public/GeometryScript/MeshUVFunctions.h b/Engine/Plugins/Runtime/GeometryScripting/Source/GeometryScriptingCore/Public/GeometryScript/MeshUVFunctions.h
--- a/Engine/Plugins/Runtime/GeometryScripting/Source/GeometryScriptingCore/Public/GeometryScript/MeshUVFunctions.h
+++ b/Engine/Plugins/Runtime/GeometryScripting/Source/GeometryScriptingCore/Public/GeometryScript/MeshUVFunctions.h
@@ -140,6 +140,40 @@ public:
 
 
 
+UENUM(BlueprintType)
+enum class EGeometryScriptTexelDensityMode : uint8
+{
+	ApplyToIslands,
+	ApplyToWhole,
+	Normalize
+};
+
+USTRUCT(Blueprintable)
+struct GEOMETRYSCRIPTINGCORE_API FGeometryScriptUVTexelDensityOptions
+{
+	GENERATED_BODY()
+public:
+
+	UPROPERTY(EditAnywhere, BlueprintReadWrite, Category = Options)
+	EGeometryScriptTexelDensityMode TexelDensityMode = EGeometryScriptTexelDensityMode::ApplyToIslands;
+
+	UPROPERTY(EditAnywhere, BlueprintReadWrite, Category = Options)
+	float TargetWorldUnits = 100;
+
+	UPROPERTY(EditAnywhere, BlueprintReadWrite, Category = Options)
+	float TargetPixelCount = 1024;
+
+	UPROPERTY(EditAnywhere, BlueprintReadWrite, Category = Options)
+	float TextureResolution = 1024;
+
+	UPROPERTY(EditAnywhere, BlueprintReadWrite, Category = Options)
+	bool bEnableUDIMLayout = false;
+
+	UPROPERTY(EditAnywhere, BlueprintReadWrite, Category = Options, meta = (DisplayName = "Per UDIM Texture Resolution"))
+	TMap<int32, int32> UDIMResolutions;
+};
+
+
 UCLASS(meta = (ScriptName = "GeometryScript_UVs"))
 class GEOMETRYSCRIPTINGCORE_API UGeometryScriptLibrary_MeshUVFunctions : public UBlueprintFunctionLibrary
 {
@@ -356,6 +390,19 @@ public:
 		UGeometryScriptDebug* Debug = nullptr );
 
 
+	/**
+	* Rescales UVs in the UV Channel for a Mesh to match a specified texel density, described by the options passed in. Supports
+	* processing on a subset of UVs via a non-empty Selection.
+	*/
+	UFUNCTION(BlueprintCallable, Category = "GeometryScript|UVs", meta=(ScriptMethod, AdvancedDisplay = "UDIMResolutions"))
+	static UPARAM(DisplayName = "Target Mesh") UDynamicMesh* 
+	ApplyTexelDensityUVScaling( 
+		UDynamicMesh* TargetMesh, 
+		UPARAM(DisplayName = "UV Channel") int UVSetIndex,
+		FGeometryScriptUVTexelDensityOptions Options,
+		FGeometryScriptMeshSelection Selection,
+		UGeometryScriptDebug* Debug = nullptr );
+
 	/**
 	* Recomputes UVs in the UV Channel for a Mesh based on different types of well-defined UV islands, such as existing UV islands, PolyGroups, 
 	* or a subset of the mesh based on a non-empty Selection.
diff --git a/Engine/Plugins/Runtime/MeshModelingToolset/Source/ModelingOperators/Private/ParameterizationOps/TexelDensityOp.cpp b/Engine/Plugins/Runtime/MeshModelingToolset/Source/ModelingOperators/Private/ParameterizationOps/TexelDensityOp.cpp
new file mode 100644
--- /dev/null
+++ b/Engine/Plugins/Runtime/MeshModelingToolset/Source/ModelingOperators/Private/ParameterizationOps/TexelDensityOp.cpp
@@ -0,0 +1,365 @@
+// Copyright Epic Games, Inc. All Rights Reserved.
+
+#include "ParameterizationOps/TexelDensityOp.h"
+
+#include "DynamicMesh/DynamicMesh3.h"
+#include "DynamicMesh/DynamicMeshAttributeSet.h"
+#include "Parameterization/MeshUVPacking.h"
+#include "Properties/UVLayoutProperties.h"
+#include "Selections/MeshConnectedComponents.h"
+#include "Parameterization/MeshUDIMClassifier.h"
+#include "UDIMUtilities.h"
+#include "Parameterization/UVMetrics.h"
+
+
+#include UE_INLINE_GENERATED_CPP_BY_NAME(TexelDensityOp)
+
+using namespace UE::Geometry;
+
+namespace UVTexelDensityOpLocals
+{
+	FVector2f ExternalUVToInternalUV(const FVector2f& UV)
+	{
+		return FVector2f(UV.X, 1 - UV.Y);
+	}
+
+	FVector2f InternalUVToExternalUV(const FVector2f& UV)
+	{
+		return FVector2f(UV.X, 1 - UV.Y);
+	}
+
+	struct TileConnectedComponents
+	{
+		TileConnectedComponents(const FMeshConnectedComponents& ConnectedComponentsIn, const FVector2i& TileIn)
+			: ConnectedComponents(ConnectedComponentsIn),
+		Tile(TileIn)
+		{
+			for (int32 ComponentIndex = 0; ComponentIndex < ConnectedComponents.Num(); ++ComponentIndex)
+			{
+				TileTids.Append(ConnectedComponents[ComponentIndex].Indices);
+			}
+		}
+
+		TArray<int32> TileTids;
+		FMeshConnectedComponents ConnectedComponents;
+		FVector2i Tile;
+	};
+
+	void CollectIslandComponentsPerTile(const FDynamicMesh3& Mesh, const FDynamicMeshUVOverlay& UVOverlay, TOptional<TSet<int32>>& Selection,
+		                                TArray< TileConnectedComponents >& ComponentsPerTile, bool bUDIMsEnabled)
+	{
+		ComponentsPerTile.Empty();
+
+		TArray<FVector2i> Tiles;
+		TArray<TUniquePtr<TArray<int32>>> TileTids;
+
+		if (bUDIMsEnabled)
+		{
+			TOptional<TArray<int32>> SelectionArray;
+			if (Selection.IsSet())
+			{
+				SelectionArray = Selection.GetValue().Array();
+			}
+			FDynamicMeshUDIMClassifier TileClassifier(&UVOverlay, SelectionArray);
+
+			Tiles = TileClassifier.ActiveTiles();
+			for (const FVector2i& Tile : Tiles)
+			{
+				TileTids.Emplace(MakeUnique<TArray<int32>>(TileClassifier.TidsForTile(Tile)));
+			}
+		}
+		else
+		{
+			if (Selection.IsSet())
+			{
+				Tiles.Add({ 0,0 });
+				TileTids.Emplace(MakeUnique<TArray<int32>>(Selection.GetValue().Array()));
+			}
+			else
+			{
+				Tiles.Add({ 0,0 });
+				TileTids.Emplace(MakeUnique<TArray<int32>>());
+				TileTids[0]->Reserve(Mesh.TriangleCount());
+				for (int32 Tid : Mesh.TriangleIndicesItr())
+				{
+					TileTids[0]->Add(Tid);
+				}
+			}
+		}
+
+		for (int32 TileIndex = 0; TileIndex < Tiles.Num(); ++TileIndex)
+		{
+			FMeshConnectedComponents ConnectedComponents(&Mesh);
+			if (TileTids[TileIndex])
+			{
+				ConnectedComponents.FindConnectedTriangles(*TileTids[TileIndex], [&](int32 Triangle0, int32 Triangle1) {
+					return UVOverlay.AreTrianglesConnected(Triangle0, Triangle1);
+				});
+			}
+			else
+			{
+				ConnectedComponents.FindConnectedTriangles([&](int32 Triangle0, int32 Triangle1) {
+					return UVOverlay.AreTrianglesConnected(Triangle0, Triangle1);
+				});
+			}
+
+			if (!ConnectedComponents.Components.IsEmpty())
+			{
+				ComponentsPerTile.Emplace(ConnectedComponents, Tiles[TileIndex]);
+			}
+		}
+	}
+}
+
+bool UUVEditorTexelDensitySettings::InSamplingMode() const
+{
+	return true;
+}
+
+
+void FUVEditorTexelDensityOp::SetTransform(const FTransformSRT3d& Transform)
+{
+	ResultTransform = Transform;
+}
+
+
+void FUVEditorTexelDensityOp::ScaleMeshSubRegionByDensity(FDynamicMeshUVOverlay* UVLayer, const TArray<int32>& Tids, TSet<int32>& UVElements, int32 TileResolution)
+{
+	double AverageTidTexelDensity = 0.0;
+	int SetTriangleCount = 0;
+	for (int Tid : Tids)
+	{
+		if (ResultMesh->Attributes()->GetUVLayer(UVLayerIndex)->IsSetTriangle(Tid))
+		{
+			AverageTidTexelDensity += FUVMetrics::TexelDensity(*ResultMesh, UVLayerIndex, Tid, TileResolution);
+			SetTriangleCount++;
+		}
+	}
+	if (SetTriangleCount)
+	{
+		AverageTidTexelDensity /= SetTriangleCount;
+	}
+
+	double TargetTexelDensity = (double)TargetPixelCountMeasurement / TargetWorldSpaceMeasurement;
+	double RequiredGlobalScaleFactor = TargetTexelDensity / AverageTidTexelDensity;
+
+	for (int ElementID : UVElements)
+	{
+		FVector2f UV = UVTexelDensityOpLocals::InternalUVToExternalUV(UVLayer->GetElement(ElementID));
+		UV = (UV * RequiredGlobalScaleFactor);
+		UVLayer->SetElement(ElementID, UVTexelDensityOpLocals::ExternalUVToInternalUV(UV));
+	}
+}
+
+
+
+void FUVEditorTexelDensityOp::CalculateResult(FProgressCancel* Progress)
+{
+	TRACE_CPUPROFILER_EVENT_SCOPE(UVLayoutOp_CalculateResult);
+
+	if (Progress && Progress->Cancelled())
+	{
+		return;
+	}
+	ResultMesh->Copy(*OriginalMesh, true, true, true, true);
+
+	if (!ensureMsgf(ResultMesh->HasAttributes(), TEXT("Attributes not found on mesh? Conversion should always create them, so this operator should not need to do so.")))
+	{
+		ResultMesh->EnableAttributes();
+	}
+
+	if (Progress && Progress->Cancelled())
+	{
+		return;
+	}
+
+
+	int UVLayerInput = UVLayerIndex;
+	FDynamicMeshUVOverlay* UseUVLayer = ResultMesh->Attributes()->GetUVLayer(UVLayerInput);
+
+
+	if (Progress && Progress->Cancelled())
+	{
+		return;
+	}
+
+	if (bMaintainOriginatingUDIM)
+	{
+		TArray<UVTexelDensityOpLocals::TileConnectedComponents> TileComponents;
+		UVTexelDensityOpLocals::CollectIslandComponentsPerTile(*ResultMesh, *ResultMesh->Attributes()->GetUVLayer(UVLayerIndex), Selection, TileComponents, true);
+
+		for (UVTexelDensityOpLocals::TileConnectedComponents Tile : TileComponents)
+		{
+			FVector2i TileIndex = Tile.Tile;	
+			const int32 TileID = UE::TextureUtilitiesCommon::GetUDIMIndex(TileIndex.X, TileIndex.Y);
+			float TileResolution = TextureResolution;
+			if (TextureResolutionPerUDIM.IsSet())
+			{
+				TileResolution = TextureResolutionPerUDIM.GetValue().FindOrAdd(TileID, this->TextureResolution);
+			}
+
+			TSet<int32> ElementsToMove;
+			ElementsToMove.Reserve(Tile.TileTids.Num() * 3);
+			for (int Tid : Tile.TileTids)
+			{
+				FIndex3i Elements = UseUVLayer->GetTriangle(Tid);
+				ElementsToMove.Add(Elements[0]);
+				ElementsToMove.Add(Elements[1]);
+				ElementsToMove.Add(Elements[2]);
+			}
+
+			for (int32 Element : ElementsToMove)
+			{
+				FVector2f UV = UVTexelDensityOpLocals::InternalUVToExternalUV(UseUVLayer->GetElement(Element));
+				UV = UV - FVector2f(TileIndex);
+				UseUVLayer->SetElement(Element, UVTexelDensityOpLocals::ExternalUVToInternalUV(UV));
+			}
+
+			// TODO: There is a second connected components call inside the packer that might be unnessessary. Could be a future optimization.
+			FDynamicMeshUVPacker Packer(UseUVLayer, MakeUnique<TArray<int32>>(Tile.TileTids));
+			Packer.TextureResolution = TileResolution;
+			ExecutePacker(Packer);
+			if (Progress && Progress->Cancelled())
+			{
+				return;
+			}
+
+			if (TexelDensityMode == EUVTexelDensityOpModes::ScaleGlobal)
+			{
+				ScaleMeshSubRegionByDensity(UseUVLayer, Tile.TileTids, ElementsToMove, TileResolution);
+			}
+
+			if (TexelDensityMode == EUVTexelDensityOpModes::ScaleIslands)
+			{
+				for (int32 ComponentIndex = 0; ComponentIndex < Tile.ConnectedComponents.Components.Num(); ComponentIndex++)
+				{
+					const TArray<int32>& ComponentTids = Tile.ConnectedComponents.Components[ComponentIndex].Indices;
+
+					TSet<int32> ComponentElementsToMove;
+					ComponentElementsToMove.Reserve(ComponentTids.Num() * 3);
+					for (int Tid : ComponentTids)
+					{
+						FIndex3i Elements = UseUVLayer->GetTriangle(Tid);
+						ComponentElementsToMove.Add(Elements[0]);
+						ComponentElementsToMove.Add(Elements[1]);
+						ComponentElementsToMove.Add(Elements[2]);
+					}
+
+					ScaleMeshSubRegionByDensity(UseUVLayer, ComponentTids, ComponentElementsToMove, TileResolution);
+				}
+			}
+
+			for (int32 Element : ElementsToMove)
+			{
+				FVector2f UV = UVTexelDensityOpLocals::InternalUVToExternalUV(UseUVLayer->GetElement(Element));
+				UV = UV + FVector2f(TileIndex);
+				UseUVLayer->SetElement(Element, UVTexelDensityOpLocals::ExternalUVToInternalUV(UV));
+			}
+		}
+	}
+	else
+	{
+		TArray<UVTexelDensityOpLocals::TileConnectedComponents> TileComponents;
+		UVTexelDensityOpLocals::CollectIslandComponentsPerTile(*ResultMesh, *ResultMesh->Attributes()->GetUVLayer(UVLayerIndex), Selection, TileComponents, false);
+
+		if(TileComponents.Num() == 0)
+		{
+			return;
+		}
+
+		TSet<int32> ElementsToMove;
+		ElementsToMove.Reserve(TileComponents[0].TileTids.Num() * 3);
+		for (int Tid : TileComponents[0].TileTids)
+		{
+			FIndex3i Elements = UseUVLayer->GetTriangle(Tid);
+			ElementsToMove.Add(Elements[0]);
+			ElementsToMove.Add(Elements[1]);
+			ElementsToMove.Add(Elements[2]);
+		}
+
+		FDynamicMeshUVPacker Packer(UseUVLayer, MakeUnique<TArray<int32>>(TileComponents[0].TileTids));
+		Packer.TextureResolution = TextureResolution;
+		ExecutePacker(Packer);
+		if (Progress && Progress->Cancelled())
+		{
+			return;
+		}
+
+		if (TexelDensityMode == EUVTexelDensityOpModes::ScaleGlobal)
+		{
+			ScaleMeshSubRegionByDensity(UseUVLayer, TileComponents[0].TileTids, ElementsToMove, TextureResolution);
+		}
+
+		if (TexelDensityMode == EUVTexelDensityOpModes::ScaleIslands)
+		{
+			for (int32 ComponentIndex = 0; ComponentIndex < TileComponents[0].ConnectedComponents.Components.Num(); ComponentIndex++)
+			{
+				const TArray<int32>& ComponentTids = TileComponents[0].ConnectedComponents.Components[ComponentIndex].Indices;
+
+				TSet<int32> ComponentElementsToMove;
+				ComponentElementsToMove.Reserve(ComponentTids.Num() * 3);
+				for (int Tid : ComponentTids)
+				{
+					FIndex3i Elements = UseUVLayer->GetTriangle(Tid);
+					ComponentElementsToMove.Add(Elements[0]);
+					ComponentElementsToMove.Add(Elements[1]);
+					ComponentElementsToMove.Add(Elements[2]);
+				}
+
+				ScaleMeshSubRegionByDensity(UseUVLayer, ComponentTids, ComponentElementsToMove, TextureResolution);
+			}
+		}
+
+	}
+
+
+}
+
+void FUVEditorTexelDensityOp::ExecutePacker(FDynamicMeshUVPacker& Packer)
+{
+	if (TexelDensityMode == EUVTexelDensityOpModes::Normalize)
+	{
+		Packer.bScaleIslandsByWorldSpaceTexelRatio = true;
+		if (Packer.StandardPack() == false)
+		{
+			// failed... what to do?
+			return;
+		}
+	}
+}
+
+TUniquePtr<FDynamicMeshOperator> UUVTexelDensityOperatorFactory::MakeNewOperator()
+{
+	TUniquePtr<FUVEditorTexelDensityOp> Op = MakeUnique<FUVEditorTexelDensityOp>();
+
+	Op->OriginalMesh = OriginalMesh;
+
+	switch (Settings->TexelDensityMode)
+	{
+	case ETexelDensityToolMode::ApplyToIslands:
+		Op->TexelDensityMode = EUVTexelDensityOpModes::ScaleIslands;
+		break;
+	case ETexelDensityToolMode::ApplyToWhole:
+		Op->TexelDensityMode = EUVTexelDensityOpModes::ScaleGlobal;
+		break;
+	case ETexelDensityToolMode::Normalize:
+		Op->TexelDensityMode = EUVTexelDensityOpModes::Normalize;
+		break;
+	default:
+		ensure(false);
+		break;
+	}
+
+	Op->TextureResolution = Settings->TextureResolution;
+	Op->TargetWorldSpaceMeasurement = Settings->TargetWorldUnits;
+	Op->TargetPixelCountMeasurement = Settings->TargetPixelCount;
+
+	Op->UVLayerIndex = GetSelectedUVChannel();
+	Op->TextureResolution = Settings->TextureResolution;
+	Op->SetTransform(TargetTransform);
+	Op->bMaintainOriginatingUDIM = Settings->bEnableUDIMLayout;
+	Op->Selection = Selection;
+	Op->TextureResolutionPerUDIM = TextureResolutionPerUDIM;
+
+	return Op;
+}
diff --git a/Engine/Plugins/Runtime/MeshModelingToolset/Source/ModelingOperators/Public/ParameterizationOps/TexelDensityOp.h b/Engine/Plugins/Runtime/MeshModelingToolset/Source/ModelingOperators/Public/ParameterizationOps/TexelDensityOp.h
new file mode 100644
--- /dev/null
+++ b/Engine/Plugins/Runtime/MeshModelingToolset/Source/ModelingOperators/Public/ParameterizationOps/TexelDensityOp.h
@@ -0,0 +1,122 @@
+// Copyright Epic Games, Inc. All Rights Reserved.
+
+#pragma once
+
+#include "CoreMinimal.h"
+#include "Util/ProgressCancel.h"
+#include "ModelingOperators.h"
+#include "InteractiveTool.h"
+#include "DynamicMesh/DynamicMeshAttributeSet.h"
+
+#include "TexelDensityOp.generated.h"
+
+UENUM()
+enum class ETexelDensityToolMode
+{
+	ApplyToIslands,
+	ApplyToWhole,
+	Normalize
+};
+
+UCLASS()
+class MODELINGOPERATORS_API UUVEditorTexelDensitySettings : public UInteractiveToolPropertySet
+{
+	GENERATED_BODY()
+public:
+
+	UFUNCTION()
+	virtual bool InSamplingMode() const;
+
+	UPROPERTY(EditAnywhere, Category = "TexelDensity", meta = (DisplayName = "Scale Mode", EditCondition = "!InSamplingMode" ))
+	ETexelDensityToolMode TexelDensityMode = ETexelDensityToolMode::ApplyToIslands;
+
+	UPROPERTY(EditAnywhere, Category = "TexelDensity", meta = (DisplayName = "World Units", EditCondition = "TexelDensityMode != ETexelDensityToolMode::Normalize"))
+	float TargetWorldUnits = 100;
+
+	UPROPERTY(EditAnywhere, Category = "TexelDensity", meta = (DisplayName = "Pixels", EditCondition = "TexelDensityMode != ETexelDensityToolMode::Normalize && !InSamplingMode"))
+	float TargetPixelCount = 1024;
+
+	UPROPERTY(EditAnywhere, Category = "TexelDensity", meta = (DisplayName = "Texture Dimensions", EditCondition = "TexelDensityMode != ETexelDensityToolMode::Normalize && !InSamplingMode", LinearDeltaSensitivity = 1000, Delta = 64, SliderExponent = 1, UIMin = 2, UIMax = 16384))
+	float TextureResolution = 1024;
+
+	UPROPERTY(EditAnywhere, Category = "TexelDensity", meta = (DisplayName = "Use UDIMs", EditCondition = "TexelDensityMode != ETexelDensityToolMode::Normalize && !InSamplingMode", TransientToolProperty))
+	bool bEnableUDIMLayout = false;
+};
+
+
+namespace UE
+{
+	namespace Geometry
+	{
+		class FDynamicMesh3;
+		class FDynamicMeshUVPacker;
+
+		enum class EUVTexelDensityOpModes
+		{
+			ScaleGlobal = 0,
+			ScaleIslands = 1,
+			Normalize = 2
+		};
+
+		class MODELINGOPERATORS_API FUVEditorTexelDensityOp : public FDynamicMeshOperator
+		{
+		public:
+			virtual ~FUVEditorTexelDensityOp() {}
+
+			// inputs
+			TSharedPtr<FDynamicMesh3, ESPMode::ThreadSafe> OriginalMesh;
+
+			EUVTexelDensityOpModes TexelDensityMode = EUVTexelDensityOpModes::ScaleGlobal;
+
+			int UVLayerIndex = 0;
+			int TextureResolution = 1024;
+			int TargetWorldSpaceMeasurement = 100;
+			int TargetPixelCountMeasurement = 1024;
+
+			bool bMaintainOriginatingUDIM = false;
+			TOptional<TSet<int32>> Selection;
+			TOptional<TMap<int32, int32>> TextureResolutionPerUDIM;
+			FVector2f UVTranslation = FVector2f::Zero();
+
+			void SetTransform(const FTransformSRT3d& Transform);
+
+			//
+			// FDynamicMeshOperator implementation
+			// 
+
+			virtual void CalculateResult(FProgressCancel* Progress) override;
+
+		protected:
+			void ExecutePacker(FDynamicMeshUVPacker& Packer);			
+			void ScaleMeshSubRegionByDensity(FDynamicMeshUVOverlay* UVLayer, const TArray<int32>& Tids, TSet<int32>& UVElements, int32 TileResolution);
+
+
+		};
+
+	} // end namespace UE::Geometry
+} // end namespace UE
+
+/**
+ * Can be hooked up to a UMeshOpPreviewWithBackgroundCompute to perform UV Layout operations.
+ *
+ * Inherits from UObject so that it can hold a strong pointer to the settings UObject, which
+ * needs to be a UObject to be displayed in the details panel.
+ */
+UCLASS()
+class MODELINGOPERATORS_API UUVTexelDensityOperatorFactory : public UObject, public UE::Geometry::IDynamicMeshOperatorFactory
+{
+	GENERATED_BODY()
+
+public:
+	// IDynamicMeshOperatorFactory API
+	virtual TUniquePtr<UE::Geometry::FDynamicMeshOperator> MakeNewOperator() override;
+
+	UPROPERTY()
+	TObjectPtr<UUVEditorTexelDensitySettings> Settings;
+
+	TOptional<TSet<int32>> Selection;
+	TSharedPtr<UE::Geometry::FDynamicMesh3, ESPMode::ThreadSafe> OriginalMesh;
+	TUniqueFunction<int32()> GetSelectedUVChannel = []() { return 0; };
+	FTransform TargetTransform;
+	TOptional<TMap<int32, int32>> TextureResolutionPerUDIM;
+};
\ No newline at end of file
