Geoscripting: Add some defensive coding for a potential mismatch with the point count versus the section info.

#rnx

diff --git a/Engine/Plugins/Runtime/MeshModelingToolset/Source/ModelingComponents/Private/ConversionUtils/SkinnedMeshToDynamicMesh.cpp b/Engine/Plugins/Runtime/MeshModelingToolset/Source/ModelingComponents/Private/ConversionUtils/SkinnedMeshToDynamicMesh.cpp
--- a/Engine/Plugins/Runtime/MeshModelingToolset/Source/ModelingComponents/Private/ConversionUtils/SkinnedMeshToDynamicMesh.cpp
+++ b/Engine/Plugins/Runtime/MeshModelingToolset/Source/ModelingComponents/Private/ConversionUtils/SkinnedMeshToDynamicMesh.cpp
@@ -90,7 +90,10 @@ public:
 			const int32 NumSectionVtx = static_cast<int32>( Section.NumVertices );
 			for (int32 VtxIndex = BaseVertexIndex; VtxIndex < NumSectionVtx + BaseVertexIndex; ++VtxIndex)
 			{
-				VertIDs.Add(VtxIndex);
+				if (ensure(VtxIndex < SkinnedVertices.Num()))
+				{
+					VertIDs.Add(VtxIndex);
+				}
 			}
 		}
 
