Geoscripting: Add temporary logging to track down an elusive vertex mismatch issue.

#rnx

diff --git a/Engine/Plugins/Runtime/MeshModelingToolset/Source/ModelingComponents/Private/ConversionUtils/SkinnedMeshToDynamicMesh.cpp b/Engine/Plugins/Runtime/MeshModelingToolset/Source/ModelingComponents/Private/ConversionUtils/SkinnedMeshToDynamicMesh.cpp
--- a/Engine/Plugins/Runtime/MeshModelingToolset/Source/ModelingComponents/Private/ConversionUtils/SkinnedMeshToDynamicMesh.cpp
+++ b/Engine/Plugins/Runtime/MeshModelingToolset/Source/ModelingComponents/Private/ConversionUtils/SkinnedMeshToDynamicMesh.cpp
@@ -78,6 +78,8 @@ public:
 		}
 		
 		// construct a list of all valid VertIDs for this mesh.
+		bool bLogMeshInfoDueToErrors = true;
+		
 		VertIDs.Reserve(NumVerts);
 		for (const FSkelMeshRenderSection& Section : LODData.RenderSections)
 		{
@@ -94,9 +96,28 @@ public:
 				{
 					VertIDs.Add(VtxIndex);
 				}
+				else 
+				{
+					bLogMeshInfoDueToErrors = true;
+				}
 			}
 		}
 
+		if (bLogMeshInfoDueToErrors)
+		{
+			UE_LOG(LogGeometry, Display, TEXT("LOD vertex mismatch with [%s -> %s]"), *SkinnedMeshComponent.GetPathName(), *SkinnedMeshComponent.GetSkinnedAsset()->GetPathName());
+			UE_LOG(LogGeometry, Display, TEXT(" - Total Vertex Count: %d"), SkinnedVertices.Num());
+			UE_LOG(LogGeometry, Display, TEXT(" - Total Index/Triangle Count: %d / %d "), IndexBuffer->Num(), IndexBuffer->Num() / 3);
+		 
+			for (int32 SectionIndex = 0; SectionIndex < LODData.RenderSections.Num(); SectionIndex++)
+			{
+				const FSkelMeshRenderSection& Section = LODData.RenderSections[SectionIndex];
+				UE_LOG(LogGeometry, Display, TEXT(" - Section[%d]: Skip? %s"), SectionIndex, SkipSection(Section) ? TEXT("Yes") : TEXT("No"));
+				UE_LOG(LogGeometry, Display, TEXT(" - - Vertex Base: %d - Vertex Count: %d"), Section.BaseVertexIndex, Section.NumVertices);
+				UE_LOG(LogGeometry, Display, TEXT(" - - Index Base: %d - Triangle Base: %d - Triangle Count: %d"),
+					Section.BaseIndex, Section.BaseIndex / 3, Section.NumTriangles);
+			}
+		}
 
 		// generate vertex weights and remap the indices
 		const FSkinWeightVertexBuffer* SkinWeightBuffer = Component->GetSkinWeightBuffer(LOD);
