MeshModelingTools: Fix PolyEd breakages from 36162322.

diff --git a/Engine/Plugins/Experimental/MeshModelingToolsetExp/Source/MeshModelingToolsEditorOnlyExp/Private/SkeletalMesh/SkinWeightsPaintTool.cpp b/Engine/Plugins/Experimental/MeshModelingToolsetExp/Source/MeshModelingToolsEditorOnlyExp/Private/SkeletalMesh/SkinWeightsPaintTool.cpp
--- a/Engine/Plugins/Experimental/MeshModelingToolsetExp/Source/MeshModelingToolsEditorOnlyExp/Private/SkeletalMesh/SkinWeightsPaintTool.cpp
+++ b/Engine/Plugins/Experimental/MeshModelingToolsetExp/Source/MeshModelingToolsEditorOnlyExp/Private/SkeletalMesh/SkinWeightsPaintTool.cpp
@@ -397,7 +397,7 @@ void UWeightToolMeshSelector::InitialSetup(
 	PolygonSelectionMechanic = NewObject<UPolygonSelectionMechanic>(this);
 	PolygonSelectionMechanic->bAddSelectionFilterPropertiesToParentTool = false;
 	PolygonSelectionMechanic->Setup(InParentTool);
-	PolygonSelectionMechanic->SetIsEnabled(false);
+	PolygonSelectionMechanic->SetIsEnabled(false, false);
 	PolygonSelectionMechanic->OnSelectionChanged.AddLambda(OnSelectionChangedFunc );
 
 	// set up style of vertex selection
@@ -506,7 +506,7 @@ void UWeightToolMeshSelector::SetIsEnabled(bool bIsEnabled)
 	// force off if there's no preview mesh
 	bIsEnabled = PreviewMesh ? bIsEnabled : false; 
 	
-	PolygonSelectionMechanic->SetIsEnabled(bIsEnabled);
+	PolygonSelectionMechanic->SetIsEnabled(bIsEnabled, bIsEnabled);
 }
 
 void UWeightToolMeshSelector::SetComponentSelectionMode(EComponentSelectionMode InMode)
diff --git a/Engine/Plugins/Runtime/MeshModelingToolset/Source/ModelingComponents/Private/Selection/MeshTopologySelectionMechanic.cpp b/Engine/Plugins/Runtime/MeshModelingToolset/Source/ModelingComponents/Private/Selection/MeshTopologySelectionMechanic.cpp
--- a/Engine/Plugins/Runtime/MeshModelingToolset/Source/ModelingComponents/Private/Selection/MeshTopologySelectionMechanic.cpp
+++ b/Engine/Plugins/Runtime/MeshModelingToolset/Source/ModelingComponents/Private/Selection/MeshTopologySelectionMechanic.cpp
@@ -124,9 +124,10 @@ void UMeshTopologySelectionMechanic::DisableBehaviors(UInteractiveTool* ParentTo
 	// TODO: Is it worth adding a way to remove the property watchers for marquee?
 }
 
-void UMeshTopologySelectionMechanic::SetIsEnabled(bool bOn)
+void UMeshTopologySelectionMechanic::SetIsEnabled(bool bBehaviorEnabledIn, bool bRenderTopologyIn)
 {
-	bIsEnabled = bOn;
+	bIsEnabled = bBehaviorEnabledIn;
+	bRenderTopology = bRenderTopologyIn;
 	UpdateMarqueeEnabled();
 }
 
@@ -167,16 +168,18 @@ TPair<FInputCapturePriority, FInputCapturePriority> UMeshTopologySelectionMechan
 
 void UMeshTopologySelectionMechanic::Render(IToolsContextRenderAPI* RenderAPI)
 {
-	if (!bIsEnabled)
+	// Cache the view camera state so we can use for snapping/etc.
+	GetParentTool()->GetToolManager()->GetContextQueriesAPI()->GetCurrentViewState(CameraState);
+	
+	if (bIsEnabled)
 	{
-		return;
+		MarqueeMechanic->Render(RenderAPI);
 	}
 	
-	MarqueeMechanic->Render(RenderAPI);
-
-	// Cache the view camera state so we can use for snapping/etc.
-	// This should not happen in Render() though...
-	GetParentTool()->GetToolManager()->GetContextQueriesAPI()->GetCurrentViewState(CameraState);
+	if (!bRenderTopology) 
+	{
+		return;
+	}
 
 	FViewCameraState RenderCameraState = RenderAPI->GetCameraState();
 
@@ -275,11 +278,9 @@ bool UMeshTopologySelectionMechanic::TopologyHitTest(const FRay& WorldRay, FHitR
 
 bool UMeshTopologySelectionMechanic::TopologyHitTest(const FRay& WorldRay, FHitResult& OutHit, FGroupTopologySelection& OutSelection, bool bUseOrthoSettings)
 {
-	if (!bIsEnabled)
-	{
-		return false;
-	}
-	
+	// Note: this function should remain callable even if the mechanic is disabled, though client
+	//  could reach in to use the TopoSelector directly.
+
 	FRay3d LocalRay(TargetTransform.InverseTransformPosition((FVector3d)WorldRay.Origin),
 		TargetTransform.InverseTransformVector((FVector3d)WorldRay.Direction));
 	UE::Geometry::Normalize(LocalRay.Direction);
@@ -838,6 +839,11 @@ FInputRayHit UMeshTopologySelectionMechanic::IsHitByClick(const FInputDeviceRay&
 
 void UMeshTopologySelectionMechanic::OnClicked(const FInputDeviceRay& ClickPos)
 {
+	if (!ensure(bIsEnabled))
+	{
+		return;
+	}
+
 	// update selection
 	ParentTool->GetToolManager()->BeginUndoTransaction(LOCTEXT("SelectionChange", "Selection"));
 	BeginChange();
@@ -865,6 +871,11 @@ void UMeshTopologySelectionMechanic::OnBeginHover(const FInputDeviceRay& DeviceP
 
 bool UMeshTopologySelectionMechanic::OnUpdateHover(const FInputDeviceRay& DevicePos)
 {
+	if (!bIsEnabled)
+	{
+		return false;
+	}
+
 	UpdateHighlight(DevicePos.WorldRay);
 	return true;
 }
diff --git a/Engine/Plugins/Runtime/MeshModelingToolset/Source/ModelingComponents/Public/Selection/MeshTopologySelectionMechanic.h b/Engine/Plugins/Runtime/MeshModelingToolset/Source/ModelingComponents/Public/Selection/MeshTopologySelectionMechanic.h
--- a/Engine/Plugins/Runtime/MeshModelingToolset/Source/ModelingComponents/Public/Selection/MeshTopologySelectionMechanic.h
+++ b/Engine/Plugins/Runtime/MeshModelingToolset/Source/ModelingComponents/Public/Selection/MeshTopologySelectionMechanic.h
@@ -159,9 +159,13 @@ public:
 	void DisableBehaviors(UInteractiveTool* ParentToolIn);
 
 	/**
-	 * Enable/disable the mechanic without permanently removing behaviors or shutting it down.
+	 * Enable/disable the selection behaviors in the mechanic without permanently removing them or
+	 *  shutting the mechanic down.
+	 * @param bBehaviorsEnabled Determines whether the mechanic selection behaviors are active
+	 * @param bRenderTopology Determines whether the topology is rendered, which is independent
+	 *   of the behaviors being active.
 	 */
-	void SetIsEnabled(bool bOn);
+	void SetIsEnabled(bool bBehaviorsEnabled, bool bRenderTopology = true);
 
 	/**
 	 * Update the location of the mesh after Initialize()
@@ -223,8 +227,9 @@ public:
 	void NotifyMeshChanged(bool bTopologyModified);
 
 	/**
-	 * Perform a hit test on the topology using the current selection settings. In cases of hitting edges and
-	 * corners, OutHit contains the following:
+	 * Perform a hit test on the topology using the current selection settings. Note that to work properly,
+	 *  the mechanic requires that Render() be called on it, as it needs to cache camera state.
+	 * In cases of hitting edges and corners, OutHit contains the following:
 	 *   OutHit.FaceIndex: edge or corner id in the topology
 	 *   OutHit.ImpactPoint: closest point on the ray to the hit element (Note: not a point on the element!)
 	 *   OutHit.Distance: distance along the ray to ImpactPoint
@@ -495,6 +500,8 @@ public:
 
 	friend class FMeshTopologySelectionMechanicSelectionChange;
 
+private:
+	bool bRenderTopology = true;
 };
 
 
