MeshModelingTools: Change some primitive creation tools to default to PerQuad, and make LatLong sphere PerFace not be the same as PerShape (and be same as PerQuad).

diff --git a/Engine/Plugins/Runtime/MeshModelingToolset/Source/MeshModelingTools/Private/AddPrimitiveTool.cpp b/Engine/Plugins/Runtime/MeshModelingToolset/Source/MeshModelingTools/Private/AddPrimitiveTool.cpp
--- a/Engine/Plugins/Runtime/MeshModelingToolset/Source/MeshModelingTools/Private/AddPrimitiveTool.cpp
+++ b/Engine/Plugins/Runtime/MeshModelingToolset/Source/MeshModelingTools/Private/AddPrimitiveTool.cpp
@@ -122,6 +122,7 @@ void UAddPrimitiveTool::Setup()
 	AddToolPropertySource(ShapeSettings);
 	
 	ShapeSettings->WatchProperty(ShapeSettings->TargetSurface, [this](EMakeMeshPlacementType){UpdateTargetSurface();});
+	ShapeSettings->PolygroupMode = GetDefaultPolygroupMode();
 	ShapeSettings->RestoreProperties(this);
 
 	MaterialProperties = NewObject<UNewMeshMaterialProperties>(this);
@@ -798,7 +799,10 @@ void UAddSpherePrimitiveTool::GenerateMesh(FDynamicMesh3* OutMesh) const
 		SphereGen.Radius = SphereSettings->Radius;
 		SphereGen.NumTheta = SphereSettings->VerticalSlices;
 		SphereGen.NumPhi = SphereSettings->HorizontalSlices + 1;
-		SphereGen.bPolygroupPerQuad = (ShapeSettings->PolygroupMode == EMakeMeshPolygroupMode::PerQuad);
+		// In FSphereGenerator, PerFace is effectively ignored, and does the same thing as 
+		//  PerShape, which is unlikely to be useful. So we might as well have PerFace do
+		//  the same thing as PerQuad.
+		SphereGen.bPolygroupPerQuad = (ShapeSettings->PolygroupMode != EMakeMeshPolygroupMode::PerShape);
 		SphereGen.Generate();
 		OutMesh->Copy(&SphereGen);
 		break;
diff --git a/Engine/Plugins/Runtime/MeshModelingToolset/Source/MeshModelingTools/Public/AddPrimitiveTool.h b/Engine/Plugins/Runtime/MeshModelingToolset/Source/MeshModelingTools/Public/AddPrimitiveTool.h
--- a/Engine/Plugins/Runtime/MeshModelingToolset/Source/MeshModelingTools/Public/AddPrimitiveTool.h
+++ b/Engine/Plugins/Runtime/MeshModelingToolset/Source/MeshModelingTools/Public/AddPrimitiveTool.h
@@ -503,6 +503,8 @@ public:
 
 
 protected:
+	virtual EMakeMeshPolygroupMode GetDefaultPolygroupMode() const { return EMakeMeshPolygroupMode::PerQuad; }
+
 	enum class EState
 	{
 		PlacingPrimitive,
@@ -658,6 +660,10 @@ class UAddSpherePrimitiveTool : public UAddPrimitiveTool
 public:
 	explicit UAddSpherePrimitiveTool(const FObjectInitializer& ObjectInitializer);
 protected:
+	virtual EMakeMeshPolygroupMode GetDefaultPolygroupMode() const override 
+	{ 
+		return EMakeMeshPolygroupMode::PerFace; 
+	}
 	virtual void GenerateMesh(FDynamicMesh3* OutMesh) const override;
 };
 
