make polygroup edit and triangle edit tools auto-convert input selections if they're the wrong topology type, rather than discarding them.
+ Also translate mode-level polygroup edge selections (without corners) to tool polygroup edge selection w/ corners if needed.

diff --git a/Engine/Plugins/Runtime/MeshModelingToolset/Source/MeshModelingTools/Private/EditMeshPolygonsTool.cpp b/Engine/Plugins/Runtime/MeshModelingToolset/Source/MeshModelingTools/Private/EditMeshPolygonsTool.cpp
--- a/Engine/Plugins/Runtime/MeshModelingToolset/Source/MeshModelingTools/Private/EditMeshPolygonsTool.cpp
+++ b/Engine/Plugins/Runtime/MeshModelingToolset/Source/MeshModelingTools/Private/EditMeshPolygonsTool.cpp
@@ -866,13 +866,45 @@ void UEditMeshPolygonsTool::Setup()
 	if (HasGeometrySelection())
 	{
 		const FGeometrySelection& CurSelection = GetGeometrySelection();
-		if (CurSelection.TopologyType == EGeometryTopologyType::Triangle && bTriangleMode)
+		// If the topology type doesn't match, we'll need to convert it here
+		FGeometrySelection ConvertedSelection;
+		const FGeometrySelection* UseSelection = &CurSelection;
+		bool bCanUseSelection = true;
+		// For polygroup edge selections, if the tool's polygroups have extra corners, need to convert to that
+		if (!bTriangleMode && 
+			CurSelection.TopologyType == EGeometryTopologyType::Polygroup && CurSelection.ElementType == EGeometryElementType::Edge && 
+			TopologyProperties->bAddExtraCorners)
+		{
+			// Convert default (no corner) group topology -> triangle topology -> tool (w/ corner) group topology
+			ConvertedSelection.InitializeTypes(EGeometryElementType::Edge, EGeometryTopologyType::Polygroup);
+			FGeometrySelection TempTriSelection(EGeometryElementType::Edge, EGeometryTopologyType::Triangle);
+			FGroupTopology GroupTopology(CurrentMesh.Get(), true);
+			bCanUseSelection = UE::Geometry::ConvertSelection(*CurrentMesh, &GroupTopology, CurSelection, TempTriSelection, EEnumerateSelectionConversionParams::ContainSelection);
+			bCanUseSelection = bCanUseSelection && UE::Geometry::ConvertSelection(*CurrentMesh, Topology.Get(), TempTriSelection, ConvertedSelection, EEnumerateSelectionConversionParams::ContainSelection);
+			UseSelection = &ConvertedSelection;
+		}
+		// If topology type is triangle but we want polygroup, or vice versa, convert accordingly
+		else if ((CurSelection.TopologyType == EGeometryTopologyType::Triangle) != bTriangleMode)
+		{
+			ConvertedSelection.InitializeTypes(CurSelection.ElementType, bTriangleMode ? EGeometryTopologyType::Triangle : EGeometryTopologyType::Polygroup);
+			const FGroupTopology* UseTopology = Topology.Get();
+			// We need a default topology to reference if we're converting from polygroup->triangle, since w/ Triangle Mode the tool's Topology has per-triangle groups
+			FGroupTopology DefaultGroupTopology(CurrentMesh.Get(), false);
+			if (bTriangleMode)
+			{
+				DefaultGroupTopology.RebuildTopology();
+				UseTopology = &DefaultGroupTopology;
+			}
+			bCanUseSelection = UE::Geometry::ConvertSelection(*CurrentMesh, UseTopology, CurSelection, ConvertedSelection, EEnumerateSelectionConversionParams::ContainSelection);
+			UseSelection = &ConvertedSelection;
+		}
+		if (bCanUseSelection && UseSelection->TopologyType == EGeometryTopologyType::Triangle && bTriangleMode)
 		{
-			SelectionMechanic->SetSelection_AsTriangleTopology(CurSelection);
+			SelectionMechanic->SetSelection_AsTriangleTopology(*UseSelection);
 		}
-		else if (CurSelection.TopologyType == EGeometryTopologyType::Polygroup && bTriangleMode == false)
+		else if (bCanUseSelection && UseSelection->TopologyType == EGeometryTopologyType::Polygroup && bTriangleMode == false)
 		{
-			SelectionMechanic->SetSelection_AsGroupTopology(CurSelection);
+			SelectionMechanic->SetSelection_AsGroupTopology(*UseSelection);
 		}
 	}
 
