MeshModelingTools: Fix extrude distance adjustment (to keep boundaries "even") not working properly in cases where extrude direction is not in plane of the corner. Fix flipping of extrude direction when extrude distance was negative before being set negative again.

diff --git a/Engine/Plugins/Runtime/MeshModelingToolset/Source/MeshModelingTools/Private/Operations/ExtrudeBoundaryEdges.cpp b/Engine/Plugins/Runtime/MeshModelingToolset/Source/MeshModelingTools/Private/Operations/ExtrudeBoundaryEdges.cpp
--- a/Engine/Plugins/Runtime/MeshModelingToolset/Source/MeshModelingTools/Private/Operations/ExtrudeBoundaryEdges.cpp
+++ b/Engine/Plugins/Runtime/MeshModelingToolset/Source/MeshModelingTools/Private/Operations/ExtrudeBoundaryEdges.cpp
@@ -177,11 +177,55 @@ namespace ExtrudeBoundaryEdgesLocals
 
 		if (ScalingLimit > 1)
 		{
-			FVector3d VectorToCompareAgainst = OutwardVector1.IsZero() ? OutwardVector2 : OutwardVector1;
-			double CosTheta = ExtrudeFrame.Frame.X().Dot(VectorToCompareAgainst);
-			if (FMath::Abs(CosTheta) > KINDA_SMALL_NUMBER)
+			// This means that we're going to be adjusting our extrude directions in an attempt to keep extruded
+			//  edges parallel. E.g. imagine extruding two edges at a 90 degree corner: the extrude directions
+			//  at the flat sides are directly outward, whereas the outward direction at the corner is at an angle.
+			//  If we extrude the same distance along each direction, the resulting corner is no longer 90 degrees
+			//  because the vertex would need to move further along the diagonal.
+			// The problem is made more complicated by the fact that we are not necessarily extruding in the plane
+			//  of the edges, and that the extrude frames at each vertex might be arbitrarily rotated depending on
+			//  the adjoining faces. So, having some in-frame adjustment that keeps edges parallel without knowledge 
+			//  of adjacent frames is impossible.
+			// Instead we'll try to do our best for some common use cases. In particular, imagine extruding some
+			//  edges on the boundary of a tesselated open-top box partially outward. We will adjust them such
+			//  that the projection onto the plane of the original boundary is parallel, and we will do it by
+			//  scaling only in the appropriate direction in the plane of the boundary. This will handle "tube"
+			//  extrusions and simple planar mesh extrusion, and for planar boundaries, the user can scale down
+			//  the result along the orthogonal axis to get parallel output even if incident faces made for
+			//  inconsistent frame rotations.
+
+			// Direction to scale is average of the two incoming edge vectors
+			FVector3d DirectionToScale = IncomingEdgeVector - OutgoingEdgeVector;
+			DirectionToScale.Normalize();
+
+			// If direction to scale is zero, that means we're along a flat section, where we don't need
+			//  to adjust.
+			if (!DirectionToScale.IsZero())
 			{
-				ExtrudeFrame.Scaling.X = FMath::Min(1 / FMath::Abs(CosTheta), ScalingLimit);
+
+				// We compare against an outward vector in the plane of the two vectors (rather than in the plane of
+				//  adjoining triangles, which we used when finding the extrude frame)
+				FVector3d NormalToEdges = IncomingEdgeVector.Cross(OutgoingEdgeVector); // Unnormalized
+				FVector3d OutwardVectorInEdgePlane = IncomingEdgeVector.Cross(NormalToEdges);
+				OutwardVectorInEdgePlane.Normalize();
+
+				// Figure out how far to scale, with an upper limit
+				double CosTheta = DirectionToScale.Dot(OutwardVectorInEdgePlane);
+				if (CosTheta == 0)
+				{
+					// Presumably this comes about from the edges doubling back onto themselves, but we'll double
+					//  check that we're not dealing with a straight edge to be safe.
+					if (IncomingEdgeVector.Dot(OutgoingEdgeVector) < 0)
+					{
+						ExtrudeFrame.InFrameScaleDirection = ExtrudeFrame.Frame.ToFrameVector(DirectionToScale);
+						ExtrudeFrame.Scaling = ScalingLimit;
+					}
+				}
+				else
+				{
+					ExtrudeFrame.InFrameScaleDirection = ExtrudeFrame.Frame.ToFrameVector(DirectionToScale);
+					ExtrudeFrame.Scaling = FMath::Min(1 / FMath::Abs(CosTheta), ScalingLimit);
+				}
 			}
 		}
 
@@ -666,4 +710,27 @@ bool FExtrudeBoundaryEdges::GetExtrudeFrame(const FDynamicMesh3& Mesh, int32 Vid
 	return true;
 }
 
+FVector3d FExtrudeBoundaryEdges::FExtrudeFrame::FromFramePoint(FVector3d FramePoint) const
+{
+	if (InFrameScaleDirection.IsSet() && Scaling != 1)
+	{
+		FVector3d Offset = FramePoint.Dot(InFrameScaleDirection.GetValue()) * (Scaling - 1) * InFrameScaleDirection.GetValue();
+		return Frame.FromFramePoint(FramePoint + Offset);
+	}
+	return Frame.FromFramePoint(FramePoint);
+}
+
+FVector3d UE::Geometry::FExtrudeBoundaryEdges::FExtrudeFrame::ToFramePoint(FVector3d WorldPoint) const
+{
+	FVector3d FramePoint = Frame.ToFramePoint(WorldPoint);
+	if (InFrameScaleDirection.IsSet() && Scaling != 1
+		// Can't undo zero scaling, but that shouldn't come up as a value
+		&& ensure(Scaling != 0))
+	{
+		double DotProduct = FramePoint.Dot(InFrameScaleDirection.GetValue());
+		return FramePoint - (1 - (1 / Scaling)) * DotProduct * InFrameScaleDirection.GetValue();
+	}
+	return FramePoint;
+}
+
 }//end UE::Geometry
\ No newline at end of file
diff --git a/Engine/Plugins/Runtime/MeshModelingToolset/Source/MeshModelingTools/Private/Operations/ExtrudeBoundaryEdges.h b/Engine/Plugins/Runtime/MeshModelingToolset/Source/MeshModelingTools/Private/Operations/ExtrudeBoundaryEdges.h
--- a/Engine/Plugins/Runtime/MeshModelingToolset/Source/MeshModelingTools/Private/Operations/ExtrudeBoundaryEdges.h
+++ b/Engine/Plugins/Runtime/MeshModelingToolset/Source/MeshModelingTools/Private/Operations/ExtrudeBoundaryEdges.h
@@ -35,13 +35,26 @@ public:
 	{
 		FExtrudeFrame() {}
 
-		FExtrudeFrame(const FFrame3d& FrameIn, const FVector3d& ScalingIn = FVector3d::One())
+		FExtrudeFrame(const FFrame3d& FrameIn)
 			: Frame(FrameIn)
+		{}
+		FExtrudeFrame(const FFrame3d& FrameIn, const FVector3d& InFrameScaleDirection, double ScalingIn)
+			: Frame(FrameIn)
+			, InFrameScaleDirection(InFrameScaleDirection)
 			, Scaling(ScalingIn)
 		{}
 
+		FVector3d FromFramePoint(FVector3d FramePoint) const;
+
+		FVector3d ToFramePoint(FVector3d WorldPoint) const;
+
 		FFrame3d Frame;
-		FVector3d Scaling = FVector3d::One();
+
+		// Extra scaling direction used to support adjusting extruded vertices when trying to keep
+		//  edges parallel. The input is scaled along this axis in frame space. See comment inside
+		//  ExtrudeBoundaryEdgesLocals::GetExtrudeFrameAtVertex.
+		TOptional<FVector3d> InFrameScaleDirection;
+		double Scaling = 1.0;
 	};
 
 	// Data needed to create a new vert and its extrude frame
@@ -74,7 +87,7 @@ public:
 	TFunction<FVector3d(const FVector3d& Position, const FExtrudeFrame& ExtrudeFrame, int32 SourceVid)> OffsetPositionFunc =
 		[this](const FVector3d& Position, const FExtrudeFrame& ExtrudeFrame, int32 SourceVid)
 	{
-		return ExtrudeFrame.Frame.FromFramePoint(FVector3d(this->DefaultOffsetDistance, 0, 0) * ExtrudeFrame.Scaling);
+		return ExtrudeFrame.FromFramePoint(FVector3d(this->DefaultOffsetDistance, 0, 0));
 	};
 
 	/** 
diff --git a/Engine/Plugins/Runtime/MeshModelingToolset/Source/MeshModelingTools/Private/ToolActivities/PolyEditExtrudeEdgeActivity.cpp b/Engine/Plugins/Runtime/MeshModelingToolset/Source/MeshModelingTools/Private/ToolActivities/PolyEditExtrudeEdgeActivity.cpp
--- a/Engine/Plugins/Runtime/MeshModelingToolset/Source/MeshModelingTools/Private/ToolActivities/PolyEditExtrudeEdgeActivity.cpp
+++ b/Engine/Plugins/Runtime/MeshModelingToolset/Source/MeshModelingTools/Private/ToolActivities/PolyEditExtrudeEdgeActivity.cpp
@@ -13,7 +13,6 @@
 #include "InteractiveToolManager.h"
 #include "MeshOpPreviewHelpers.h"
 #include "ModelingOperators.h" // FDynamicMeshOperator
-#include "Operations/ExtrudeBoundaryEdges.h"
 #include "SceneManagement.h"
 #include "Selection/PolygonSelectionMechanic.h"
 #include "Selection/StoredMeshSelectionUtil.h"
@@ -30,6 +29,7 @@ namespace PolyEditExtrudeEdgeActivityLocals
 	// These should match what is used in UMeshTopologySelectionMechanic (or better, live in a common place)
 	const FColor GroupLineColor = FColor::Red;
 	const float GroupLineThickness = 1.0f;
+	const float GroupLineDepthBias = 0.5f;
 
 	const FColor ExtrudeFrameLineColor = FColor::Orange;
 	const float ExtrudeFrameLineThickness = 0.5f;
@@ -90,7 +90,7 @@ namespace PolyEditExtrudeEdgeActivityLocals
 			case EPolyEditExtrudeEdgeDirectionMode::LocalExtrudeFrames:
 				Extruder.OffsetPositionFunc = [this](const FVector3d& Position, const FExtrudeBoundaryEdges::FExtrudeFrame& ExtrudeFrame, int32 SourceVid)
 				{
-					return ExtrudeFrame.Frame.FromFramePoint(ExtrudeFrame.Scaling * LocalFrameParams);
+					return ExtrudeFrame.FromFramePoint(LocalFrameParams);
 				};
 				break;
 			case EPolyEditExtrudeEdgeDirectionMode::SingleDirection:
@@ -166,7 +166,7 @@ TUniquePtr<UE::Geometry::FDynamicMeshOperator> UPolyEditExtrudeEdgeActivity::Mak
 		FVector3d WorldDestination = SingleDirectionVectorWorldSpace + ExtrudeFrameForGizmoWorldSpace.Origin;
 
 		Op->SingleDirectionVector = CurrentMeshTransform.InverseTransformPosition(WorldDestination)
-			- ExtrudeFrameForGizmoMeshSpace.Origin;
+			- ExtrudeFrameForGizmoMeshSpace.Frame.Origin;
 	}
 	else
 	{
@@ -174,7 +174,6 @@ TUniquePtr<UE::Geometry::FDynamicMeshOperator> UPolyEditExtrudeEdgeActivity::Mak
 		
 		Op->LocalFrameParams = ExtrudeFrameForGizmoMeshSpace.ToFramePoint(
 			CurrentMeshTransform.InverseTransformPosition(WorldDestination));
-		Op->LocalFrameParams /= ExtrudeFrameScaling;
 	}
 
 	return Op;
@@ -228,8 +227,13 @@ void UPolyEditExtrudeEdgeActivity::Setup(UInteractiveTool* ParentToolIn)
 
 		auto AdjustVectorLength = [this](FVector3d& VectorToAdjust)
 		{
+			if (bExtrudeDistanceWasNegative)
+			{
+				// The normlized extrude vector is actually in the opposite direction
+				VectorToAdjust *= -1;
+			}
 			VectorToAdjust.Normalize();
-			if (VectorToAdjust.IsNearlyZero())
+			if (VectorToAdjust.IsZero())
 			{
 				VectorToAdjust = FVector3d::UnitX();
 			}
@@ -244,6 +248,7 @@ void UPolyEditExtrudeEdgeActivity::Setup(UInteractiveTool* ParentToolIn)
 		{
 			AdjustVectorLength(ParamsInWorldExtrudeFrame);
 		}
+		bExtrudeDistanceWasNegative = Settings->Distance < 0;
 
 		if (bIsRunning)
 		{
@@ -592,7 +597,7 @@ void UPolyEditExtrudeEdgeActivity::UpdateDrawnPreviewEdges()
 			
 			FIndex2i EdgeVids = Mesh.GetEdgeV(Eid);
 			Lines.Add(FRenderableLine(Mesh.GetVertex(EdgeVids.A), Mesh.GetVertex(EdgeVids.B),
-				GroupLineColor, GroupLineThickness));
+				GroupLineColor, GroupLineThickness, GroupLineDepthBias));
 		});
 	});
 }
@@ -698,15 +703,14 @@ void UPolyEditExtrudeEdgeActivity::RecalculateGizmoExtrudeFrame()
 		NewVertData[0].SourceEidPair.B, 
 		ExtrudeFrame, Settings->bAdjustToExtrudeEvenly ? MAX_VERT_MOVEMENT_ADJUSTMENT_SCALE : 1.0);
 
-	ExtrudeFrameForGizmoMeshSpace = ExtrudeFrame.Frame;
-	ExtrudeFrameScaling = ExtrudeFrame.Scaling;
+	ExtrudeFrameForGizmoMeshSpace = ExtrudeFrame;
 
 	// This method of finding the extrude frame in world space isn't quite right in the case of nonuniform scale. However
 	// it seems hard to find any perfect solution, and in the end, it's not critical because we'll probably be able to get
 	// extrusion parameters out of an imperfect frame anyway.
 	ExtrudeFrameForGizmoWorldSpace = FFrame3d(
-		CurrentMeshTransform.TransformPosition(ExtrudeFrameForGizmoMeshSpace.Origin),
-		CurrentMeshTransform.TransformRotation(FQuat(ExtrudeFrameForGizmoMeshSpace.Rotation)));
+		CurrentMeshTransform.TransformPosition(ExtrudeFrameForGizmoMeshSpace.Frame.Origin),
+		CurrentMeshTransform.TransformRotation(FQuat(ExtrudeFrameForGizmoMeshSpace.Frame.Rotation)));
 }
 
 void UPolyEditExtrudeEdgeActivity::Render(IToolsContextRenderAPI* RenderAPI)
diff --git a/Engine/Plugins/Runtime/MeshModelingToolset/Source/MeshModelingTools/Private/ToolActivities/PolyEditExtrudeEdgeActivity.h b/Engine/Plugins/Runtime/MeshModelingToolset/Source/MeshModelingTools/Private/ToolActivities/PolyEditExtrudeEdgeActivity.h
--- a/Engine/Plugins/Runtime/MeshModelingToolset/Source/MeshModelingTools/Private/ToolActivities/PolyEditExtrudeEdgeActivity.h
+++ b/Engine/Plugins/Runtime/MeshModelingToolset/Source/MeshModelingTools/Private/ToolActivities/PolyEditExtrudeEdgeActivity.h
@@ -7,6 +7,7 @@
 #include "InteractiveToolActivity.h"
 #include "GroupTopology.h" // FGroupTopologySelection
 #include "ModelingOperators.h" // IDynamicMeshOperatorFactory
+#include "Operations/ExtrudeBoundaryEdges.h" // FExtrudeFrame
 
 #include "PolyEditExtrudeEdgeActivity.generated.h"
 
@@ -93,11 +94,10 @@ public:
 	// IDynamicMeshOperatorFactory
 	virtual TUniquePtr<UE::Geometry::FDynamicMeshOperator> MakeNewOperator() override;
 
+private:
 	UPROPERTY()
 	TObjectPtr<UPolyEditExtrudeEdgeActivityProperties> Settings = nullptr;
 
-protected:
-
 	UPROPERTY()
 	TObjectPtr<UPolyEditActivityContext> ActivityContext;
 
@@ -121,8 +121,8 @@ protected:
 	void ApplyExtrude();
 	void EndInternal();
 
-	FFrame3d ExtrudeFrameForGizmoMeshSpace;
-	FVector3d ExtrudeFrameScaling; // in mesh space, when using bAdjustToExtrudeEvenly 
+	UE::Geometry::FExtrudeBoundaryEdges::FExtrudeFrame ExtrudeFrameForGizmoMeshSpace;
+
 	FFrame3d ExtrudeFrameForGizmoWorldSpace;
 	
 	// These are the parameters we actually end up using, initialized from gizmos or detail panel.
@@ -150,4 +150,8 @@ protected:
 
 	// TODO: Might be worth having a getter in ActivityContext
 	FTransform CurrentMeshTransform;
+
+	// This is helpful so that in our distance watcher, we know that the previous distance
+	//  was negative, and therefore the extrude direction is actually opposite.
+	bool bExtrudeDistanceWasNegative = false;
 };
