MeshModelingTools: Fix a bug in PolyEdit where edges are not selectable in a small region near corners when also hitting a face, and tweak selection sensitivity to make edges and corners properly hittable when also hitting faces.

diff --git a/Engine/Plugins/Runtime/MeshModelingToolset/Source/ModelingComponents/Private/Selection/MeshTopologySelector.cpp b/Engine/Plugins/Runtime/MeshModelingToolset/Source/ModelingComponents/Private/Selection/MeshTopologySelector.cpp
--- a/Engine/Plugins/Runtime/MeshModelingToolset/Source/ModelingComponents/Private/Selection/MeshTopologySelector.cpp
+++ b/Engine/Plugins/Runtime/MeshModelingToolset/Source/ModelingComponents/Private/Selection/MeshTopologySelector.cpp
@@ -142,36 +142,36 @@ bool FMeshTopologySelector::FindSelectedElement(const FSelectionSettings& Settin
 		}
 	}
 
-	// if we have both corner and edge hit, pick the corner if it is within this (relative) tolerance to the hit point
-	constexpr double CornerProximityTolerance = 0.75;
-
-	// if we have both corner and edge, want to keep the one we are closer to
-	if (bHaveCornerHit && bHaveEdgeHit)
+	// If we have multiple types of hits, we need a prioritization to resolve between them. If we have a face hit,
+	//  we will still prefer corner and then edge, but we will lower our tolerance so that it is easier to select
+	//  the insides of small triangles on a mesh. The tolerance is different for corners and edges because we
+	//  expect the corners to be drawn a bit bigger (these specific values are tuned for PolyEdit)
+	if (bHaveFaceHit)
 	{
-		if (PointsWithinToleranceTest(CornerPosition, Ray.ClosestPoint(CornerPosition), CornerProximityTolerance))
+		constexpr double CornerFaceResolveProximityTolerance = 0.25;
+		constexpr double EdgeFaceResolveProximityTolerance = 0.20;
+		if (bHaveCornerHit)
 		{
-			bHaveEdgeHit = false;
+			bHaveCornerHit = PointsWithinToleranceTest(CornerPosition, Ray.ClosestPoint(CornerPosition), 
+				CornerFaceResolveProximityTolerance);
+			bHaveEdgeHit = bHaveEdgeHit && !bHaveCornerHit;
 		}
-		else
+		if (bHaveEdgeHit)
 		{
-			bHaveCornerHit = false;
+			ensure(!bHaveCornerHit); // cleared above, or else cleared bHaveEdgeHit
+			bHaveEdgeHit = PointsWithinToleranceTest(EdgePosition, Ray.ClosestPoint(EdgePosition),
+				EdgeFaceResolveProximityTolerance);
 		}
 	}
-
-	// if we have a corner or edge hit and a face hit, pick corner/edge if it is within this (relative) tolerance to the hit point
-	constexpr double EdgeOrCornerProximityTolerance = 0.15;
-
-	// if we have a corner or edge hit, and a face hit, pick face unless we are really close to corner/edge
-	if ((bHaveCornerHit || bHaveEdgeHit) && bHaveFaceHit)
+	// If we have both edge and corner, we prefer corner, but with a lower tolerance so that we can hit
+	//  more of the middle portion of the edge.
+	else if (bHaveCornerHit && bHaveEdgeHit)
 	{
-		FVector3d TestPos = (bHaveCornerHit) ? CornerPosition : EdgePosition;
-		if (!PointsWithinToleranceTest(TestPos, Ray.ClosestPoint(TestPos), EdgeOrCornerProximityTolerance))
-		{
-			bHaveEdgeHit = bHaveCornerHit = false;
-		}
+		constexpr double CornerEdgeResolveProximityTolerance = 0.75;
+		bHaveCornerHit = PointsWithinToleranceTest(CornerPosition, Ray.ClosestPoint(CornerPosition), CornerEdgeResolveProximityTolerance);
+		bHaveEdgeHit = !bHaveCornerHit;
 	}
 
-
 	if (bHaveCornerHit)
 	{
 		ResultOut = CornerResults;
