cis fix for incorrect init of temp tri selection in editmeshpolygonstool.cpp

diff --git a/Engine/Plugins/Runtime/MeshModelingToolset/Source/MeshModelingTools/Private/EditMeshPolygonsTool.cpp b/Engine/Plugins/Runtime/MeshModelingToolset/Source/MeshModelingTools/Private/EditMeshPolygonsTool.cpp
--- a/Engine/Plugins/Runtime/MeshModelingToolset/Source/MeshModelingTools/Private/EditMeshPolygonsTool.cpp
+++ b/Engine/Plugins/Runtime/MeshModelingToolset/Source/MeshModelingTools/Private/EditMeshPolygonsTool.cpp
@@ -877,7 +877,8 @@ void UEditMeshPolygonsTool::Setup()
 		{
 			// Convert default (no corner) group topology -> triangle topology -> tool (w/ corner) group topology
 			ConvertedSelection.InitializeTypes(EGeometryElementType::Edge, EGeometryTopologyType::Polygroup);
-			FGeometrySelection TempTriSelection(EGeometryElementType::Edge, EGeometryTopologyType::Triangle);
+			FGeometrySelection TempTriSelection;
+			TempTriSelection.InitializeTypes(EGeometryElementType::Edge, EGeometryTopologyType::Triangle);
 			FGroupTopology GroupTopology(CurrentMesh.Get(), true);
 			bCanUseSelection = UE::Geometry::ConvertSelection(*CurrentMesh, &GroupTopology, CurSelection, TempTriSelection, EEnumerateSelectionConversionParams::ContainSelection);
 			bCanUseSelection = bCanUseSelection && UE::Geometry::ConvertSelection(*CurrentMesh, Topology.Get(), TempTriSelection, ConvertedSelection, EEnumerateSelectionConversionParams::ContainSelection);
