MeshModelingTools: Make bridge action have correct selection in TriEdit after completion.

diff --git a/Engine/Plugins/Runtime/MeshModelingToolset/Source/MeshModelingTools/Private/EditMeshPolygonsTool.cpp b/Engine/Plugins/Runtime/MeshModelingToolset/Source/MeshModelingTools/Private/EditMeshPolygonsTool.cpp
--- a/Engine/Plugins/Runtime/MeshModelingToolset/Source/MeshModelingTools/Private/EditMeshPolygonsTool.cpp
+++ b/Engine/Plugins/Runtime/MeshModelingToolset/Source/MeshModelingTools/Private/EditMeshPolygonsTool.cpp
@@ -3074,7 +3074,7 @@ void UEditMeshPolygonsTool::ApplyBridgeEdges()
 	FGroupTopologySelection NewSelection;
 	for (int32 Tid : TrianglesToSelect)
 	{
-		NewSelection.SelectedGroupIDs.Add(Mesh->GetTriangleGroup(Tid));
+		NewSelection.SelectedGroupIDs.Add(Topology->GetGroupID(Tid));
 	}
 	if (ensure(!NewSelection.IsEmpty()))
 	{
