make smoothing operator at least try to respond to progresscancel being cancelled, to help more quickly reduce memory pressure from cancelled operators

diff --git a/Engine/Plugins/Runtime/MeshModelingToolset/Source/ModelingOperators/Private/SmoothingOps/CotanSmoothingOp.cpp b/Engine/Plugins/Runtime/MeshModelingToolset/Source/ModelingOperators/Private/SmoothingOps/CotanSmoothingOp.cpp
--- a/Engine/Plugins/Runtime/MeshModelingToolset/Source/ModelingOperators/Private/SmoothingOps/CotanSmoothingOp.cpp
+++ b/Engine/Plugins/Runtime/MeshModelingToolset/Source/ModelingOperators/Private/SmoothingOps/CotanSmoothingOp.cpp
@@ -32,14 +32,19 @@ double FCotanSmoothingOp::GetSmoothPower(int32 VertexID, bool bIsBoundary)
 void FCotanSmoothingOp::CalculateResult(FProgressCancel* Progress)
 {
 	// Update the values in the position buffer with smoothed positions.
-	Smooth();
+	Smooth(Progress);
+
+	if (Progress && Progress->Cancelled())
+	{
+		return;
+	}
 
 	// Copy the results back into the result mesh and update normals
 	UpdateResultMesh();
 
 }
 
-void FCotanSmoothingOp::Smooth()
+void FCotanSmoothingOp::Smooth(FProgressCancel* Progress)
 {
 	ELaplacianWeightScheme UseScheme = ELaplacianWeightScheme::ClampedCotangent;
 	if (SmoothOptions.bUniform)
@@ -47,9 +52,19 @@ void FCotanSmoothingOp::Smooth()
 		UseScheme = ELaplacianWeightScheme::Uniform;
 	}
 
+	if (Progress && Progress->Cancelled())
+	{
+		return;
+	}
+
 	TUniquePtr<UE::Solvers::IConstrainedMeshSolver> Smoother = UE::MeshDeformation::ConstructConstrainedMeshSmoother(
 		UseScheme, *ResultMesh);
 
+	if (Progress && Progress->Cancelled())
+	{
+		return;
+	}
+
 	if (SmoothOptions.SmoothPower < 0.0001)
 	{
 		for (int32 vid : ResultMesh->VertexIndicesItr())
@@ -77,6 +92,12 @@ void FCotanSmoothingOp::Smooth()
 
 			Smoother->AddConstraint(vid, Weight, Position, false);
 		}
+
+		if (Progress && Progress->Cancelled())
+		{
+			return;
+		}
+
 		Smoother->Deform(PositionBuffer);
 	}
 
diff --git a/Engine/Plugins/Runtime/MeshModelingToolset/Source/ModelingOperators/Public/SmoothingOps/CotanSmoothingOp.h b/Engine/Plugins/Runtime/MeshModelingToolset/Source/ModelingOperators/Public/SmoothingOps/CotanSmoothingOp.h
--- a/Engine/Plugins/Runtime/MeshModelingToolset/Source/ModelingOperators/Public/SmoothingOps/CotanSmoothingOp.h
+++ b/Engine/Plugins/Runtime/MeshModelingToolset/Source/ModelingOperators/Public/SmoothingOps/CotanSmoothingOp.h
@@ -24,7 +24,7 @@ public:
 
 private:
 	// Compute the smoothed result by using Cotan Biharmonic
-	void Smooth();	
+	void Smooth(FProgressCancel* Progress);
 
 	double GetSmoothPower(int32 VertexID, bool bIsBoundary);
 };
