Mitigate tools crashing from too many cancelled-but-still-running background tasks accumulating by limiting the number of active background tasks permitted for a single background compute object

#jria UE-220226

diff --git a/Engine/Plugins/Runtime/MeshModelingToolset/Source/ModelingComponents/ModelingComponents.Build.cs b/Engine/Plugins/Runtime/MeshModelingToolset/Source/ModelingComponents/ModelingComponents.Build.cs
--- a/Engine/Plugins/Runtime/MeshModelingToolset/Source/ModelingComponents/ModelingComponents.Build.cs
+++ b/Engine/Plugins/Runtime/MeshModelingToolset/Source/ModelingComponents/ModelingComponents.Build.cs
@@ -49,7 +49,8 @@ public class ModelingComponents : ModuleRules
 			PrivateDependencyModuleNames.AddRange(
 				new string[]
 				{
-					"UnrealEd"
+					"UnrealEd",
+					"Slate"
 				});
 		}
 
diff --git a/Engine/Plugins/Runtime/MeshModelingToolset/Source/ModelingComponents/Private/MeshOpPreviewHelpers.cpp b/Engine/Plugins/Runtime/MeshModelingToolset/Source/ModelingComponents/Private/MeshOpPreviewHelpers.cpp
--- a/Engine/Plugins/Runtime/MeshModelingToolset/Source/ModelingComponents/Private/MeshOpPreviewHelpers.cpp
+++ b/Engine/Plugins/Runtime/MeshModelingToolset/Source/ModelingComponents/Private/MeshOpPreviewHelpers.cpp
@@ -3,8 +3,41 @@
 #include "MeshOpPreviewHelpers.h"
 #include "Engine/World.h"
 
+#if WITH_EDITOR
+#include "Widgets/Notifications/SNotificationList.h"
+#include "Framework/Notifications/NotificationManager.h"
+#endif
+
 #include UE_INLINE_GENERATED_CPP_BY_NAME(MeshOpPreviewHelpers)
 
+
+#define LOCTEXT_NAMESPACE "MeshOpPreviewHelpers"
+
+namespace UE::Private::MeshOpPreviewLocal
+{
+	static void DisplayCriticalWarningMessage(const FText& InMessage, float ExpireDuration = 5.0f)
+	{
+#if WITH_EDITOR
+		FNotificationInfo Info(InMessage);
+		Info.ExpireDuration = ExpireDuration;
+		FSlateNotificationManager::Get().AddNotification(Info);
+#endif
+
+		UE_LOG(LogGeometry, Warning, TEXT("%s"), *InMessage.ToString());
+	}
+
+	static TAutoConsoleVariable<int32> CVarOverrideMaxBackgroundTasks(
+		TEXT("modeling.MaxBackgroundTasksOverride"), 0,
+		TEXT("Optional override for maximum allowed background tasks when generating preview results in tools. 0 to use default values. [def: 0]"));
+
+	int32 MaxActiveBackgroundTasksWithOverride(int32 MaxWithoutOverride)
+	{
+		int32 Override = CVarOverrideMaxBackgroundTasks.GetValueOnAnyThread();
+		return Override > 0 ? Override : MaxWithoutOverride;
+	}
+}
+
+
 using namespace UE::Geometry;
 
 void UMeshOpPreviewWithBackgroundCompute::Setup(UWorld* InWorld)
@@ -20,12 +53,14 @@ void UMeshOpPreviewWithBackgroundCompute::Setup(UWorld* InWorld, IDynamicMeshOpe
 {
 	Setup(InWorld);
 	BackgroundCompute = MakeUnique<FBackgroundDynamicMeshComputeSource>(OpGenerator);
+	BackgroundCompute->MaxActiveTaskCount = UE::Private::MeshOpPreviewLocal::MaxActiveBackgroundTasksWithOverride(MaxActiveBackgroundTasks);
 }
 
 void UMeshOpPreviewWithBackgroundCompute::ChangeOpFactory(IDynamicMeshOperatorFactory* OpGenerator)
 {
 	CancelCompute();
 	BackgroundCompute = MakeUnique<FBackgroundDynamicMeshComputeSource>(OpGenerator);
+	BackgroundCompute->MaxActiveTaskCount = UE::Private::MeshOpPreviewLocal::MaxActiveBackgroundTasksWithOverride(MaxActiveBackgroundTasks);
 	bResultValid = false;
 	bMeshInitialized = false;
 }
@@ -111,6 +146,16 @@ void UMeshOpPreviewWithBackgroundCompute::Tick(float DeltaTime)
 }
 
 
+void UMeshOpPreviewWithBackgroundCompute::SetMaxActiveBackgroundTasks(int32 InMaxActiveBackgroundTasks)
+{
+	MaxActiveBackgroundTasks = InMaxActiveBackgroundTasks;
+	if (BackgroundCompute)
+	{
+		BackgroundCompute->MaxActiveTaskCount = UE::Private::MeshOpPreviewLocal::MaxActiveBackgroundTasksWithOverride(MaxActiveBackgroundTasks);
+	}
+}
+
+
 void UMeshOpPreviewWithBackgroundCompute::UpdateResults()
 {
 	if (BackgroundCompute == nullptr)
@@ -143,6 +188,20 @@ void UMeshOpPreviewWithBackgroundCompute::UpdateResults()
 		ValidResultComputeTimeSeconds = Status.ElapsedTime;
 
 		OnMeshUpdated.Broadcast(this);
+
+		bWaitingForBackgroundTasks = false;
+	}
+	else if (int WaitingTaskCount; BackgroundCompute->IsWaitingForBackgroundTasks(WaitingTaskCount))
+	{
+		if (!bWaitingForBackgroundTasks)
+		{
+			UE::Private::MeshOpPreviewLocal::DisplayCriticalWarningMessage(LOCTEXT("TooManyBackgroundTasks", "Too many background tasks: Cancelling earlier tasks before generating new preview."));
+			bWaitingForBackgroundTasks = true;
+		}
+	}
+	else
+	{
+		bWaitingForBackgroundTasks = false;
 	}
 }
 
@@ -238,3 +297,4 @@ bool UMeshOpPreviewWithBackgroundCompute::IsUsingWorkingMaterial()
 		&& BackgroundCompute->GetElapsedComputeTime() > SecondsBeforeWorkingMaterial;
 }
 
+#undef LOCTEXT_NAMESPACE
diff --git a/Engine/Plugins/Runtime/MeshModelingToolset/Source/ModelingComponents/Public/MeshOpPreviewHelpers.h b/Engine/Plugins/Runtime/MeshModelingToolset/Source/ModelingComponents/Public/MeshOpPreviewHelpers.h
--- a/Engine/Plugins/Runtime/MeshModelingToolset/Source/ModelingComponents/Public/MeshOpPreviewHelpers.h
+++ b/Engine/Plugins/Runtime/MeshModelingToolset/Source/ModelingComponents/Public/MeshOpPreviewHelpers.h
@@ -209,6 +209,15 @@ public:
 	 */
 	bool IsUsingWorkingMaterial();
 
+	// Configure the maximum allowed number of background tasks
+	void SetMaxActiveBackgroundTasks(int32 NewMaxTasks);
+
+	// Configure maximum allowed number of background tasks heuristically from the input triangle count
+	void SetMaxActiveBackgroundTasksFromMeshSizeHeuristic(int32 InputMeshTriangleCount, int32 MaxSimultaneousTrianglesToProcess = 10000000, int32 MaxShouldNotExceed = 20)
+	{
+		SetMaxActiveBackgroundTasks(FMath::Clamp(MaxSimultaneousTrianglesToProcess / FMath::Max(1, InputMeshTriangleCount), 1, MaxShouldNotExceed));
+	}
+
 
 	//
 	// Change notification
@@ -276,6 +285,10 @@ protected:
 	// update the PreviewMesh if a new result is available from BackgroundCompute
 	void UpdateResults();
 
+private:
+
+	int32 MaxActiveBackgroundTasks = 5;
+	bool bWaitingForBackgroundTasks = false;
 };
 
 
diff --git a/Engine/Plugins/Runtime/MeshModelingToolset/Source/ModelingOperators/Public/BackgroundModelingComputeSource.h b/Engine/Plugins/Runtime/MeshModelingToolset/Source/ModelingOperators/Public/BackgroundModelingComputeSource.h
--- a/Engine/Plugins/Runtime/MeshModelingToolset/Source/ModelingOperators/Public/BackgroundModelingComputeSource.h
+++ b/Engine/Plugins/Runtime/MeshModelingToolset/Source/ModelingOperators/Public/BackgroundModelingComputeSource.h
@@ -146,6 +146,8 @@ protected:
 	mutable double LastEndTime = 0; // mutable because it must be set in CheckStatus, which is const
 	double LastInvalidateTime = 0;
 
+	TSharedPtr<std::atomic<int>> ActiveTaskCount;
+
 	void StartNewCompute();
 
 public:
@@ -153,6 +155,7 @@ public:
 		: OperatorSource(OperatorSourceIn)
 	{
 		TaskState = EBackgroundComputeTaskState::NotActive;
+		ActiveTaskCount = MakeShared<std::atomic<int>>(0);
 	}
 
 	~TBackgroundModelingComputeSource() {}
@@ -191,6 +194,17 @@ public:
 	 */
 	FStatus CheckStatus() const;
 
+	// If waiting for background tasks, returns true, and sets NumTasks to the number of active tasks
+	bool IsWaitingForBackgroundTasks(int32& NumTasks) const
+	{
+		if (bWaitingForActiveTasksToFinish)
+		{
+			NumTasks = *ActiveTaskCount;
+			return NumTasks >= MaxActiveTaskCount;
+		}
+		return false;
+	}
+
 	/**
 	 * @return The last computed Operator. This may only be called once, the caller then owns the Operator.
 	 */
@@ -211,6 +225,11 @@ public:
 public:
 	/** Default wait delay for cancel/restart cycle */
 	double CancelActiveOpDelaySeconds = 0.5;
+	/** Maximum number of active tasks to allow to run in the background until we wait to launch more */
+	int32 MaxActiveTaskCount = 5;
+
+private:
+	bool bWaitingForActiveTasksToFinish = false;
 };
 
 
@@ -229,10 +248,27 @@ void TBackgroundModelingComputeSource<OpType, OpTypeFactory>::Tick(float DeltaTi
 
 	if (TaskState == EBackgroundComputeTaskState::WaitingToCancel)
 	{
+		// immediately cancel (since that can take time), but only start new task if (1) delay time has elapsed and (2) active task count is acceptable
+		CancelActiveCompute();
+		bool bStartedTask = false;
 		if ((AccumTime - LastInvalidateTime) > CancelActiveOpDelaySeconds)
 		{
-			CancelActiveCompute();
-			StartNewCompute();
+			int Active = *ActiveTaskCount;
+			if (Active < MaxActiveTaskCount)
+			{
+				bWaitingForActiveTasksToFinish = false;
+				StartNewCompute();
+				bStartedTask = true;
+			}
+			else
+			{
+				bWaitingForActiveTasksToFinish = true;
+			}
+		}
+		if (!bStartedTask)
+		{
+			// failed to start new task, return to 'waiting to cancel' state
+			TaskState = EBackgroundComputeTaskState::WaitingToCancel;
 		}
 	}
 }
@@ -255,8 +291,10 @@ void TBackgroundModelingComputeSource<OpType, OpTypeFactory>::StartNewCompute()
 {
 	check(ActiveBackgroundTask == nullptr);
 
+	(*ActiveTaskCount)++;
 	TUniquePtr<OpType> NewOp = OperatorSource->MakeNewOperator();
 	ActiveBackgroundTask = new FAsyncTaskExecuterWithAbort<TModelingOpTask<OpType> >(MoveTemp(NewOp));
+	ActiveBackgroundTask->TaskCounter = ActiveTaskCount;
 	ActiveBackgroundTask->StartBackgroundTask();
 
 	LastStartTime = AccumTime;
@@ -291,7 +329,12 @@ TBackgroundModelingComputeSource<OpType, OpTypeFactory>::CheckStatus() const
 {
 	FStatus Status;
 
-	if (ActiveBackgroundTask == nullptr)
+	if (bWaitingForActiveTasksToFinish)
+	{
+		Status.TaskStatus = EBackgroundComputeTaskStatus::InProgress;
+		Status.ElapsedTime = AccumTime - LastStartTime;
+	}
+	else if (ActiveBackgroundTask == nullptr)
 	{
 		Status.TaskStatus = EBackgroundComputeTaskStatus::NotComputing;
 		Status.ElapsedTime = AccumTime - LastInvalidateTime;
diff --git a/Engine/Plugins/Runtime/MeshModelingToolset/Source/ModelingOperators/Public/ModelingTaskTypes.h b/Engine/Plugins/Runtime/MeshModelingToolset/Source/ModelingOperators/Public/ModelingTaskTypes.h
--- a/Engine/Plugins/Runtime/MeshModelingToolset/Source/ModelingOperators/Public/ModelingTaskTypes.h
+++ b/Engine/Plugins/Runtime/MeshModelingToolset/Source/ModelingOperators/Public/ModelingTaskTypes.h
@@ -6,6 +6,8 @@
 #include "Async/AsyncWork.h"
 #include "Util/ProgressCancel.h"
 
+#include <atomic>
+
 
 namespace UE
 {
@@ -74,6 +76,9 @@ public:
 	/** Set to true to abort the child task. CancelAndDelete() should be used in most cases, instead of changing bAbort directly. */
 	bool bAbort = false;
 
+	// Optional counter, if set will be decremented on task destruction
+	TSharedPtr<std::atomic<int>> TaskCounter;
+
 	template <typename Arg0Type, typename... ArgTypes>
 	FAsyncTaskExecuterWithAbort(Arg0Type&& Arg0, ArgTypes&&... Args)
 		: FAsyncTask<TTask>(Forward<Arg0Type>(Arg0), Forward<ArgTypes>(Args)...)
@@ -82,6 +87,13 @@ public:
 		FAsyncTask<TTask>::GetTask().SetAbortSource(&bAbort);
 	}
 
+	virtual ~FAsyncTaskExecuterWithAbort()
+	{
+		if (TaskCounter.Get())
+		{
+			--(*TaskCounter);
+		}
+	}
 
 	/**
 	 * Tells the child FAbandonableTask to terminate itself, via the bAbort flag
