FSkinnedComponentWrapper: make sure that the requested LOD for conversion is available

#rnx

diff --git a/Engine/Plugins/Runtime/MeshModelingToolset/Source/ModelingComponents/Private/ConversionUtils/SkinnedMeshToDynamicMesh.cpp b/Engine/Plugins/Runtime/MeshModelingToolset/Source/ModelingComponents/Private/ConversionUtils/SkinnedMeshToDynamicMesh.cpp
--- a/Engine/Plugins/Runtime/MeshModelingToolset/Source/ModelingComponents/Private/ConversionUtils/SkinnedMeshToDynamicMesh.cpp
+++ b/Engine/Plugins/Runtime/MeshModelingToolset/Source/ModelingComponents/Private/ConversionUtils/SkinnedMeshToDynamicMesh.cpp
@@ -41,22 +41,28 @@ public:
 	, bUseDisabledSections(UseDisabledSections)
 	, bHasTangents(bWantTangents)
 	, Asset(SkinnedMeshComponent.GetSkinnedAsset())
-	, SrcLODInfo(SkinnedMeshComponent.GetSkinnedAsset()->GetLODInfo(RequestedLOD))
 	, SkeletalMeshRenderData(SkinnedMeshComponent.MeshObject->GetSkeletalMeshRenderData())
-	, IndexBuffer(SkinnedMeshComponent.MeshObject->GetSkeletalMeshRenderData().LODRenderData[RequestedLOD].MultiSizeIndexContainer.GetIndexBuffer())
-
 	{
-
-		// This step does the actual skinning (and isn't as const as the api suggests..)
-		Component->GetCPUSkinnedVertices(SkinnedVertices, RequestedLOD);
-		const FSkeletalMeshLODRenderData& LODData = SkeletalMeshRenderData.LODRenderData[RequestedLOD];
+		const int32 MinLODIndex = SkinnedMeshComponent.ComputeMinLOD();
+		const int32 MaxLODIndex = SkinnedMeshComponent.GetNumLODs() - 1;
+		LOD = FMath::Clamp<int32>(RequestedLOD, MinLODIndex, MaxLODIndex);
+		if (LOD != RequestedLOD)
+		{
+			UE_LOG(LogGeometry, Warning, TEXT("Requesting LOD %d within [%d, %d], clamping to %d"), RequestedLOD, MinLODIndex, MaxLODIndex, LOD);
+		}
+		
+		SrcLODInfo = Asset->GetLODInfo(LOD);
+		const FSkeletalMeshLODRenderData& LODData = SkeletalMeshRenderData.LODRenderData[LOD];
+		IndexBuffer = LODData.MultiSizeIndexContainer.GetIndexBuffer();
 		
+		// This step does the actual skinning (and isn't as const as the api suggests..)
+		Component->GetCPUSkinnedVertices(SkinnedVertices, LOD);
 
 		const int32 NumSections = LODData.RenderSections.Num();
 		
 		auto SkipSection = [&](const FSkelMeshRenderSection& SkelMeshSection)
 							{
-								return ( !bUseDisabledSections && !SkinnedMeshComponent.IsMaterialSectionShown(SkelMeshSection.MaterialIndex, RequestedLOD));
+								return ( !bUseDisabledSections && !SkinnedMeshComponent.IsMaterialSectionShown(SkelMeshSection.MaterialIndex, LOD));
 							};
 
 		// pre-compute number of tris & verts
@@ -107,6 +113,7 @@ public:
 		{
 			UE_LOG(LogGeometry, Display, TEXT("LOD vertex mismatch with [%s -> %s]"), *SkinnedMeshComponent.GetPathName(), *SkinnedMeshComponent.GetSkinnedAsset()->GetPathName());
 			UE_LOG(LogGeometry, Display, TEXT(" - Total Vertex Count: %d"), SkinnedVertices.Num());
+			UE_LOG(LogGeometry, Display, TEXT(" - Render Data Vertex Count: %d"), LODData.StaticVertexBuffers.PositionVertexBuffer.GetNumVertices());
 			UE_LOG(LogGeometry, Display, TEXT(" - Total Index/Triangle Count: %d / %d "), IndexBuffer->Num(), IndexBuffer->Num() / 3);
 		 
 			for (int32 SectionIndex = 0; SectionIndex < LODData.RenderSections.Num(); SectionIndex++)
