make basemeshprocessing- and basecreatefromselected-derived tools set a max number of background tasks based on the input mesh sizes
#jria UE-220226

diff --git a/Engine/Plugins/Runtime/MeshModelingToolset/Source/ModelingComponents/Private/BaseTools/BaseCreateFromSelectedTool.cpp b/Engine/Plugins/Runtime/MeshModelingToolset/Source/ModelingComponents/Private/BaseTools/BaseCreateFromSelectedTool.cpp
--- a/Engine/Plugins/Runtime/MeshModelingToolset/Source/ModelingComponents/Private/BaseTools/BaseCreateFromSelectedTool.cpp
+++ b/Engine/Plugins/Runtime/MeshModelingToolset/Source/ModelingComponents/Private/BaseTools/BaseCreateFromSelectedTool.cpp
@@ -89,6 +89,13 @@ void UBaseCreateFromSelectedTool::Setup()
 	Preview->Setup(GetTargetWorld(), this);
 	ToolSetupUtil::ApplyRenderingConfigurationToPreview(Preview->PreviewMesh, nullptr);
 
+	int32 SourcesTriCount = 0;
+	for (int32 ComponentIdx = 0; ComponentIdx < Targets.Num(); ComponentIdx++)
+	{
+		SourcesTriCount += UE::ToolTarget::GetTriangleCount(Targets[ComponentIdx]);
+	}
+	Preview->SetMaxActiveBackgroundTasksFromMeshSizeHeuristic(SourcesTriCount);
+
 	SetPreviewCallbacks();
 	Preview->OnMeshUpdated.AddLambda(
 		[this](const UMeshOpPreviewWithBackgroundCompute* UpdatedPreview)
diff --git a/Engine/Plugins/Runtime/MeshModelingToolset/Source/ModelingComponents/Private/BaseTools/BaseMeshProcessingTool.cpp b/Engine/Plugins/Runtime/MeshModelingToolset/Source/ModelingComponents/Private/BaseTools/BaseMeshProcessingTool.cpp
--- a/Engine/Plugins/Runtime/MeshModelingToolset/Source/ModelingComponents/Private/BaseTools/BaseMeshProcessingTool.cpp
+++ b/Engine/Plugins/Runtime/MeshModelingToolset/Source/ModelingComponents/Private/BaseTools/BaseMeshProcessingTool.cpp
@@ -135,6 +135,7 @@ void UBaseMeshProcessingTool::Setup()
 	Preview = NewObject<UMeshOpPreviewWithBackgroundCompute>(this, "Preview");
 	Preview->Setup(GetTargetWorld(), this); // Adds the actual functional tool in the Preview object
 	Preview->PreviewMesh->SetTangentsMode(EDynamicMeshComponentTangentsMode::AutoCalculated);
+	Preview->SetMaxActiveBackgroundTasksFromMeshSizeHeuristic(InitialMesh.TriangleCount());
 	ToolSetupUtil::ApplyRenderingConfigurationToPreview(Preview->PreviewMesh, Target);
 
 	FComponentMaterialSet MaterialSet;
diff --git a/Engine/Plugins/Runtime/MeshModelingToolset/Source/ModelingComponents/Private/ModelingToolTargetUtil.cpp b/Engine/Plugins/Runtime/MeshModelingToolset/Source/ModelingComponents/Private/ModelingToolTargetUtil.cpp
--- a/Engine/Plugins/Runtime/MeshModelingToolset/Source/ModelingComponents/Private/ModelingToolTargetUtil.cpp
+++ b/Engine/Plugins/Runtime/MeshModelingToolset/Source/ModelingComponents/Private/ModelingToolTargetUtil.cpp
@@ -231,6 +231,27 @@ FDynamicMesh3 UE::ToolTarget::GetDynamicMeshCopy(UToolTarget* Target, bool bWant
 	return GetDynamicMeshCopy(Target, GetMeshParams);
 }
 
+int32 UE::ToolTarget::GetTriangleCount(UToolTarget* Target)
+{
+	IPersistentDynamicMeshSource* DynamicMeshSource = Cast<IPersistentDynamicMeshSource>(Target);
+	if (DynamicMeshSource)
+	{
+		UDynamicMesh* DynamicMesh = DynamicMeshSource->GetDynamicMeshContainer();
+		return DynamicMesh->GetTriangleCount();
+	}
+
+	IMeshDescriptionProvider* MeshDescriptionProvider = Cast<IMeshDescriptionProvider>(Target);
+	if (MeshDescriptionProvider)
+	{
+		if (const FMeshDescription* MeshDescription = MeshDescriptionProvider->GetMeshDescription())
+		{
+			return MeshDescription->Triangles().Num();
+		}
+	}
+
+	return 0;
+}
+
 FDynamicMesh3 UE::ToolTarget::GetDynamicMeshCopy(UToolTarget* Target, const FGetMeshParameters& InGetMeshParams)
 {
 	IPersistentDynamicMeshSource* DynamicMeshSource = Cast<IPersistentDynamicMeshSource>(Target);
diff --git a/Engine/Plugins/Runtime/MeshModelingToolset/Source/ModelingComponents/Public/ModelingToolTargetUtil.h b/Engine/Plugins/Runtime/MeshModelingToolset/Source/ModelingComponents/Public/ModelingToolTargetUtil.h
--- a/Engine/Plugins/Runtime/MeshModelingToolset/Source/ModelingComponents/Public/ModelingToolTargetUtil.h
+++ b/Engine/Plugins/Runtime/MeshModelingToolset/Source/ModelingComponents/Public/ModelingToolTargetUtil.h
@@ -155,6 +155,9 @@ MODELINGCOMPONENTS_API FMeshDescription GetMeshDescriptionCopy(
 		UToolTarget* Target,
 		const FGetMeshParameters& InGetMeshParams = FGetMeshParameters());
 
+/** @return The triangle count of Target's persistent dynamic mesh or its mesh description, if available, or 0 otherwise */
+MODELINGCOMPONENTS_API int32 GetTriangleCount(UToolTarget* Target);
+
 /**
  * EDynamicMeshUpdateResult is returned by functions below that update a ToolTarget with a new Mesh
  */
