Geoscripting: Prefer deforming the mesh description directly rather than always go through render data

#rnx

diff --git a/Engine/Plugins/Runtime/MeshModelingToolset/Source/ModelingComponents/Private/ConversionUtils/SceneComponentToDynamicMesh.cpp b/Engine/Plugins/Runtime/MeshModelingToolset/Source/ModelingComponents/Private/ConversionUtils/SceneComponentToDynamicMesh.cpp
--- a/Engine/Plugins/Runtime/MeshModelingToolset/Source/ModelingComponents/Private/ConversionUtils/SceneComponentToDynamicMesh.cpp
+++ b/Engine/Plugins/Runtime/MeshModelingToolset/Source/ModelingComponents/Private/ConversionUtils/SceneComponentToDynamicMesh.cpp
@@ -29,6 +29,7 @@
 #include "MeshDescriptionToDynamicMesh.h"
 #include "Physics/ComponentCollisionUtil.h"
 #include "PlanarCut.h"
+#include "SkeletalMeshOperations.h"
 #include "StaticMeshAttributes.h"
 #include "StaticMeshLODResourcesToDynamicMesh.h"
 #include "StaticMeshOperations.h"
@@ -431,6 +432,117 @@ namespace Private::ConversionHelper
 			return CopyMeshFromStaticMesh_SourceData(FromStaticMeshAsset, AssetOptions, LODType, LODIndex, OutMesh, OutErrorMessage);
 		}
 	}
+	
+	
+	static bool CopyMeshFromSkinnedAsset(
+		USkinnedAsset* FromSkinnedAsset,
+		USkinnedMeshComponent* SkinnedMeshComponent,
+		EMeshLODType LODType,
+		int32 LODIndex,
+		bool bUseClosestLOD,
+		bool bWantTangents,
+		FDynamicMesh3& OutMesh,
+		FText& OutErrorMessage
+	)
+	{
+		if (!FromSkinnedAsset)
+		{
+			OutErrorMessage = LOCTEXT("CopyMeshFromSkinnedAsset_NullMesh", "Skinned mesh is null");
+			return false;
+		}
+
+		USkeletalMesh* SkeletalMesh = Cast<USkeletalMesh>(FromSkinnedAsset);
+
+		// If using non-skeletal mesh variations of skinned meshes, just go straight to render data.
+		if (!SkeletalMesh)
+		{
+			LODType = EMeshLODType::RenderData;
+		}
+
+		if (bUseClosestLOD)
+		{
+			// attempt to detect if an unavailable LOD was requested, and if so re-map to an available one
+			if (LODType == EMeshLODType::MaxAvailable || LODType == EMeshLODType::HiResSourceModel)
+			{
+				LODIndex = 0;
+			}
+#if WITH_EDITOR
+			if (LODType == EMeshLODType::MaxAvailable || LODType == EMeshLODType::HiResSourceModel)
+			{
+				LODType = EMeshLODType::SourceModel;
+			}
+			if (LODType == EMeshLODType::SourceModel)
+			{
+				LODIndex = FMath::Clamp(LODIndex, 0, SkeletalMesh->GetNumSourceModels() - 1);
+				if (!SkeletalMesh->GetSourceModel(LODIndex).HasMeshDescription())
+				{
+					LODType = EMeshLODType::RenderData;
+				}
+			}
+			if (LODType == EMeshLODType::RenderData)
+			{
+				LODIndex = FMath::Clamp(LODIndex, 0, FromSkinnedAsset->GetLODNum() - 1);
+			}
+#else
+			LODType = EMeshLODType::RenderData;
+			LODIndex = FMath::Clamp(LODIndex, 0, FromSkinnedAsset->GetLODNum() - 1);
+#endif
+		}
+
+		if (LODType == EMeshLODType::RenderData)
+		{
+			return SkinnedMeshComponentToDynamicMesh(*SkinnedMeshComponent, OutMesh, LODIndex, bWantTangents);
+		}
+		else
+		{
+#if WITH_EDITOR
+			const FMeshDescription* SourceMesh = nullptr;
+
+			// Check first if we have bulk data available and non-empty.
+			if (SkeletalMesh->HasMeshDescription(LODIndex))
+			{
+				SourceMesh = SkeletalMesh->GetMeshDescription(LODIndex); 
+			}
+			if (SourceMesh == nullptr)
+			{
+				OutErrorMessage = LOCTEXT("CopyMeshFromSkinnedAsset_LODNotAvailable", "Requested LOD source mesh is not available");
+				return false;
+			}
+
+			TMap<FName, float> MorphTargetWeights;
+
+			for (const TPair<const UMorphTarget*, int32>& MorphTarget: SkinnedMeshComponent->ActiveMorphTargets)
+			{
+				const FName MorphName = MorphTarget.Key->GetFName();
+				const float MorphWeight = SkinnedMeshComponent->MorphTargetWeights[MorphTarget.Value];
+
+				MorphTargetWeights.Add(MorphName, MorphWeight);
+			}
+			
+			const TArray<FTransform>& ComponentSpaceTransforms = SkinnedMeshComponent->GetComponentSpaceTransforms();
+			FMeshDescription DeformedMesh;
+			if (!FSkeletalMeshOperations::GetPosedMesh(*SourceMesh, DeformedMesh, ComponentSpaceTransforms, NAME_None, MorphTargetWeights))
+			{
+				OutErrorMessage = LOCTEXT("CopyMeshFromSkinnedAsset_CannotPose", "Unable to pose the source mesh");
+				return false;
+			}
+			
+			FDynamicMesh3 NewMesh;
+			FMeshDescriptionToDynamicMesh Converter;
+
+			// Leave this on, since the set morph target node uses this. 
+			Converter.bVIDsFromNonManifoldMeshDescriptionAttr = true;
+			
+			Converter.Convert(&DeformedMesh, OutMesh, bWantTangents);
+		
+			return true;
+#else
+			OutErrorMessage = LOCTEXT("CopyMeshFromSkinnedAsset_EditorOnly", "Source Models are not available at Runtime");
+			return false;
+#endif
+		}
+	}
+	
 }
 
 TArray<int32> GetPolygonGroupToMaterialIndexMap(const UStaticMesh* StaticMesh, EMeshLODType LODType, int32 LODIndex)
@@ -514,29 +626,29 @@ bool SceneComponentToDynamicMesh(USceneComponent* Component, const FToMeshOption
 		}
 		else
 		{
-			USkinnedAsset* SkinnedAsset = SkinnedMeshComponent->GetSkinnedAsset();
-			if (SkinnedAsset)
+			if (USkinnedAsset* SkinnedAsset = SkinnedMeshComponent->GetSkinnedAsset())
 			{
-				SkinnedMeshComponentToDynamicMesh(*SkinnedMeshComponent, OutMesh, RequestedLOD, Options.bWantTangents);
-				OutMesh.DiscardTriangleGroups();
-
-				if (OutAssetMaterials)
+				bSuccess = Private::ConversionHelper::CopyMeshFromSkinnedAsset(SkinnedAsset, SkinnedMeshComponent, Options.LODType, Options.LODIndex, Options.bUseClosestLOD, Options.bWantTangents, OutMesh, OutErrorMessage);
+				if (bSuccess)
 				{
-					const TArray<FSkeletalMaterial>& Materials = SkinnedAsset->GetMaterials();
-					OutAssetMaterials->SetNum(Materials.Num());
-					for (int32 k = 0; k < Materials.Num(); ++k)
+					OutMesh.DiscardTriangleGroups();
+
+					if (OutAssetMaterials)
 					{
-						(*OutAssetMaterials)[k] = Materials[k].MaterialInterface;
+						const TArray<FSkeletalMaterial>& Materials = SkinnedAsset->GetMaterials();
+						OutAssetMaterials->SetNum(Materials.Num());
+						for (int32 k = 0; k < Materials.Num(); ++k)
+						{
+							(*OutAssetMaterials)[k] = Materials[k].MaterialInterface;
+						}
 					}
 				}
-				bSuccess = true;
 			}
 			else
 			{
 				OutErrorMessage = LOCTEXT("CopyMeshFromComponent_MissingSkinnedAsset", "SkinnedMeshComponent has a null SkinnedAsset");
 			}
 		}
-
 	}
 	else if (USplineMeshComponent* SplineMeshComponent = Cast<USplineMeshComponent>(Component))
 	{
diff --git a/Engine/Plugins/Runtime/MeshModelingToolset/Source/ModelingComponents/Private/ConversionUtils/SkinnedMeshToDynamicMesh.cpp b/Engine/Plugins/Runtime/MeshModelingToolset/Source/ModelingComponents/Private/ConversionUtils/SkinnedMeshToDynamicMesh.cpp
--- a/Engine/Plugins/Runtime/MeshModelingToolset/Source/ModelingComponents/Private/ConversionUtils/SkinnedMeshToDynamicMesh.cpp
+++ b/Engine/Plugins/Runtime/MeshModelingToolset/Source/ModelingComponents/Private/ConversionUtils/SkinnedMeshToDynamicMesh.cpp
@@ -446,7 +446,7 @@ private:
 
 
 
-void SkinnedMeshComponentToDynamicMesh(USkinnedMeshComponent& SkinnedMeshComponent, Geometry::FDynamicMesh3& MeshOut, int32 RequestedLOD, bool bWantTangents)
+bool SkinnedMeshComponentToDynamicMesh(USkinnedMeshComponent& SkinnedMeshComponent, Geometry::FDynamicMesh3& MeshOut, int32 RequestedLOD, bool bWantTangents)
 {
 
 	MeshOut.Clear();
@@ -456,17 +456,17 @@ void SkinnedMeshComponentToDynamicMesh(USkinnedMeshComponent& SkinnedMeshCompone
 
 	if (!bIsValidSkinnedMeshComponent)
 	{
-		return;
+		return false;
 	}
 
 	if (RequestedLOD < 0 || RequestedLOD > NumLODs - 1)
 	{
-		return;
+		return false;
 	}
 
 	if (!SkinnedMeshComponent.GetSkinnedAsset())
 	{
-		return;
+		return false;
 	}
 
 
@@ -483,6 +483,7 @@ void SkinnedMeshComponentToDynamicMesh(USkinnedMeshComponent& SkinnedMeshCompone
 		SkinnedComponentConverter.Convert(MeshOut, SkinnedComponentWrapper, TriToGroupID, TriToMaterialID, bCopyTangents);
 	}
 
+	return true;
 }
 
 } // end namespace Conversion
diff --git a/Engine/Plugins/Runtime/MeshModelingToolset/Source/ModelingComponents/Public/ConversionUtils/SkinnedMeshToDynamicMesh.h b/Engine/Plugins/Runtime/MeshModelingToolset/Source/ModelingComponents/Public/ConversionUtils/SkinnedMeshToDynamicMesh.h
--- a/Engine/Plugins/Runtime/MeshModelingToolset/Source/ModelingComponents/Public/ConversionUtils/SkinnedMeshToDynamicMesh.h
+++ b/Engine/Plugins/Runtime/MeshModelingToolset/Source/ModelingComponents/Public/ConversionUtils/SkinnedMeshToDynamicMesh.h
@@ -24,7 +24,7 @@ namespace Conversion {
 	* @param RequestedLOD  - the LOD to be converted, if a non-existent LOD is requested the result mesh will be empty.
 	* @param bWantTangents - controls if tangents are transfered
 	*/
-	void MODELINGCOMPONENTS_API SkinnedMeshComponentToDynamicMesh(USkinnedMeshComponent& SkinnedMeshComponent, Geometry::FDynamicMesh3& MeshOut, int32 RequestedLOD,  bool bWantTangents);
+	bool MODELINGCOMPONENTS_API SkinnedMeshComponentToDynamicMesh(USkinnedMeshComponent& SkinnedMeshComponent, Geometry::FDynamicMesh3& MeshOut, int32 RequestedLOD,  bool bWantTangents);
 
 } // end namespace Geometry
 } // end namespace UE
\ No newline at end of file
