make most primitives not apply xy bounds centering, since they are already centered and the bounds center can be slightly off the 'true' center especially at very low sample counts

diff --git a/Engine/Plugins/Runtime/MeshModelingToolset/Source/MeshModelingTools/Private/AddPrimitiveTool.cpp b/Engine/Plugins/Runtime/MeshModelingToolset/Source/MeshModelingTools/Private/AddPrimitiveTool.cpp
--- a/Engine/Plugins/Runtime/MeshModelingToolset/Source/MeshModelingTools/Private/AddPrimitiveTool.cpp
+++ b/Engine/Plugins/Runtime/MeshModelingToolset/Source/MeshModelingTools/Private/AddPrimitiveTool.cpp
@@ -363,6 +363,10 @@ void UAddPrimitiveTool::UpdatePreviewMesh() const
 	// set mesh position
 	const FAxisAlignedBox3d Bounds = NewMesh.GetBounds(true);
 	FVector3d TargetOrigin = Bounds.Center();
+	if (!ShouldCenterXY())
+	{
+		TargetOrigin.X = TargetOrigin.Y = 0;
+	}
 	if (ShapeSettings->PivotLocation == EMakeMeshPivotLocation::Base)
 	{
 		TargetOrigin.Z = Bounds.Min.Z;
diff --git a/Engine/Plugins/Runtime/MeshModelingToolset/Source/MeshModelingTools/Public/AddPrimitiveTool.h b/Engine/Plugins/Runtime/MeshModelingToolset/Source/MeshModelingTools/Public/AddPrimitiveTool.h
--- a/Engine/Plugins/Runtime/MeshModelingToolset/Source/MeshModelingTools/Public/AddPrimitiveTool.h
+++ b/Engine/Plugins/Runtime/MeshModelingToolset/Source/MeshModelingTools/Public/AddPrimitiveTool.h
@@ -550,6 +550,13 @@ protected:
 
 	void UpdateTargetSurface();
 
+	// @return true if the primitive needs to be centered in the XY plane when placed.
+	virtual bool ShouldCenterXY() const
+	{
+		// Most primitives are already XY centered, and re-centering them only introduces issues at very low samplings where the bounds center is offset from the intended center.
+		return false;
+	}
+
 	// Used to make the initial placement of the mesh undoable
 	class FStateChange : public FToolCommandChange
 	{
@@ -675,6 +682,10 @@ public:
 	explicit UAddStairsPrimitiveTool(const FObjectInitializer& ObjectInitializer);
 protected:
 	virtual void GenerateMesh(FDynamicMesh3* OutMesh) const override;
+	virtual bool ShouldCenterXY() const override
+	{
+		return true;
+	}
 };
 
 
