make polygon selection mechanic translate compacted edge selections for output

diff --git a/Engine/Plugins/Runtime/MeshModelingToolset/Source/ModelingComponents/Private/Selection/PolygonSelectionMechanic.cpp b/Engine/Plugins/Runtime/MeshModelingToolset/Source/ModelingComponents/Private/Selection/PolygonSelectionMechanic.cpp
--- a/Engine/Plugins/Runtime/MeshModelingToolset/Source/ModelingComponents/Private/Selection/PolygonSelectionMechanic.cpp
+++ b/Engine/Plugins/Runtime/MeshModelingToolset/Source/ModelingComponents/Private/Selection/PolygonSelectionMechanic.cpp
@@ -59,7 +59,6 @@ void UPolygonSelectionMechanic::GetSelection_AsGroupTopology(UE::Geometry::FGeom
 	}
 	else if (SelectionOut.ElementType == EGeometryElementType::Edge)
 	{
-		// TODO: not sure how we can apply compact maps here because mapping does not included compacted edges? are edges even compacted?
 		if ( CompactMapsToApply == nullptr )
 		{
 			for (int32 GroupEdgeID : CurSelection.SelectedEdgeIDs)
@@ -69,6 +68,24 @@ void UPolygonSelectionMechanic::GetSelection_AsGroupTopology(UE::Geometry::FGeom
 				SelectionOut.Selection.Add( FGeoSelectionID(TriEdgeID.Encoded(), GroupEdgeID).Encoded() );
 			}
 		}
+		else
+		{
+			for (int32 GroupEdgeID : CurSelection.SelectedEdgeIDs)
+			{
+				const TArray<int>& GroupEdgeVerts = Topology->GetGroupEdgeVertices(GroupEdgeID);
+				if (GroupEdgeVerts.Num() > 1)
+				{
+					int32 VID0 = CompactMapsToApply->GetVertexMapping(GroupEdgeVerts[0]);
+					int32 VID1 = CompactMapsToApply->GetVertexMapping(GroupEdgeVerts[1]);
+					int32 FoundEID = Topology->GetMesh()->FindEdge(VID0, VID1);
+					if (FoundEID != IndexConstants::InvalidID)
+					{
+						FMeshTriEdgeID TriEdgeID = Topology->GetMesh()->GetTriEdgeIDFromEdgeID(FoundEID);
+						SelectionOut.Selection.Add(FGeoSelectionID(TriEdgeID.Encoded(), GroupEdgeID).Encoded());
+					}
+				}
+			}
+		}
 	}
 	else if (SelectionOut.ElementType == EGeometryElementType::Face)
 	{
@@ -87,7 +104,7 @@ void UPolygonSelectionMechanic::GetSelection_AsGroupTopology(UE::Geometry::FGeom
 
 void UPolygonSelectionMechanic::GetSelection_AsTriangleTopology(UE::Geometry::FGeometrySelection& SelectionOut, const FCompactMaps* CompactMapsToApply) const
 {
-	// note: this is currently the same code as GetSelection_AsGroupTopology() except for the topology-type verification check
+	// note: this is currently the same code as GetSelection_AsGroupTopology() except for the topology-type verification check, and the topology type of the selected elements
 
 	const FGroupTopologySelection& CurSelection = PersistentSelection;
 	if (SelectionOut.TopologyType != EGeometryTopologyType::Triangle)
@@ -108,7 +125,6 @@ void UPolygonSelectionMechanic::GetSelection_AsTriangleTopology(UE::Geometry::FG
 	}
 	else if (SelectionOut.ElementType == EGeometryElementType::Edge)
 	{
-		// TODO: not sure how we can apply compact maps here because mapping does not included compacted edges? are edges even compacted?
 		if ( CompactMapsToApply == nullptr )
 		{
 			for (int32 GroupEdgeID : CurSelection.SelectedEdgeIDs)
@@ -118,6 +134,24 @@ void UPolygonSelectionMechanic::GetSelection_AsTriangleTopology(UE::Geometry::FG
 				SelectionOut.Selection.Add( FGeoSelectionID::MeshEdge(TriEdgeID).Encoded() );
 			}
 		}
+		else
+		{
+			for (int32 GroupEdgeID : CurSelection.SelectedEdgeIDs)
+			{
+				const TArray<int>& GroupEdgeVerts = Topology->GetGroupEdgeVertices(GroupEdgeID);
+				if (GroupEdgeVerts.Num() > 1)
+				{
+					int32 VID0 = CompactMapsToApply->GetVertexMapping(GroupEdgeVerts[0]);
+					int32 VID1 = CompactMapsToApply->GetVertexMapping(GroupEdgeVerts[1]);
+					int32 FoundEID = Topology->GetMesh()->FindEdge(VID0, VID1);
+					if (FoundEID != IndexConstants::InvalidID)
+					{
+						FMeshTriEdgeID TriEdgeID = Topology->GetMesh()->GetTriEdgeIDFromEdgeID(FoundEID);
+						SelectionOut.Selection.Add(FGeoSelectionID::MeshEdge(TriEdgeID).Encoded());
+					}
+				}
+			}
+		}
 	}
 	else if (SelectionOut.ElementType == EGeometryElementType::Face)
 	{
diff --git a/Engine/Source/Runtime/GeometryCore/Public/DynamicMesh/DynamicMesh3.h b/Engine/Source/Runtime/GeometryCore/Public/DynamicMesh/DynamicMesh3.h
--- a/Engine/Source/Runtime/GeometryCore/Public/DynamicMesh/DynamicMesh3.h
+++ b/Engine/Source/Runtime/GeometryCore/Public/DynamicMesh/DynamicMesh3.h
@@ -1204,7 +1204,6 @@ public:
 	 * Should be faster if the amount of compacting is not too significant, and is useful in some places.
 	 *
 	 * @param CompactInfo if not nullptr, will be filled with mapping indicating how vertex and triangle IDs were changed during compaction
-	 * @todo VertexEdgeLists is not compacted. does not affect indices, but does keep memory.
 	 */
 	GEOMETRYCORE_API void CompactInPlace(FCompactMaps* CompactInfo = nullptr);
 
