fix responsiveness of background operators in cases where we show 'dirty' results after operators are cancelled

diff --git a/Engine/Plugins/Runtime/MeshModelingToolset/Source/ModelingOperators/Public/BackgroundModelingComputeSource.h b/Engine/Plugins/Runtime/MeshModelingToolset/Source/ModelingOperators/Public/BackgroundModelingComputeSource.h
--- a/Engine/Plugins/Runtime/MeshModelingToolset/Source/ModelingOperators/Public/BackgroundModelingComputeSource.h
+++ b/Engine/Plugins/Runtime/MeshModelingToolset/Source/ModelingOperators/Public/BackgroundModelingComputeSource.h
@@ -248,28 +248,23 @@ void TBackgroundModelingComputeSource<OpType, OpTypeFactory>::Tick(float DeltaTi
 
 	if (TaskState == EBackgroundComputeTaskState::WaitingToCancel)
 	{
-		// immediately cancel (since that can take time), but only start new task if (1) delay time has elapsed and (2) active task count is acceptable
-		CancelActiveCompute();
-		bool bStartedTask = false;
 		if ((AccumTime - LastInvalidateTime) > CancelActiveOpDelaySeconds)
 		{
+			CancelActiveCompute();
+
 			int Active = *ActiveTaskCount;
 			if (Active < MaxActiveTaskCount)
 			{
 				bWaitingForActiveTasksToFinish = false;
 				StartNewCompute();
-				bStartedTask = true;
 			}
 			else
 			{
 				bWaitingForActiveTasksToFinish = true;
+				// failed to start new task, return to 'waiting to cancel' state
+				TaskState = EBackgroundComputeTaskState::WaitingToCancel;
 			}
 		}
-		if (!bStartedTask)
-		{
-			// failed to start new task, return to 'waiting to cancel' state
-			TaskState = EBackgroundComputeTaskState::WaitingToCancel;
-		}
 	}
 }
 
