MeshModelingTools: Fix some selection mechanic operations being unsafe because they assume a FTriangleGroupTopology.

diff --git a/Engine/Plugins/ChaosClothAssetEditor/Source/ChaosClothAssetEditorTools/Private/ChaosClothAsset/ClothMeshSelectionTool.cpp b/Engine/Plugins/ChaosClothAssetEditor/Source/ChaosClothAssetEditorTools/Private/ChaosClothAsset/ClothMeshSelectionTool.cpp
--- a/Engine/Plugins/ChaosClothAssetEditor/Source/ChaosClothAssetEditorTools/Private/ChaosClothAsset/ClothMeshSelectionTool.cpp
+++ b/Engine/Plugins/ChaosClothAssetEditor/Source/ChaosClothAssetEditorTools/Private/ChaosClothAsset/ClothMeshSelectionTool.cpp
@@ -613,13 +613,17 @@ void UClothMeshSelectionTool::RequestAction(EClothMeshSelectionToolActions Actio
 
 void UClothMeshSelectionTool::ApplyAction(EClothMeshSelectionToolActions ActionType)
 {
+	// We use a triangle topology, so the below actions can be done in mesh space instead
+	//  of working with the topology.
+	bool bAsTriangleTopology = true;
+
 	switch (ActionType)
 	{
 	case EClothMeshSelectionToolActions::GrowSelection:
-		SelectionMechanic->GrowSelection();
+		SelectionMechanic->GrowSelection(bAsTriangleTopology);
 		break;
 	case EClothMeshSelectionToolActions::ShrinkSelection:
-		SelectionMechanic->ShrinkSelection();
+		SelectionMechanic->ShrinkSelection(bAsTriangleTopology);
 		break;
 	case EClothMeshSelectionToolActions::FloodSelection:
 		SelectionMechanic->FloodSelection();
diff --git a/Engine/Plugins/Experimental/MeshModelingToolsetExp/Source/MeshModelingToolsEditorOnlyExp/Private/SkeletalMesh/SkinWeightsPaintTool.cpp b/Engine/Plugins/Experimental/MeshModelingToolsetExp/Source/MeshModelingToolsEditorOnlyExp/Private/SkeletalMesh/SkinWeightsPaintTool.cpp
--- a/Engine/Plugins/Experimental/MeshModelingToolsetExp/Source/MeshModelingToolsEditorOnlyExp/Private/SkeletalMesh/SkinWeightsPaintTool.cpp
+++ b/Engine/Plugins/Experimental/MeshModelingToolsetExp/Source/MeshModelingToolsEditorOnlyExp/Private/SkeletalMesh/SkinWeightsPaintTool.cpp
@@ -674,7 +674,7 @@ void UWeightToolMeshSelector::GrowSelection() const
 		return;
 	}
 	
-	PolygonSelectionMechanic->GrowSelection();
+	PolygonSelectionMechanic->GrowSelection(/*bAsTriangleTopology*/ true);
 }
 
 void UWeightToolMeshSelector::ShrinkSelection() const
@@ -684,7 +684,7 @@ void UWeightToolMeshSelector::ShrinkSelection() const
 		return;
 	}
 	
-	PolygonSelectionMechanic->ShrinkSelection();
+	PolygonSelectionMechanic->ShrinkSelection(/*bAsTriangleTopology*/ true);
 }
 
 void UWeightToolMeshSelector::FloodSelection() const
@@ -704,7 +704,7 @@ void UWeightToolMeshSelector::SelectBorder() const
 		return;
 	}
 	
-	PolygonSelectionMechanic->ConvertSelectionToBorder();
+	PolygonSelectionMechanic->ConvertSelectionToBorderVertices(/*bAsTriangleTopology*/ true);
 }
 
 void UWeightToolTransferManager::InitialSetup(USkinWeightsPaintTool* InWeightTool, FEditorViewportClient* InViewportClient)
diff --git a/Engine/Plugins/Runtime/MeshModelingToolset/Source/ModelingComponents/Private/Selection/MeshTopologySelectionMechanic.cpp b/Engine/Plugins/Runtime/MeshModelingToolset/Source/ModelingComponents/Private/Selection/MeshTopologySelectionMechanic.cpp
--- a/Engine/Plugins/Runtime/MeshModelingToolset/Source/ModelingComponents/Private/Selection/MeshTopologySelectionMechanic.cpp
+++ b/Engine/Plugins/Runtime/MeshModelingToolset/Source/ModelingComponents/Private/Selection/MeshTopologySelectionMechanic.cpp
@@ -533,292 +533,6 @@ void UMeshTopologySelectionMechanic::SelectAll()
 	ParentTool->GetToolManager()->EndUndoTransaction();
 }
 
-void UMeshTopologySelectionMechanic::GrowSelection()
-{
-	ParentTool->GetToolManager()->BeginUndoTransaction(LOCTEXT("GrowSelectionChange", "Grow Selection"));
-	BeginChange();
-
-	// add the neighbor vertices to the current selection
-	if (Properties->bSelectVertices)
-	{
-		TSet<int32> VerticesToAdd;
-		for (const int32 VertexIndex : PersistentSelection.SelectedCornerIDs)
-		{
-			for (const int32 NeighborIndex : Mesh->VtxVerticesItr(VertexIndex))
-			{
-				if (!PersistentSelection.SelectedCornerIDs.Contains(NeighborIndex))
-				{
-					VerticesToAdd.Add(NeighborIndex);
-				}
-			}	
-		}
-		
-		PersistentSelection.SelectedCornerIDs.Append(VerticesToAdd);
-	}
-
-	if (Properties->bSelectEdges)
-	{
-		// grow edge selection
-		TSet<int32> EdgesToAdd;
-		for (const int32 EdgeIndex : PersistentSelection.SelectedEdgeIDs)
-		{
-			FDynamicMesh3::FEdge CurrentEdge = Mesh->GetEdge(EdgeIndex);
-			for (int32 EdgeVertex = 0; EdgeVertex<2; ++EdgeVertex)
-			{
-				for (int32 EdgeID : Mesh->VtxEdgesItr(CurrentEdge.Vert[EdgeVertex]))
-				{
-					if (!PersistentSelection.SelectedEdgeIDs.Contains(EdgeID))
-					{
-						EdgesToAdd.Add(EdgeID);
-					}
-				}
-			}
-		}
-		
-		PersistentSelection.SelectedEdgeIDs.Append(EdgesToAdd);
-	}
-
-	if (Properties->bSelectFaces)
-	{
-		// grow face selection
-		TSet<int32> TrianglesToAdd;
-		for (const int32 TriIndex : PersistentSelection.SelectedGroupIDs)
-		{
-			FIndex3i NeighborTriangles = Mesh->GetTriNeighbourTris(TriIndex);
-			for (int32 TriSideIndex = 0; TriSideIndex < 3; ++TriSideIndex)
-			{
-				const int32 NeighborTriangleIndex = NeighborTriangles[TriSideIndex];
-				if (NeighborTriangleIndex == FDynamicMesh3::InvalidID)
-				{
-					continue;
-				}
-				
-				if (!PersistentSelection.SelectedGroupIDs.Contains(NeighborTriangleIndex))
-				{
-					TrianglesToAdd.Add(NeighborTriangleIndex);
-				}
-			}
-		}
-		
-		PersistentSelection.SelectedGroupIDs.Append(TrianglesToAdd);
-	}
-
-	SelectionTimestamp++;
-	OnSelectionChanged.Broadcast();
-	EndChangeAndEmitIfModified();
-	ParentTool->GetToolManager()->EndUndoTransaction();
-}
-
-void UMeshTopologySelectionMechanic::ShrinkSelection()
-{
-	ParentTool->GetToolManager()->BeginUndoTransaction(LOCTEXT("ShrinkSelectionChange", "Shrink Selection"));
-	BeginChange();
-
-	// remove border vertices from the current selection
-	if (Properties->bSelectVertices)
-	{
-		TSet<int32> BorderVertices;
-		for (const int32 VertexIndex : PersistentSelection.SelectedCornerIDs)
-		{
-			for (const int32 NeighborIndex : Mesh->VtxVerticesItr(VertexIndex))
-			{
-				if (!PersistentSelection.SelectedCornerIDs.Contains(NeighborIndex))
-				{
-					BorderVertices.Add(VertexIndex);
-					break;
-				}
-			}	
-		}
-		
-		for (const int32 BorderVertex : BorderVertices)
-		{
-			PersistentSelection.SelectedCornerIDs.Remove(BorderVertex);
-		}
-	}
-
-	if (Properties->bSelectEdges)
-	{
-		// shrink edge selection
-		TSet<int32> BorderEdges;
-		for (const int32 EdgeIndex : PersistentSelection.SelectedEdgeIDs)
-		{
-			FDynamicMesh3::FEdge CurrentEdge = Mesh->GetEdge(EdgeIndex);
-			for (int32 EdgeVertex = 0; EdgeVertex<2; ++EdgeVertex)
-			{
-				for (const int32 EdgeID : Mesh->VtxEdgesItr(CurrentEdge.Vert[EdgeVertex]))
-				{
-					if (!PersistentSelection.SelectedEdgeIDs.Contains(EdgeID))
-					{
-						BorderEdges.Add(EdgeIndex);
-					}
-				}
-			}
-		}
-		
-		for (const int32 BorderEdge : BorderEdges)
-		{
-			PersistentSelection.SelectedEdgeIDs.Remove(BorderEdge);
-		}
-	}
-
-	if (Properties->bSelectFaces)
-	{
-		// shrink face selection
-		TSet<int32> BorderTriangles;
-		for (const int32 TriIndex : PersistentSelection.SelectedGroupIDs)
-		{
-			FIndex3i NeighborTriangles = Mesh->GetTriNeighbourTris(TriIndex);
-			for (int32 TriSideIndex = 0; TriSideIndex < 3; ++TriSideIndex)
-			{
-				const int32 NeighborTriangleIndex = NeighborTriangles[TriSideIndex];
-				if (NeighborTriangleIndex == FDynamicMesh3::InvalidID)
-				{
-					continue;
-				}
-				
-				if (!PersistentSelection.SelectedGroupIDs.Contains(NeighborTriangleIndex))
-				{
-					BorderTriangles.Add(TriIndex);
-				}
-			}
-		}
-		
-		for (const int32 BorderTriangle : BorderTriangles)
-		{
-			PersistentSelection.SelectedGroupIDs.Remove(BorderTriangle);
-		}
-	}
-
-	SelectionTimestamp++;
-	OnSelectionChanged.Broadcast();
-	EndChangeAndEmitIfModified();
-	ParentTool->GetToolManager()->EndUndoTransaction();
-}
-
-void UMeshTopologySelectionMechanic::FloodSelection()
-{
-	ParentTool->GetToolManager()->BeginUndoTransaction(LOCTEXT("FloodSelectionChange", "Flood Selection"));
-	BeginChange();
-
-	// add all connected VERTICES to selection
-	if (Properties->bSelectVertices)
-	{
-		TSet<int32>& SelectedVertices = PersistentSelection.SelectedCornerIDs;
-		TSet<int32> VerticesAddedPrevIteration = PersistentSelection.SelectedCornerIDs;
-		while(!VerticesAddedPrevIteration.IsEmpty())
-		{
-			TSet<int32> VerticesToAddThisIteration;
-			for (const int32 VertexAdded : VerticesAddedPrevIteration)
-			{
-				for (const int32 NeighborIndex : Mesh->VtxVerticesItr(VertexAdded))
-				{
-					if (!SelectedVertices.Contains(NeighborIndex))
-					{
-						VerticesToAddThisIteration.Add(NeighborIndex);
-					}
-				}
-			}
-			
-			SelectedVertices.Append(VerticesToAddThisIteration);
-			VerticesAddedPrevIteration = VerticesToAddThisIteration;
-		}
-	}
-
-	// add all connected EDGES to selection
-	if (Properties->bSelectEdges)
-	{
-		TSet<int32>& SelectedEdges = PersistentSelection.SelectedEdgeIDs;
-		TSet<int32> EdgesAddedInPrevIteration = PersistentSelection.SelectedEdgeIDs;
-		while(!EdgesAddedInPrevIteration.IsEmpty())
-		{
-			TSet<int32> EdgesToAddThisIteration;
-			for (const int32 EdgeAdded : EdgesAddedInPrevIteration)
-			{
-				FDynamicMesh3::FEdge CurrentEdge = Mesh->GetEdge(EdgeAdded);
-				for (int32 EdgeVertex = 0; EdgeVertex<2; ++EdgeVertex)
-				{
-					for (const int32 NeighborEdgeID : Mesh->VtxEdgesItr(CurrentEdge.Vert[EdgeVertex]))
-					{
-						if (!SelectedEdges.Contains(NeighborEdgeID))
-                    	{
-                    		EdgesToAddThisIteration.Add(NeighborEdgeID);
-                    	}
-					}
-				}
-			}
-			
-			SelectedEdges.Append(EdgesToAddThisIteration);
-			EdgesAddedInPrevIteration = EdgesToAddThisIteration;
-		}
-	}
-
-	// add all connected FACES to selection
-	if (Properties->bSelectFaces)
-	{
-		TSet<int32>& SelectedFaces = PersistentSelection.SelectedGroupIDs;
-		TSet<int32> FacesAddedInPrevIteration = PersistentSelection.SelectedGroupIDs;
-		while(!FacesAddedInPrevIteration.IsEmpty())
-		{
-			TSet<int32> FacesToAddThisIteration;
-			for (const int32 FaceAdded : FacesAddedInPrevIteration)
-			{
-				FIndex3i NeighborTriangles = Mesh->GetTriNeighbourTris(FaceAdded);
-				for (int32 TriSideIndex = 0; TriSideIndex < 3; ++TriSideIndex)
-				{
-					const int32 NeighborTriangleIndex = NeighborTriangles[TriSideIndex];
-					if (NeighborTriangleIndex == FDynamicMesh3::InvalidID)
-					{
-						continue;
-					}
-				
-					if (!PersistentSelection.SelectedGroupIDs.Contains(NeighborTriangleIndex))
-					{
-						FacesToAddThisIteration.Add(NeighborTriangleIndex);
-					}
-				}
-			}
-			
-			SelectedFaces.Append(FacesToAddThisIteration);
-			FacesAddedInPrevIteration = FacesToAddThisIteration;
-		}
-	}
-
-	SelectionTimestamp++;
-	OnSelectionChanged.Broadcast();
-	EndChangeAndEmitIfModified();
-	ParentTool->GetToolManager()->EndUndoTransaction();
-}
-
-void UMeshTopologySelectionMechanic::ConvertSelectionToBorder()
-{
-	// find all border vertices
-	// a "border" vertex is one that has a 1-ring neighbor that is not in the selection set
-	const FGroupTopologySelection& CurrentSelection = GetActiveSelection();
-	FGroupTopologySelection Selection;
-	for (const int32 SelectedVertex : CurrentSelection.SelectedCornerIDs)
-	{
-		for (const int32 NeighborIndex : Mesh->VtxVerticesItr(SelectedVertex))
-		{
-			if (!CurrentSelection.SelectedCornerIDs.Contains(NeighborIndex))
-			{
-				Selection.SelectedCornerIDs.Add(SelectedVertex);
-			}
-		}
-	}
-
-	ParentTool->GetToolManager()->BeginUndoTransaction(LOCTEXT("BorderSelectionChange", "Select Border"));
-	BeginChange();
-	
-	// select vertices
-	PersistentSelection = Selection;
-
-	// broadcast and end transaction
-	SelectionTimestamp++;
-	OnSelectionChanged.Broadcast();
-	EndChangeAndEmitIfModified();
-	ParentTool->GetToolManager()->EndUndoTransaction();
-}
-
 FInputRayHit UMeshTopologySelectionMechanic::IsHitByClick(const FInputDeviceRay& ClickPos)
 {
 	if (!bIsEnabled)
diff --git a/Engine/Plugins/Runtime/MeshModelingToolset/Source/ModelingComponents/Private/Selection/PolygonSelectionMechanic.cpp b/Engine/Plugins/Runtime/MeshModelingToolset/Source/ModelingComponents/Private/Selection/PolygonSelectionMechanic.cpp
--- a/Engine/Plugins/Runtime/MeshModelingToolset/Source/ModelingComponents/Private/Selection/PolygonSelectionMechanic.cpp
+++ b/Engine/Plugins/Runtime/MeshModelingToolset/Source/ModelingComponents/Private/Selection/PolygonSelectionMechanic.cpp
@@ -2,9 +2,11 @@
 
 #include "Selection/PolygonSelectionMechanic.h"
 #include "Engine/World.h"
+#include "InteractiveToolManager.h"
 #include "Selection/GroupTopologySelector.h"
 #include "Selection/PersistentMeshSelection.h"
 #include "Selections/GeometrySelection.h"
+#include "Selections/GeometrySelectionUtil.h"
 
 #include UE_INLINE_GENERATED_CPP_BY_NAME(PolygonSelectionMechanic)
 
@@ -12,6 +14,24 @@ using namespace UE::Geometry;
 
 #define LOCTEXT_NAMESPACE "UPolygonSelectionMechanic"
 
+namespace PolygonSelectionMechanicLocals
+{
+	using namespace UE::Geometry;
+
+	EGeometryElementType ToGeometryElementType(const FGroupTopologySelection& Selection)
+	{
+		if (!Selection.SelectedCornerIDs.IsEmpty())
+		{
+			return EGeometryElementType::Vertex;
+		}
+		else if (!Selection.SelectedEdgeIDs.IsEmpty())
+		{
+			return EGeometryElementType::Edge;
+		}
+		return EGeometryElementType::Face;
+	}
+}
+
 void UPolygonSelectionMechanic::Initialize(
 	const FDynamicMesh3* MeshIn,
 	FTransform3d TargetTransformIn,
@@ -173,6 +193,8 @@ void UPolygonSelectionMechanic::SetSelection_AsGroupTopology(const UE::Geometry:
 	{
 		return;
 	}
+	PersistentSelection.Clear();
+
 	if (Selection.ElementType == EGeometryElementType::Vertex)
 	{
 		for (uint64 ElementID : Selection.Selection)
@@ -217,6 +239,8 @@ void UPolygonSelectionMechanic::SetSelection_AsTriangleTopology(const UE::Geomet
 	{
 		return;
 	}
+	PersistentSelection.Clear();
+	
 	if (Selection.ElementType == EGeometryElementType::Vertex)
 	{
 		for (uint64 ElementID : Selection.Selection)
@@ -255,6 +279,157 @@ void UPolygonSelectionMechanic::SetSelection_AsTriangleTopology(const UE::Geomet
 	}
 }
 
+bool UPolygonSelectionMechanic::ExecuteActionThroughGeometrySelection(bool bAsTriangleTopology, const FText& TransactionName, 
+	TFunctionRef<bool(UE::Geometry::FGeometrySelection& SelectionToModifyInPlace)> SelectionProcessor)
+{
+	using namespace PolygonSelectionMechanicLocals;
+	using namespace UE::Geometry;
+
+	if (!Topology || !Topology->GetMesh())
+	{
+		return false;
+	}
+
+	FGeometrySelection GeometrySelection;
+	GeometrySelection.InitializeTypes(ToGeometryElementType(PersistentSelection),
+		bAsTriangleTopology ? EGeometryTopologyType::Triangle : EGeometryTopologyType::Polygroup);
+
+	if (bAsTriangleTopology)
+	{
+		GetSelection_AsTriangleTopology(GeometrySelection);
+	}
+	else
+	{
+		GetSelection_AsGroupTopology(GeometrySelection);
+	}
+
+	bool bSuccess = SelectionProcessor(GeometrySelection);
+	if (!bSuccess)
+	{
+		return false;
+	}
+
+	ParentTool->GetToolManager()->BeginUndoTransaction(TransactionName);
+	BeginChange();
+
+	if (bAsTriangleTopology)
+	{
+		SetSelection_AsTriangleTopology(GeometrySelection);
+	}
+	else
+	{
+		SetSelection_AsGroupTopology(GeometrySelection);
+	}
+
+	SelectionTimestamp++;
+	OnSelectionChanged.Broadcast();
+	EndChangeAndEmitIfModified();
+	ParentTool->GetToolManager()->EndUndoTransaction();
+
+	return true;
+}
+
+void UPolygonSelectionMechanic::GrowSelection(bool bAsTriangleTopology)
+{
+	using namespace UE::Geometry;
+
+	ExecuteActionThroughGeometrySelection(bAsTriangleTopology, LOCTEXT("GrowSelectionChange", "Grow Selection"),
+		[this](FGeometrySelection& GeometrySelection) 
+	{
+		FGeometrySelection BoundaryConnectedSelection;
+		BoundaryConnectedSelection.InitializeTypes(GeometrySelection);
+
+		bool bSuccess = MakeBoundaryConnectedSelection(*Topology->GetMesh(), Topology, GeometrySelection,
+			[](FGeoSelectionID) { return true; }, BoundaryConnectedSelection);
+		bSuccess = bSuccess && CombineSelectionInPlace(GeometrySelection, BoundaryConnectedSelection, 
+			UE::Geometry::EGeometrySelectionCombineModes::Add);
+
+		return bSuccess;
+	});
+}
+
+void UPolygonSelectionMechanic::ShrinkSelection(bool bAsTriangleTopology)
+{
+	using namespace UE::Geometry;
+
+	ExecuteActionThroughGeometrySelection(bAsTriangleTopology, LOCTEXT("ShrinkSelectionChange", "Shrink Selection"),
+		[this](FGeometrySelection& GeometrySelection) 
+	{
+		FGeometrySelection BoundaryConnectedSelection;
+		BoundaryConnectedSelection.InitializeTypes(GeometrySelection);
+
+		bool bSuccess = MakeBoundaryConnectedSelection(*Topology->GetMesh(), Topology, GeometrySelection,
+			[](FGeoSelectionID) { return true; }, BoundaryConnectedSelection);
+		bSuccess = bSuccess && CombineSelectionInPlace(GeometrySelection, BoundaryConnectedSelection, 
+			UE::Geometry::EGeometrySelectionCombineModes::Subtract);
+
+		return bSuccess;
+	});
+}
+
+void UPolygonSelectionMechanic::ConvertSelectionToBorderVertices(bool bAsTriangleTopology)
+{
+	using namespace UE::Geometry;
+
+	ExecuteActionThroughGeometrySelection(bAsTriangleTopology, LOCTEXT("BorderSelectionChange", "Select Border"),
+		[this, bAsTriangleTopology](FGeometrySelection& GeometrySelection) 
+	{
+		TSet<int32> Unused;
+		if (bAsTriangleTopology)
+		{
+			TSet<int32> BoundaryVertices;
+			bool bSuccess = GetSelectionBoundaryVertices(*Topology->GetMesh(), Topology, GeometrySelection,
+				BoundaryVertices, Unused);
+			
+			if (!bSuccess) { return false; }
+
+			GeometrySelection.Selection.Reset();
+			GeometrySelection.InitializeTypes(EGeometryElementType::Vertex, EGeometryTopologyType::Triangle);
+			for (int32 Vid : BoundaryVertices)
+			{
+				GeometrySelection.Selection.Add(FGeoSelectionID::MeshVertex(Vid).Encoded());
+			}
+		}
+		else
+		{
+			TSet<int32> BoundaryCorners;
+			bool bSuccess = GetSelectionBoundaryCorners(*Topology->GetMesh(), Topology, GeometrySelection, 
+				BoundaryCorners, Unused);
+
+			if (!bSuccess) { return false; }
+
+			GeometrySelection.Selection.Reset();
+			GeometrySelection.InitializeTypes(EGeometryElementType::Vertex, EGeometryTopologyType::Polygroup);
+			for (int32 CornerID : BoundaryCorners)
+			{
+				GeometrySelection.Selection.Add(FGeoSelectionID(Topology->GetCornerVertexID(CornerID), CornerID).Encoded());
+			}
+		}
+
+		return true;
+	});
+}
+
+void UPolygonSelectionMechanic::FloodSelection()
+{
+	using namespace UE::Geometry;
+
+	ExecuteActionThroughGeometrySelection(true, LOCTEXT("FloodSelectionChange", "Flood Selection"),
+		[this](FGeometrySelection& GeometrySelection)
+	{
+		FGeometrySelection NewSelection;
+		NewSelection.InitializeTypes(GeometrySelection);
+		bool bSuccess = MakeSelectAllConnectedSelection(*Topology->GetMesh(), Topology, GeometrySelection, 
+			[](FGeoSelectionID) { return true; }, [](FGeoSelectionID, FGeoSelectionID) { return true; }, 
+			NewSelection);
+
+		if (!bSuccess) { return false; }
+			
+		GeometrySelection = NewSelection;
+		return true;
+	});
+}
+
 
 bool UPolygonSelectionMechanic::UpdateHighlight(const FRay& WorldRay)
 {
diff --git a/Engine/Plugins/Runtime/MeshModelingToolset/Source/ModelingComponents/Public/Selection/MeshTopologySelectionMechanic.h b/Engine/Plugins/Runtime/MeshModelingToolset/Source/ModelingComponents/Public/Selection/MeshTopologySelectionMechanic.h
--- a/Engine/Plugins/Runtime/MeshModelingToolset/Source/ModelingComponents/Public/Selection/MeshTopologySelectionMechanic.h
+++ b/Engine/Plugins/Runtime/MeshModelingToolset/Source/ModelingComponents/Public/Selection/MeshTopologySelectionMechanic.h
@@ -284,10 +284,14 @@ public:
 	void InvertSelection();
 	void SelectAll();
 
-	void GrowSelection();
-	void ShrinkSelection();
-	void FloodSelection();
-	void ConvertSelectionToBorder();
+	UE_DEPRECATED(5.5, "Calling this function on the base class is not supported. It can be called on UPolygonSelectionMechanic.")
+	void GrowSelection() {};
+	UE_DEPRECATED(5.5, "Calling this function on the base class is not supported. It can be called on UPolygonSelectionMechanic.")
+	void ShrinkSelection() {};
+	UE_DEPRECATED(5.5, "Calling this function on the base class is not supported. It can be called on UPolygonSelectionMechanic.")
+	void FloodSelection() {};
+	UE_DEPRECATED(5.5, "Calling this function on the base class is not supported. It can be called on UPolygonSelectionMechanic.")
+	void ConvertSelectionToBorder() {};
 
 	/** 
 	 * @return true if the current selection is non-empty 
diff --git a/Engine/Plugins/Runtime/MeshModelingToolset/Source/ModelingComponents/Public/Selection/PolygonSelectionMechanic.h b/Engine/Plugins/Runtime/MeshModelingToolset/Source/ModelingComponents/Public/Selection/PolygonSelectionMechanic.h
--- a/Engine/Plugins/Runtime/MeshModelingToolset/Source/ModelingComponents/Public/Selection/PolygonSelectionMechanic.h
+++ b/Engine/Plugins/Runtime/MeshModelingToolset/Source/ModelingComponents/Public/Selection/PolygonSelectionMechanic.h
@@ -44,6 +44,34 @@ public:
 		TFunction<FDynamicMeshAABBTree3* ()> GetSpatialSourceFunc
 	);
 
+	/*
+	 * Expands selection at the borders.
+	 * 
+	 * @param bAsTriangleTopology Can be set true if the topology type is FTriangleGroupTopology, to
+	 *  perform the operation a bit more efficiently by using the mesh topology directly.
+	 */
+	void GrowSelection(bool bAsTriangleTopology);
+	/*
+	 * Shrinks selection at the borders.
+	 * 
+	 * @param bAsTriangleTopology Can be set true if the topology type is FTriangleGroupTopology, to
+	 *  perform the operation a bit more efficiently by using the mesh topology directly.
+	 */
+	void ShrinkSelection(bool bAsTriangleTopology);
+
+	/**
+	 * Converts selection to a vertex/corner selection of just the boundary vertices/corners
+	 * 
+	 * @param bAsTriangleTopology Can be set true if the topology type is FTriangleGroupTopology, to
+	 *  perform the operation a bit more efficiently by using the mesh topology directly.
+	 */
+	void ConvertSelectionToBorderVertices(bool bAsTriangleTopology);
+	
+	/**
+	 * Expands selection to encompass connected components.
+	 */
+	void FloodSelection();
+
 	// UMeshTopologySelectionMechanic
 	virtual bool UpdateHighlight(const FRay& WorldRay) override;
 	virtual bool UpdateSelection(const FRay& WorldRay, FVector3d& LocalHitPositionOut, FVector3d& LocalHitNormalOut) override;
@@ -75,4 +103,7 @@ private:
 	// TODO: Would be nice to get rid of this and write everything in terms of TopologySelector and TopologyProvider
 	const FGroupTopology* Topology;
 
+	// Helper to execute selection actions through existing geometry selection code.
+	bool ExecuteActionThroughGeometrySelection(bool bAsTriangleTopology, const FText& TransactionName,
+		TFunctionRef<bool(UE::Geometry::FGeometrySelection& SelectionToModifyInPlace)> SelectionProcessor);
 };
