Modeling Tools: Changes how Mesh Element Selection Colors are updated during customization for a faster update on dense meshes

diff --git a/Engine/Plugins/Runtime/MeshModelingToolset/Source/ModelingComponents/Private/Selection/GeometrySelectionManager.cpp b/Engine/Plugins/Runtime/MeshModelingToolset/Source/ModelingComponents/Private/Selection/GeometrySelectionManager.cpp
--- a/Engine/Plugins/Runtime/MeshModelingToolset/Source/ModelingComponents/Private/Selection/GeometrySelectionManager.cpp
+++ b/Engine/Plugins/Runtime/MeshModelingToolset/Source/ModelingComponents/Private/Selection/GeometrySelectionManager.cpp
@@ -1810,21 +1810,53 @@ void UGeometrySelectionManager::SetSelectionColors(const FLinearColor Unselected
 	HoverOverUnselectedParams.Color = HoverOverUnselectedCol.ToFColor(true);
 	SelectedParams.Color = GeometrySelectedCol.ToFColor(true);
 
-	UnselectedParams.SelectionFillColor =
-		ToolSetupUtil::GetCustomTwoSidedDepthOffsetMaterial(this->ToolsContext->ToolManager, UnselectedCol, UnselectedParams.DepthBias, UnselectedCol.A);
+	// on initial set up of Materials used for selection colors material (typically when entering modeling mode)
+	if (UnselectedParams.SelectionFillColor == nullptr) // if one is null, they all will be
+	{
+		auto SetMaterial = [this] (FMeshElementSelectionParams& Params, const FLinearColor Color)
+		{
+			Params.SelectionFillColor =
+				ToolSetupUtil::GetCustomTwoSidedDepthOffsetMaterial(this->ToolsContext->ToolManager, Color, Params.DepthBias, Color.A);
+		};
 
-	HoverOverSelectedParams.SelectionFillColor =
-		ToolSetupUtil::GetCustomTwoSidedDepthOffsetMaterial(this->ToolsContext->ToolManager, HoverOverSelectedCol, HoverOverSelectedParams.DepthBias);
+		SetMaterial(UnselectedParams, UnselectedCol);
+		SetMaterial(HoverOverUnselectedParams, HoverOverUnselectedCol);
+		SetMaterial(SelectedParams, GeometrySelectedCol);
 
-	HoverOverUnselectedParams.SelectionFillColor =
-		ToolSetupUtil::GetCustomTwoSidedDepthOffsetMaterial(this->ToolsContext->ToolManager, HoverOverUnselectedCol, HoverOverUnselectedParams.DepthBias, HoverOverUnselectedCol.A);
+		// to avoid flickering, this version of GetCustomTwoSidedDeptOffsetMaterial (without opacity parameter) must be called for HoverOverSelected
+		HoverOverSelectedParams.SelectionFillColor =
+			ToolSetupUtil::GetCustomTwoSidedDepthOffsetMaterial(this->ToolsContext->ToolManager, HoverOverSelectedCol, HoverOverSelectedParams.DepthBias);
+	}
+	// setting colors after initialization of Materials (typically using color customization in editor preferences)
+	else
+	{
+		auto SetColorAndOpacity = [](const FMeshElementSelectionParams& Params, const FLinearColor Color)
+		{
+			Params.SelectionFillColor->SetScalarParameterValue("Opacity", Color.A); // no effect for HoverOverSelected
+			Params.SelectionFillColor->SetVectorParameterValue("Color", Color);
+		};
 
-	SelectedParams.SelectionFillColor =
-		ToolSetupUtil::GetCustomTwoSidedDepthOffsetMaterial(this->ToolsContext->ToolManager, GeometrySelectedCol, SelectedParams.DepthBias, GeometrySelectedCol.A);
+		SetColorAndOpacity(UnselectedParams, UnselectedCol);
+		SetColorAndOpacity(HoverOverSelectedParams, HoverOverSelectedCol);
+		SetColorAndOpacity(HoverOverUnselectedParams, HoverOverUnselectedCol);
+		SetColorAndOpacity(SelectedParams, GeometrySelectedCol);
+	}
 
 	// ensures that when color is changed in Editor Preferences, colors are immediately updated in the UI
-	MarkRenderCachesDirty();
-	RebuildSelectionRenderCaches();
+	auto UpdateAllSetsColor = [this](const FMeshElementSelectionParams& Params)
+	{
+		UPointSetComponent* PointSet = PreviewGeometry->FindPointSet(Params.Identifiers[0]);
+		ULineSetComponent* LineSet = PreviewGeometry->FindLineSet(Params.Identifiers[1]);
+		UTriangleSetComponent* TriSet = PreviewGeometry->FindTriangleSet(Params.Identifiers[2]);
+		if (PointSet) {	PointSet->SetAllPointsColor(Params.Color); }
+		if (LineSet) { LineSet->SetAllLinesColor(Params.Color); }
+		if (TriSet) { TriSet->SetAllTrianglesColor(Params.Color); }
+	};
+
+	UpdateAllSetsColor(UnselectedParams);
+	UpdateAllSetsColor(HoverOverSelectedParams);
+	UpdateAllSetsColor(HoverOverUnselectedParams);
+	UpdateAllSetsColor(SelectedParams);
 }
 
 
