Modeling Tools: Avoid MeshElementSelection PreviewGeo crash on editor exit

diff --git a/Engine/Plugins/Editor/ModelingToolsEditorMode/Source/ModelingToolsEditorMode/Private/ModelingToolsEditorMode.cpp b/Engine/Plugins/Editor/ModelingToolsEditorMode/Source/ModelingToolsEditorMode/Private/ModelingToolsEditorMode.cpp
--- a/Engine/Plugins/Editor/ModelingToolsEditorMode/Source/ModelingToolsEditorMode/Private/ModelingToolsEditorMode.cpp
+++ b/Engine/Plugins/Editor/ModelingToolsEditorMode/Source/ModelingToolsEditorMode/Private/ModelingToolsEditorMode.cpp
@@ -1206,6 +1206,7 @@ void UModelingToolsEditorMode::OnEditorClosed()
 
 	if (SelectionManager != nullptr)
 	{
+		SelectionManager->DisconnectPreviewGeometry();
 		SelectionManager->ClearSelection();
 		SelectionManager->ClearActiveTargets();
 	}
diff --git a/Engine/Plugins/Runtime/MeshModelingToolset/Source/ModelingComponents/Private/Selection/GeometrySelectionManager.cpp b/Engine/Plugins/Runtime/MeshModelingToolset/Source/ModelingComponents/Private/Selection/GeometrySelectionManager.cpp
--- a/Engine/Plugins/Runtime/MeshModelingToolset/Source/ModelingComponents/Private/Selection/GeometrySelectionManager.cpp
+++ b/Engine/Plugins/Runtime/MeshModelingToolset/Source/ModelingComponents/Private/Selection/GeometrySelectionManager.cpp
@@ -52,12 +52,7 @@ void UGeometrySelectionManager::Shutdown()
 	ToolsContext = nullptr;
 	TransactionsAPI = nullptr;
 
-	if (PreviewGeometry)
-	{
-		RemoveAllSets();
-		PreviewGeometry->Disconnect();
-		PreviewGeometry = nullptr;
-	}
+	DisconnectPreviewGeometry();
 
 	for (TSharedPtr<FGeometrySelectionTarget> Target : ActiveTargetReferences)
 	{
@@ -1859,6 +1854,15 @@ void UGeometrySelectionManager::SetSelectionColors(const FLinearColor Unselected
 	UpdateAllSetsColor(SelectedParams);
 }
 
+void UGeometrySelectionManager::DisconnectPreviewGeometry()
+{
+	if (PreviewGeometry)
+	{
+		RemoveAllSets();
+		PreviewGeometry->Disconnect();
+		PreviewGeometry = nullptr;
+	}
+}
 
 
 #undef LOCTEXT_NAMESPACE
diff --git a/Engine/Plugins/Runtime/MeshModelingToolset/Source/ModelingComponents/Public/Selection/GeometrySelectionManager.h b/Engine/Plugins/Runtime/MeshModelingToolset/Source/ModelingComponents/Public/Selection/GeometrySelectionManager.h
--- a/Engine/Plugins/Runtime/MeshModelingToolset/Source/ModelingComponents/Public/Selection/GeometrySelectionManager.h
+++ b/Engine/Plugins/Runtime/MeshModelingToolset/Source/ModelingComponents/Public/Selection/GeometrySelectionManager.h
@@ -437,6 +437,9 @@ public:
 	 * Unselected elements, Hover over selection, Hover over non-selection, and Selected elements
 	 */
 	void SetSelectionColors(FLinearColor UnselectedCol, FLinearColor HoverOverSelectedCol, FLinearColor HoverOverUnselectedCol, FLinearColor GeometrySelectedCol);
+
+	/** Disconnect and cleanup for PreviewGeometry object. */
+	void DisconnectPreviewGeometry();
 	
 protected:
 
