expose point set component to bp (equivalently to line set component)

diff --git a/Engine/Plugins/Runtime/MeshModelingToolset/Source/ModelingComponents/Private/Drawing/PointSetComponent.cpp b/Engine/Plugins/Runtime/MeshModelingToolset/Source/ModelingComponents/Private/Drawing/PointSetComponent.cpp
--- a/Engine/Plugins/Runtime/MeshModelingToolset/Source/ModelingComponents/Private/Drawing/PointSetComponent.cpp
+++ b/Engine/Plugins/Runtime/MeshModelingToolset/Source/ModelingComponents/Private/Drawing/PointSetComponent.cpp
@@ -279,6 +279,27 @@ int32 UPointSetComponent::AddPointInternal(const FRenderablePoint& Point)
 	return ID;
 }
 
+int32 UPointSetComponent::AddPoints(
+	const TArray<FVector>& Positions,
+	const FColor& InColor,
+	const float InSize,
+	const float InDepthBias
+)
+{
+	const int NumPoints = Positions.Num();
+	if (NumPoints == 0)
+	{
+		return 0;
+	}
+	MarkRenderStateDirty();
+	bBoundsDirty = true;
+	for (FVector P : Positions)
+	{
+		Points.Add(FRenderablePoint(P, InColor, InSize, InDepthBias));
+	}
+	return NumPoints;
+}
+
 
 void UPointSetComponent::InsertPoint(const int32 ID, const FRenderablePoint& OverlayPoint)
 {
diff --git a/Engine/Plugins/Runtime/MeshModelingToolset/Source/ModelingComponents/Public/Drawing/PointSetComponent.h b/Engine/Plugins/Runtime/MeshModelingToolset/Source/ModelingComponents/Public/Drawing/PointSetComponent.h
--- a/Engine/Plugins/Runtime/MeshModelingToolset/Source/ModelingComponents/Public/Drawing/PointSetComponent.h
+++ b/Engine/Plugins/Runtime/MeshModelingToolset/Source/ModelingComponents/Public/Drawing/PointSetComponent.h
@@ -44,7 +44,7 @@ struct FRenderablePoint
  * The points are drawn as two triangles (ie a square) orthogonal to the view direction. 
  * The actual point size is calculated in the shader, and so a custom material must be used.
  */
-UCLASS()
+UCLASS(meta = (BlueprintSpawnableComponent))
 class MODELINGCOMPONENTS_API UPointSetComponent : public UMeshComponent
 {
 	GENERATED_BODY()
@@ -53,9 +53,11 @@ public:
 	UPointSetComponent();
 
 	/** Specify material which handles points */
+	UFUNCTION(BlueprintCallable, Category = "Point Set")
 	void SetPointMaterial(UMaterialInterface* InPointMaterial);
 
 	/** Clear all primitives */
+	UFUNCTION(BlueprintCallable, Category = "Point Set")
 	void Clear();
 
 	/** Reserve enough memory for up to the given ID */
@@ -71,6 +73,18 @@ public:
 		return AddPoint(FRenderablePoint(InPosition, InColor, InSize, InDepthBias));
 	}
 
+	/** 
+	 * Add points to be rendered using the component. 
+	 * @return Number of added points
+	 */
+	UFUNCTION(BlueprintCallable, Category = "Point Set")
+	int32 AddPoints(
+		const TArray<FVector>& Positions,
+		const FColor& InColor,
+		const float InSize,
+		const float InDepthBias = 0.0f
+	);
+
 	/** Insert a point with the given ID into the set. */
 	void InsertPoint(const int32 ID, const FRenderablePoint& OverlayPoint);
 
