Modeling Tools: Fixes bug that would not deselect Edge in MeshElementSeletion on Undo/Redo and caused crash

diff --git a/Engine/Plugins/Runtime/MeshModelingToolset/Source/ModelingComponents/Private/Selection/DynamicMeshSelector.cpp b/Engine/Plugins/Runtime/MeshModelingToolset/Source/ModelingComponents/Private/Selection/DynamicMeshSelector.cpp
--- a/Engine/Plugins/Runtime/MeshModelingToolset/Source/ModelingComponents/Private/Selection/DynamicMeshSelector.cpp
+++ b/Engine/Plugins/Runtime/MeshModelingToolset/Source/ModelingComponents/Private/Selection/DynamicMeshSelector.cpp
@@ -296,28 +296,24 @@ void FBaseDynamicMeshSelector::UpdateSelectionViaRaycast_MeshTopology(
 	// therefore need to include both of the TriEdgeIDs which the edge belongs to to the SelectionEditor
 	if (SelectionEditor.GetElementType() == EGeometryElementType::Edge)
 	{
-		auto SelectionFunc = [&SelectionEditor](const uint32 EncodedID)
+		auto AffectBothTriEdgeIDs = [this, &SelectionEditor, &ResultOut](const EGeometrySelectionChangeType ChangeType)
 		{
-			SelectionEditor.Select(EncodedID);
-		};
-		auto DeselectionFunc = [&SelectionEditor](const uint32 EncodedID)
-		{
-			SelectionEditor.RemoveFromSelection(EncodedID);
-		};
+			ensure(ChangeType == EGeometrySelectionChangeType::Add || ChangeType == EGeometrySelectionChangeType::Remove);
 
-		auto AffectBothTriEdgeIDs = [this, &SelectionEditor](TArray<uint64> DeltaElements, TFunctionRef<void(const uint32)> SelectionOrDeselectionFunction)
-		{
+			const TArray<uint64>& DeltaElements = (ChangeType == EGeometrySelectionChangeType::Add ? ResultOut.SelectionDelta.Added : ResultOut.SelectionDelta.Removed);
+			
+			TArray<uint64> ElementsToUpdateInSelection;
 			for (const uint64 Element : DeltaElements)
 			{
 				FMeshTriEdgeID TriEdgeID(FGeoSelectionID(Element).GeometryID);
-				TargetMesh->ProcessMesh([TriEdgeID, &SelectionEditor, &SelectionOrDeselectionFunction](const UE::Geometry::FDynamicMesh3& SourceMesh)
+				TargetMesh->ProcessMesh([TriEdgeID, &ElementsToUpdateInSelection](const UE::Geometry::FDynamicMesh3& SourceMesh)
 				{
 					// the added or removed EdgeID
 					const int32 EdgeID = SourceMesh.IsTriangle(TriEdgeID.TriangleID) ? SourceMesh.GetTriEdge(TriEdgeID.TriangleID, TriEdgeID.TriEdgeIndex) : IndexConstants::InvalidID;
 					if (SourceMesh.IsEdge(EdgeID))
 					{
 						SourceMesh.EnumerateTriEdgeIDsFromEdgeID(EdgeID,
-							[&SelectionEditor, &TriEdgeID, &SelectionOrDeselectionFunction](const FMeshTriEdgeID OtherTriEdgeID)
+							[TriEdgeID, &ElementsToUpdateInSelection](const FMeshTriEdgeID OtherTriEdgeID)
 							{
 								// avoid adding to the selection the edge which already exists in selection
 								// OR removing from the selection the edge which has already been removed
@@ -325,16 +321,18 @@ void FBaseDynamicMeshSelector::UpdateSelectionViaRaycast_MeshTopology(
 								{
 									return;
 								}
-								SelectionOrDeselectionFunction(OtherTriEdgeID.Encoded());
+								ElementsToUpdateInSelection.Add(OtherTriEdgeID.Encoded());
 							});
 					}
 				});
 			}
+			// adds/removes from SelectionEditor and adds to ResultOut.SelectionDelta.Added/Removed
+			UpdateSelectionWithNewElements(&SelectionEditor, ChangeType, ElementsToUpdateInSelection, &ResultOut.SelectionDelta);
 		};
 
 		// selecting or deselecting (whichever applicable) the secondary TriEdgeID
-		AffectBothTriEdgeIDs(ResultOut.SelectionDelta.Added, SelectionFunc);
-		AffectBothTriEdgeIDs(ResultOut.SelectionDelta.Removed, DeselectionFunc);
+		AffectBothTriEdgeIDs(EGeometrySelectionChangeType::Add);
+		AffectBothTriEdgeIDs(EGeometrySelectionChangeType::Remove);
 	}
 }
 
diff --git a/Engine/Plugins/Runtime/MeshModelingToolset/Source/ModelingComponents/Private/Selection/GeometrySelectionManager.cpp b/Engine/Plugins/Runtime/MeshModelingToolset/Source/ModelingComponents/Private/Selection/GeometrySelectionManager.cpp
--- a/Engine/Plugins/Runtime/MeshModelingToolset/Source/ModelingComponents/Private/Selection/GeometrySelectionManager.cpp
+++ b/Engine/Plugins/Runtime/MeshModelingToolset/Source/ModelingComponents/Private/Selection/GeometrySelectionManager.cpp
@@ -686,7 +686,6 @@ TArray<FGeometryIdentifier> UGeometrySelectionManager::GetCurrentTargetIdentifie
 
 void UGeometrySelectionManager::SetTargetsOnUndoRedo(TArray<FGeometryIdentifier> NewTargets)
 {
-	check(HasSelection() == false);
 	ClearActiveTargets();
 	for (FGeometryIdentifier Identifier : NewTargets)
 	{
