Make tool target factories generate writable or non-writable targets depending on the source

diff --git a/Engine/Plugins/Runtime/MeshModelingToolset/Source/ModelingComponentsEditorOnly/Private/ToolTargets/SkeletalMeshComponentToolTarget.cpp b/Engine/Plugins/Runtime/MeshModelingToolset/Source/ModelingComponentsEditorOnly/Private/ToolTargets/SkeletalMeshComponentToolTarget.cpp
--- a/Engine/Plugins/Runtime/MeshModelingToolset/Source/ModelingComponentsEditorOnly/Private/ToolTargets/SkeletalMeshComponentToolTarget.cpp
+++ b/Engine/Plugins/Runtime/MeshModelingToolset/Source/ModelingComponentsEditorOnly/Private/ToolTargets/SkeletalMeshComponentToolTarget.cpp
@@ -232,20 +232,52 @@ bool USkeletalMeshComponentToolTargetFactory::CanBuildTarget(UObject* SourceObje
 	// If you want to make the tool target work with some subclass of USkeletalMeshComponent,
 	// just add another factory that allows that class specifically(but make sure that
 	// GetMeshDescription and such work properly)
-
-	return Cast<USkeletalMeshComponent>(SourceObject) && Cast<USkeletalMeshComponent>(SourceObject)->GetSkeletalMeshAsset() &&
+	
+	bool bValid = Cast<USkeletalMeshComponent>(SourceObject) && Cast<USkeletalMeshComponent>(SourceObject)->GetSkeletalMeshAsset() &&
 		ExactCast<USkeletalMesh>(Cast<USkeletalMeshComponent>(SourceObject)->GetSkeletalMeshAsset()) &&
-		!ExactCast<USkeletalMesh>(Cast<USkeletalMeshComponent>(SourceObject)->GetSkeletalMeshAsset())->GetOutermost()->bIsCookedForEditor &&
-		Requirements.AreSatisfiedBy(USkeletalMeshComponentToolTarget::StaticClass());
+		!ExactCast<USkeletalMesh>(Cast<USkeletalMeshComponent>(SourceObject)->GetSkeletalMeshAsset())->GetOutermost()->bIsCookedForEditor;
+	if (!bValid)
+	{
+		return false;
+	}
+
+	if (CanWriteToSource(SourceObject))
+	{
+		return Requirements.AreSatisfiedBy(USkeletalMeshComponentToolTarget::StaticClass());
+	}
+	else
+	{
+		return Requirements.AreSatisfiedBy(USkeletalMeshComponentReadOnlyToolTarget::StaticClass());
+	}
 }
 
 UToolTarget* USkeletalMeshComponentToolTargetFactory::BuildTarget(UObject* SourceObject, const FToolTargetTypeRequirements& Requirements)
 {
-	USkeletalMeshComponentToolTarget* Target = NewObject<USkeletalMeshComponentToolTarget>();
+	USkeletalMeshComponentReadOnlyToolTarget* Target;
+	if (CanWriteToSource(SourceObject))
+	{
+		Target = NewObject<USkeletalMeshComponentToolTarget>();
+	}
+	else
+	{
+		Target = NewObject<USkeletalMeshComponentReadOnlyToolTarget>();
+	}	
 	Target->InitializeComponent(Cast<USkeletalMeshComponent>(SourceObject));
 	checkSlow(Target->Component.IsValid() && Requirements.AreSatisfiedBy(Target));
 
 	return Target;
 }
 
+bool USkeletalMeshComponentToolTargetFactory::CanWriteToSource(const UObject* SourceObject)
+{
+	if (const USkeletalMeshComponent* Component = GetValid(Cast<USkeletalMeshComponent>(SourceObject)))
+	{
+		if (const USkeletalMesh* SkeletalMesh = GetValid(Cast<USkeletalMesh>(Component->GetSkeletalMeshAsset())))
+		{
+			return !SkeletalMesh->GetPathName().StartsWith(TEXT("/Engine/"));
+		}
+	}
+	return false;
+}
+
 #undef LOCTEXT_NAMESPACE
diff --git a/Engine/Plugins/Runtime/MeshModelingToolset/Source/ModelingComponentsEditorOnly/Private/ToolTargets/SkeletalMeshToolTarget.cpp b/Engine/Plugins/Runtime/MeshModelingToolset/Source/ModelingComponentsEditorOnly/Private/ToolTargets/SkeletalMeshToolTarget.cpp
--- a/Engine/Plugins/Runtime/MeshModelingToolset/Source/ModelingComponentsEditorOnly/Private/ToolTargets/SkeletalMeshToolTarget.cpp
+++ b/Engine/Plugins/Runtime/MeshModelingToolset/Source/ModelingComponentsEditorOnly/Private/ToolTargets/SkeletalMeshToolTarget.cpp
@@ -324,20 +324,49 @@ bool USkeletalMeshToolTargetFactory::CanBuildTarget(UObject* SourceObject, const
 	// GetMeshDescription and such work properly)
 	
 	const USkeletalMesh* SkeletalMesh =  ExactCast<USkeletalMesh>(SourceObject);
-
-	return SkeletalMesh
-		&& !SkeletalMesh->GetOutermost()->bIsCookedForEditor
-		&& Requirements.AreSatisfiedBy(USkeletalMeshToolTarget::StaticClass());
+	
+	bool bValid = SkeletalMesh
+		&& !SkeletalMesh->GetOutermost()->bIsCookedForEditor;
+	if (!bValid)
+	{
+		return false;
+	}
+	
+	if (CanWriteToSource(SourceObject))
+	{
+		return Requirements.AreSatisfiedBy(USkeletalMeshToolTarget::StaticClass());
+	}
+	else
+	{
+		return Requirements.AreSatisfiedBy(USkeletalMeshReadOnlyToolTarget::StaticClass());
+	}
 }
 
 UToolTarget* USkeletalMeshToolTargetFactory::BuildTarget(UObject* SourceObject, const FToolTargetTypeRequirements& Requirements)
 {
-	USkeletalMeshToolTarget* Target = NewObject<USkeletalMeshToolTarget>();
+	USkeletalMeshReadOnlyToolTarget* Target;
+	if (CanWriteToSource(SourceObject))
+	{
+		Target = NewObject<USkeletalMeshToolTarget>();
+	}
+	else
+	{
+		Target = NewObject<USkeletalMeshReadOnlyToolTarget>();
+	}
 	Target->SkeletalMesh = Cast<USkeletalMesh>(SourceObject);
 	checkSlow(Target->SkeletalMesh.IsValid() && Requirements.AreSatisfiedBy(Target));
 
 	return Target;
 }
 
+bool USkeletalMeshToolTargetFactory::CanWriteToSource(const UObject* SourceObject)
+{
+	if (const USkeletalMesh* SkeletalMesh = GetValid(ExactCast<USkeletalMesh>(SourceObject)))
+	{
+		return !SkeletalMesh->GetPathName().StartsWith(TEXT("/Engine/"));
+	}
+	return false;
+}
+
 #undef LOCTEXT_NAMESPACE
 
diff --git a/Engine/Plugins/Runtime/MeshModelingToolset/Source/ModelingComponentsEditorOnly/Private/ToolTargets/StaticMeshComponentToolTarget.cpp b/Engine/Plugins/Runtime/MeshModelingToolset/Source/ModelingComponentsEditorOnly/Private/ToolTargets/StaticMeshComponentToolTarget.cpp
--- a/Engine/Plugins/Runtime/MeshModelingToolset/Source/ModelingComponentsEditorOnly/Private/ToolTargets/StaticMeshComponentToolTarget.cpp
+++ b/Engine/Plugins/Runtime/MeshModelingToolset/Source/ModelingComponentsEditorOnly/Private/ToolTargets/StaticMeshComponentToolTarget.cpp
@@ -19,7 +19,7 @@
 
 using namespace UE::Geometry;
 
-void UStaticMeshComponentToolTarget::SetEditingLOD(EMeshLODIdentifier RequestedEditingLOD)
+void UStaticMeshComponentReadOnlyToolTarget::SetEditingLOD(EMeshLODIdentifier RequestedEditingLOD)
 {
 	EMeshLODIdentifier ValidEditingLOD = EMeshLODIdentifier::LOD0;
 
@@ -27,14 +27,14 @@ void UStaticMeshComponentToolTarget::SetEditingLOD(EMeshLODIdentifier RequestedE
 	if (ensure(StaticMeshComponent != nullptr))
 	{
 		UStaticMesh* StaticMeshAsset = StaticMeshComponent->GetStaticMesh();
-		ValidEditingLOD = UStaticMeshToolTarget::GetValidEditingLOD(StaticMeshAsset, RequestedEditingLOD);
+		ValidEditingLOD = UStaticMeshReadOnlyToolTarget::GetValidEditingLOD(StaticMeshAsset, RequestedEditingLOD);
 	}
 
 	EditingLOD = ValidEditingLOD;
 }
 
 
-bool UStaticMeshComponentToolTarget::IsValid() const
+bool UStaticMeshComponentReadOnlyToolTarget::IsValid() const
 {
 	if (!UPrimitiveComponentToolTarget::IsValid())
 	{
@@ -47,29 +47,29 @@ bool UStaticMeshComponentToolTarget::IsValid() const
 	}
 	UStaticMesh* StaticMesh = StaticMeshComponent->GetStaticMesh();
 	
-	return UStaticMeshToolTarget::HasNonGeneratedLOD(StaticMesh, EditingLOD);
+	return UStaticMeshReadOnlyToolTarget::HasNonGeneratedLOD(StaticMesh, EditingLOD);
 }
 
 
 
-int32 UStaticMeshComponentToolTarget::GetNumMaterials() const
+int32 UStaticMeshComponentReadOnlyToolTarget::GetNumMaterials() const
 {
 	return ensure(IsValid()) ? Component->GetNumMaterials() : 0;
 }
 
-UMaterialInterface* UStaticMeshComponentToolTarget::GetMaterial(int32 MaterialIndex) const
+UMaterialInterface* UStaticMeshComponentReadOnlyToolTarget::GetMaterial(int32 MaterialIndex) const
 {
 	return ensure(IsValid()) ? Component->GetMaterial(MaterialIndex) : nullptr;
 }
 
-void UStaticMeshComponentToolTarget::GetMaterialSet(FComponentMaterialSet& MaterialSetOut, bool bPreferAssetMaterials) const
+void UStaticMeshComponentReadOnlyToolTarget::GetMaterialSet(FComponentMaterialSet& MaterialSetOut, bool bPreferAssetMaterials) const
 {
 	if (!ensure(IsValid())) return;
 
 	if (bPreferAssetMaterials)
 	{
 		UStaticMesh* StaticMesh = Cast<UStaticMeshComponent>(Component)->GetStaticMesh();
-		UStaticMeshToolTarget::GetMaterialSet(StaticMesh, MaterialSetOut, bPreferAssetMaterials);
+		UStaticMeshReadOnlyToolTarget::GetMaterialSet(StaticMesh, MaterialSetOut, bPreferAssetMaterials);
 	}
 	else
 	{
@@ -82,7 +82,7 @@ void UStaticMeshComponentToolTarget::GetMaterialSet(FComponentMaterialSet& Mater
 	}
 }
 
-bool UStaticMeshComponentToolTarget::CommitMaterialSetUpdate(const FComponentMaterialSet& MaterialSet, bool bApplyToAsset)
+bool UStaticMeshComponentReadOnlyToolTarget::CommitMaterialSetUpdate(const FComponentMaterialSet& MaterialSet, bool bApplyToAsset)
 {
 	if (!ensure(IsValid())) return false;
 
@@ -93,7 +93,7 @@ bool UStaticMeshComponentToolTarget::CommitMaterialSetUpdate(const FComponentMat
 		// unregister the component while we update it's static mesh
 		TUniquePtr<FComponentReregisterContext> ComponentReregisterContext = MakeUnique<FComponentReregisterContext>(Component.Get());
 
-		return UStaticMeshToolTarget::CommitMaterialSetUpdate(StaticMesh, MaterialSet, bApplyToAsset);
+		return UStaticMeshReadOnlyToolTarget::CommitMaterialSetUpdate(StaticMesh, MaterialSet, bApplyToAsset);
 	}
 	else
 	{
@@ -130,7 +130,7 @@ bool UStaticMeshComponentToolTarget::CommitMaterialSetUpdate(const FComponentMat
 }
 
 
-const FMeshDescription* UStaticMeshComponentToolTarget::GetMeshDescription(const FGetMeshParameters& GetMeshParams)
+const FMeshDescription* UStaticMeshComponentReadOnlyToolTarget::GetMeshDescription(const FGetMeshParameters& GetMeshParams)
 {
 	if (ensure(IsValid()))
 	{
@@ -140,26 +140,26 @@ const FMeshDescription* UStaticMeshComponentToolTarget::GetMeshDescription(const
 		EMeshLODIdentifier UseLOD = EditingLOD;
 		if (StaticMesh && GetMeshParams.bHaveRequestLOD)
 		{
-			UseLOD = UStaticMeshToolTarget::GetValidEditingLOD(StaticMesh, GetMeshParams.RequestLOD);
+			UseLOD = UStaticMeshReadOnlyToolTarget::GetValidEditingLOD(StaticMesh, GetMeshParams.RequestLOD);
 			ensure(UseLOD == GetMeshParams.RequestLOD);		// probably a bug somewhere if this is not true
 		}
 
-		return UStaticMeshToolTarget::GetMeshDescriptionWithScaleApplied(StaticMesh, (int32)UseLOD, CachedMeshDescriptions);
+		return UStaticMeshReadOnlyToolTarget::GetMeshDescriptionWithScaleApplied(StaticMesh, (int32)UseLOD, CachedMeshDescriptions);
 	}
 	return nullptr;
 }
 
-TArray<int32> UStaticMeshComponentToolTarget::GetPolygonGroupToMaterialIndexMap() const
+TArray<int32> UStaticMeshComponentReadOnlyToolTarget::GetPolygonGroupToMaterialIndexMap() const
 {
 	if (IsValid())
 	{
 		UStaticMesh* StaticMesh = Cast<UStaticMeshComponent>(Component)->GetStaticMesh();
-		return UStaticMeshToolTarget::MapSectionToMaterialID(StaticMesh, EditingLOD);
+		return UStaticMeshReadOnlyToolTarget::MapSectionToMaterialID(StaticMesh, EditingLOD);
 	}
 	return TArray<int32>();
 }
 
-FMeshDescription UStaticMeshComponentToolTarget::GetEmptyMeshDescription()
+FMeshDescription UStaticMeshComponentReadOnlyToolTarget::GetEmptyMeshDescription()
 {
 	FMeshDescription EmptyMeshDescription;
 	FStaticMeshAttributes Attributes(EmptyMeshDescription);
@@ -167,7 +167,7 @@ FMeshDescription UStaticMeshComponentToolTarget::GetEmptyMeshDescription()
 	return EmptyMeshDescription;
 }
 
-FMeshDescription UStaticMeshComponentToolTarget::GetMeshDescriptionCopy(const FGetMeshParameters& GetMeshParams)
+FMeshDescription UStaticMeshComponentReadOnlyToolTarget::GetMeshDescriptionCopy(const FGetMeshParameters& GetMeshParams)
 {
 	auto ApplyBuildScaleIfNeeded = [](FMeshDescription& MeshDescription, FVector BuildScale)
 	{
@@ -185,7 +185,7 @@ FMeshDescription UStaticMeshComponentToolTarget::GetMeshDescriptionCopy(const FG
 			EMeshLODIdentifier UseLOD = EditingLOD;
 			if (GetMeshParams.bHaveRequestLOD)
 			{
-				UseLOD = UStaticMeshToolTarget::GetValidEditingLOD(StaticMesh, GetMeshParams.RequestLOD);
+				UseLOD = UStaticMeshReadOnlyToolTarget::GetValidEditingLOD(StaticMesh, GetMeshParams.RequestLOD);
 				ensure(UseLOD == GetMeshParams.RequestLOD);		// probably a bug somewhere if this is not true
 			}
 			if (UseLOD == EMeshLODIdentifier::HiResSource )
@@ -223,10 +223,10 @@ FMeshDescription UStaticMeshComponentToolTarget::GetMeshDescriptionCopy(const FG
 }
 
 
-TArray<EMeshLODIdentifier> UStaticMeshComponentToolTarget::GetAvailableLODs(bool bSkipAutoGenerated) const
+TArray<EMeshLODIdentifier> UStaticMeshComponentReadOnlyToolTarget::GetAvailableLODs(bool bSkipAutoGenerated) const
 {
 	UStaticMesh* StaticMesh = Cast<UStaticMeshComponent>(Component)->GetStaticMesh();
-	return UStaticMeshToolTarget::GetAvailableLODs(StaticMesh, bSkipAutoGenerated);
+	return UStaticMeshReadOnlyToolTarget::GetAvailableLODs(StaticMesh, bSkipAutoGenerated);
 }
 
 
@@ -247,12 +247,12 @@ void UStaticMeshComponentToolTarget::CommitMeshDescription(const FCommitter& Com
 	Component->RecreatePhysicsState();
 }
 
-FDynamicMesh3 UStaticMeshComponentToolTarget::GetDynamicMesh()
+FDynamicMesh3 UStaticMeshComponentReadOnlyToolTarget::GetDynamicMesh()
 {
 	return GetDynamicMeshViaMeshDescription(*this);
 }
 
-FDynamicMesh3 UStaticMeshComponentToolTarget::GetDynamicMesh(const FGetMeshParameters& InGetMeshParams)
+FDynamicMesh3 UStaticMeshComponentReadOnlyToolTarget::GetDynamicMesh(const FGetMeshParameters& InGetMeshParams)
 {
 	return GetDynamicMeshViaMeshDescription(*this, InGetMeshParams);
 }
@@ -267,13 +267,13 @@ void UStaticMeshComponentToolTarget::CommitDynamicMesh(const FDynamicMesh3& Mesh
 	}
 }
 
-UStaticMesh* UStaticMeshComponentToolTarget::GetStaticMesh() const
+UStaticMesh* UStaticMeshComponentReadOnlyToolTarget::GetStaticMesh() const
 {
 	return IsValid() ? Cast<UStaticMeshComponent>(Component)->GetStaticMesh() : nullptr;
 }
 
 
-UBodySetup* UStaticMeshComponentToolTarget::GetBodySetup() const
+UBodySetup* UStaticMeshComponentReadOnlyToolTarget::GetBodySetup() const
 {
 	UStaticMesh* StaticMesh = GetStaticMesh();
 	if (StaticMesh)
@@ -284,7 +284,7 @@ UBodySetup* UStaticMeshComponentToolTarget::GetBodySetup() const
 }
 
 
-IInterface_CollisionDataProvider* UStaticMeshComponentToolTarget::GetComplexCollisionProvider() const
+IInterface_CollisionDataProvider* UStaticMeshComponentReadOnlyToolTarget::GetComplexCollisionProvider() const
 {
 	UStaticMesh* StaticMesh = GetStaticMesh();
 	if (StaticMesh)
@@ -297,19 +297,44 @@ IInterface_CollisionDataProvider* UStaticMeshComponentToolTarget::GetComplexColl
 
 // Factory
 
-bool UStaticMeshComponentToolTargetFactory::CanBuildTarget(UObject* SourceObject, const FToolTargetTypeRequirements& Requirements) const
+const UStaticMesh* UStaticMeshComponentToolTargetFactory::SourceToStaticMesh(const UObject* SourceObject)
 {
 	const UStaticMeshComponent* Component = GetValid(Cast<UStaticMeshComponent>(SourceObject));
 	const UStaticMesh* StaticMesh = (Component && !Component->IsUnreachable() && Component->IsValidLowLevel()) ? Component->GetStaticMesh() : nullptr;
-	return StaticMesh
+	return StaticMesh;
+}
+
+bool UStaticMeshComponentToolTargetFactory::CanBuildTarget(UObject* SourceObject, const FToolTargetTypeRequirements& Requirements) const
+{
+	const UStaticMesh* StaticMesh = SourceToStaticMesh(SourceObject);
+	bool bValid = StaticMesh
 		&& !StaticMesh->GetOutermost()->bIsCookedForEditor
-		&& UStaticMeshToolTarget::HasNonGeneratedLOD(StaticMesh, EditingLOD)
-		&& Requirements.AreSatisfiedBy(UStaticMeshComponentToolTarget::StaticClass());
+		&& UStaticMeshToolTarget::HasNonGeneratedLOD(StaticMesh, EditingLOD);
+	if (!bValid)
+	{
+		return false;
+	}
+	if (CanWriteToSource(SourceObject))
+	{
+		return Requirements.AreSatisfiedBy(UStaticMeshComponentToolTarget::StaticClass());
+	}
+	else
+	{
+		return Requirements.AreSatisfiedBy(UStaticMeshComponentReadOnlyToolTarget::StaticClass());
+	}
 }
 
 UToolTarget* UStaticMeshComponentToolTargetFactory::BuildTarget(UObject* SourceObject, const FToolTargetTypeRequirements& Requirements)
 {
-	UStaticMeshComponentToolTarget* Target = NewObject<UStaticMeshComponentToolTarget>();// TODO: Should we set an outer here?
+	UStaticMeshComponentReadOnlyToolTarget* Target;
+	if (CanWriteToSource(SourceObject))
+	{
+		Target = NewObject<UStaticMeshComponentToolTarget>();
+	}
+	else
+	{
+		Target = NewObject<UStaticMeshComponentReadOnlyToolTarget>();
+	}
 	Target->InitializeComponent(Cast<UStaticMeshComponent>(SourceObject));
 	Target->SetEditingLOD(EditingLOD);
 	checkSlow(Target->Component.Get() && Requirements.AreSatisfiedBy(Target));
@@ -322,3 +347,14 @@ void UStaticMeshComponentToolTargetFactory::SetActiveEditingLOD(EMeshLODIdentifi
 {
 	EditingLOD = NewEditingLOD;
 }
+
+
+bool UStaticMeshComponentToolTargetFactory::CanWriteToSource(const UObject* Source)
+{
+	const UStaticMesh* StaticMesh = SourceToStaticMesh(Source);
+	if (StaticMesh)
+	{
+		return !StaticMesh->GetPathName().StartsWith(TEXT("/Engine/"));
+	}
+	return false;
+}
diff --git a/Engine/Plugins/Runtime/MeshModelingToolset/Source/ModelingComponentsEditorOnly/Private/ToolTargets/StaticMeshToolTarget.cpp b/Engine/Plugins/Runtime/MeshModelingToolset/Source/ModelingComponentsEditorOnly/Private/ToolTargets/StaticMeshToolTarget.cpp
--- a/Engine/Plugins/Runtime/MeshModelingToolset/Source/ModelingComponentsEditorOnly/Private/ToolTargets/StaticMeshToolTarget.cpp
+++ b/Engine/Plugins/Runtime/MeshModelingToolset/Source/ModelingComponentsEditorOnly/Private/ToolTargets/StaticMeshToolTarget.cpp
@@ -35,12 +35,12 @@ namespace StaticMeshToolTargetLocals
 	}
 }
 
-void UStaticMeshToolTarget::SetEditingLOD(EMeshLODIdentifier RequestedEditingLOD)
+void UStaticMeshReadOnlyToolTarget::SetEditingLOD(EMeshLODIdentifier RequestedEditingLOD)
 {
 	EditingLOD = GetValidEditingLOD(StaticMesh.Get(), RequestedEditingLOD);
 }
 
-EMeshLODIdentifier UStaticMeshToolTarget::GetValidEditingLOD(const UStaticMesh* StaticMeshIn, 
+EMeshLODIdentifier UStaticMeshReadOnlyToolTarget::GetValidEditingLOD(const UStaticMesh* StaticMeshIn, 
 	EMeshLODIdentifier RequestedEditingLOD)
 {
 	using namespace StaticMeshToolTargetLocals;
@@ -76,12 +76,12 @@ EMeshLODIdentifier UStaticMeshToolTarget::GetValidEditingLOD(const UStaticMesh*
 	return ValidEditingLOD;
 }
 
-TArray<EMeshLODIdentifier> UStaticMeshToolTarget::GetAvailableLODs(bool bSkipAutoGenerated) const
+TArray<EMeshLODIdentifier> UStaticMeshReadOnlyToolTarget::GetAvailableLODs(bool bSkipAutoGenerated) const
 {
 	return GetAvailableLODs(StaticMesh.Get(), bSkipAutoGenerated);
 }
 
-TArray<EMeshLODIdentifier> UStaticMeshToolTarget::GetAvailableLODs(const UStaticMesh* InStaticMesh, bool bSkipAutoGenerated)
+TArray<EMeshLODIdentifier> UStaticMeshReadOnlyToolTarget::GetAvailableLODs(const UStaticMesh* InStaticMesh, bool bSkipAutoGenerated)
 {
 	TArray<EMeshLODIdentifier> AvailableLODs;
 	
@@ -112,14 +112,14 @@ TArray<EMeshLODIdentifier> UStaticMeshToolTarget::GetAvailableLODs(const UStatic
 }
 
 
-bool UStaticMeshToolTarget::IsValid() const
+bool UStaticMeshReadOnlyToolTarget::IsValid() const
 {
 	// The StaticMesh.Get() call here will give a nullptr if object is not alive, which 
 	// HasNonGeneratedLOD checks for.
 	return HasNonGeneratedLOD(StaticMesh.Get(), EditingLOD);
 }
 
-bool UStaticMeshToolTarget::IsValid(const UStaticMesh* StaticMeshIn, EMeshLODIdentifier EditingLODIn)
+bool UStaticMeshReadOnlyToolTarget::IsValid(const UStaticMesh* StaticMeshIn, EMeshLODIdentifier EditingLODIn)
 {
 	if (!StaticMeshIn || !IsValidChecked(StaticMeshIn) || StaticMeshIn->IsUnreachable() || !StaticMeshIn->IsValidLowLevel())
 	{
@@ -141,7 +141,7 @@ bool UStaticMeshToolTarget::IsValid(const UStaticMesh* StaticMeshIn, EMeshLODIde
 	return true;
 }
 
-bool UStaticMeshToolTarget::HasNonGeneratedLOD(const UStaticMesh* StaticMeshIn, EMeshLODIdentifier MeshLOD)
+bool UStaticMeshReadOnlyToolTarget::HasNonGeneratedLOD(const UStaticMesh* StaticMeshIn, EMeshLODIdentifier MeshLOD)
 {
 	if (!StaticMeshIn)
 	{
@@ -175,24 +175,24 @@ bool UStaticMeshToolTarget::HasNonGeneratedLOD(const UStaticMesh* StaticMeshIn,
 	}
 }
 
-int32 UStaticMeshToolTarget::GetNumMaterials() const
+int32 UStaticMeshReadOnlyToolTarget::GetNumMaterials() const
 {
 	return ensure(IsValid()) ? StaticMesh->GetStaticMaterials().Num() : 0;
 }
 
-UMaterialInterface* UStaticMeshToolTarget::GetMaterial(int32 MaterialIndex) const
+UMaterialInterface* UStaticMeshReadOnlyToolTarget::GetMaterial(int32 MaterialIndex) const
 {
 	return ensure(IsValid()) ? StaticMesh->GetMaterial(MaterialIndex) : nullptr;
 }
 
-void UStaticMeshToolTarget::GetMaterialSet(FComponentMaterialSet& MaterialSetOut, bool bPreferAssetMaterials) const
+void UStaticMeshReadOnlyToolTarget::GetMaterialSet(FComponentMaterialSet& MaterialSetOut, bool bPreferAssetMaterials) const
 {
 	if (!ensure(IsValid())) return;
 
 	GetMaterialSet(StaticMesh.Get(), MaterialSetOut, bPreferAssetMaterials);
 }
 
-void UStaticMeshToolTarget::GetMaterialSet(const UStaticMesh* StaticMeshIn, 
+void UStaticMeshReadOnlyToolTarget::GetMaterialSet(const UStaticMesh* StaticMeshIn, 
 	FComponentMaterialSet& MaterialSetOut, bool bPreferAssetMaterials)
 {
 	int32 NumMaterials = StaticMeshIn->GetStaticMaterials().Num();
@@ -203,14 +203,14 @@ void UStaticMeshToolTarget::GetMaterialSet(const UStaticMesh* StaticMeshIn,
 	}
 }
 
-bool UStaticMeshToolTarget::CommitMaterialSetUpdate(const FComponentMaterialSet& MaterialSet, bool bApplyToAsset)
+bool UStaticMeshReadOnlyToolTarget::CommitMaterialSetUpdate(const FComponentMaterialSet& MaterialSet, bool bApplyToAsset)
 {
 	if (!ensure(IsValid())) return false;
 
 	return CommitMaterialSetUpdate(StaticMesh.Get(), MaterialSet, bApplyToAsset);
 }
 
-bool UStaticMeshToolTarget::CommitMaterialSetUpdate(UStaticMesh* StaticMeshIn,
+bool UStaticMeshReadOnlyToolTarget::CommitMaterialSetUpdate(UStaticMesh* StaticMeshIn,
 	const FComponentMaterialSet& MaterialSet, bool bApplyToAsset)
 {
 	if (!bApplyToAsset)
@@ -260,14 +260,14 @@ bool UStaticMeshToolTarget::CommitMaterialSetUpdate(UStaticMesh* StaticMeshIn,
 	return true;
 }
 
-const FMeshDescription* UStaticMeshToolTarget::GetMeshDescription(const FGetMeshParameters& GetMeshParams)
+const FMeshDescription* UStaticMeshReadOnlyToolTarget::GetMeshDescription(const FGetMeshParameters& GetMeshParams)
 {
 	if (ensure(IsValid()))
 	{
 		EMeshLODIdentifier UseLOD = EditingLOD;
 		if (GetMeshParams.bHaveRequestLOD)
 		{
-			UseLOD = UStaticMeshToolTarget::GetValidEditingLOD(StaticMesh.Get(), GetMeshParams.RequestLOD);
+			UseLOD = UStaticMeshReadOnlyToolTarget::GetValidEditingLOD(StaticMesh.Get(), GetMeshParams.RequestLOD);
 			ensure(UseLOD == GetMeshParams.RequestLOD);		// probably a bug somewhere if this is not true
 		}
 
@@ -276,16 +276,16 @@ const FMeshDescription* UStaticMeshToolTarget::GetMeshDescription(const FGetMesh
 	return nullptr;
 }
 
-TArray<int32> UStaticMeshToolTarget::GetPolygonGroupToMaterialIndexMap() const
+TArray<int32> UStaticMeshReadOnlyToolTarget::GetPolygonGroupToMaterialIndexMap() const
 {
 	if (IsValid())
 	{
-		return UStaticMeshToolTarget::MapSectionToMaterialID(StaticMesh.Get(), EditingLOD);
+		return UStaticMeshReadOnlyToolTarget::MapSectionToMaterialID(StaticMesh.Get(), EditingLOD);
 	}
 	return TArray<int32>();
 }
 
-TArray<int32> UStaticMeshToolTarget::MapSectionToMaterialID(const UStaticMesh* Mesh, EMeshLODIdentifier EditingLOD)
+TArray<int32> UStaticMeshReadOnlyToolTarget::MapSectionToMaterialID(const UStaticMesh* Mesh, EMeshLODIdentifier EditingLOD)
 {
 	UE::Conversion::EMeshLODType LODType = [EditingLOD]() {
 		switch (EditingLOD)
@@ -306,7 +306,7 @@ TArray<int32> UStaticMeshToolTarget::MapSectionToMaterialID(const UStaticMesh* M
 	return UE::Conversion::GetPolygonGroupToMaterialIndexMap(Mesh, LODType, UseLOD);
 }
 
-const FMeshDescription* UStaticMeshToolTarget::GetMeshDescriptionWithScaleApplied(UStaticMesh* StaticMesh, int32 UseLOD, FMeshDescriptionCache& CachedMeshDescriptions)
+const FMeshDescription* UStaticMeshReadOnlyToolTarget::GetMeshDescriptionWithScaleApplied(UStaticMesh* StaticMesh, int32 UseLOD, FMeshDescriptionCache& CachedMeshDescriptions)
 {
 	static FMeshDescription EmptyMeshDescription;
 	static bool bFirst = true;
@@ -364,7 +364,7 @@ const FMeshDescription* UStaticMeshToolTarget::GetMeshDescriptionWithScaleApplie
 	return (FoundMeshDescription != nullptr) ? FoundMeshDescription : &EmptyMeshDescription;
 }
 
-FMeshDescription UStaticMeshToolTarget::GetEmptyMeshDescription()
+FMeshDescription UStaticMeshReadOnlyToolTarget::GetEmptyMeshDescription()
 {
 	FMeshDescription EmptyMeshDescription;
 	FStaticMeshAttributes Attributes(EmptyMeshDescription);
@@ -485,12 +485,12 @@ void UStaticMeshToolTarget::CommitMeshDescription(UStaticMesh* StaticMeshIn, con
 	UE::ToolTarget::Internal::PostEditChangeWithConditionalUndo(StaticMeshIn);
 }
 
-FDynamicMesh3 UStaticMeshToolTarget::GetDynamicMesh()
+FDynamicMesh3 UStaticMeshReadOnlyToolTarget::GetDynamicMesh()
 {
 	return GetDynamicMeshViaMeshDescription(*this);
 }
 
-FDynamicMesh3 UStaticMeshToolTarget::GetDynamicMesh(const FGetMeshParameters& InGetMeshParams)
+FDynamicMesh3 UStaticMeshReadOnlyToolTarget::GetDynamicMesh(const FGetMeshParameters& InGetMeshParams)
 {
 	return GetDynamicMeshViaMeshDescription(*this, InGetMeshParams);
 }
@@ -501,7 +501,8 @@ void UStaticMeshToolTarget::CommitDynamicMesh(const FDynamicMesh3& Mesh, const F
 	CommitDynamicMeshViaMeshDescription(MoveTemp(CurrentMeshDescription), *this, Mesh, CommitInfo);
 }
 
-UStaticMesh* UStaticMeshToolTarget::GetStaticMesh() const
+
+UStaticMesh* UStaticMeshReadOnlyToolTarget::GetStaticMesh() const
 {
 	return IsValid() ? StaticMesh.Get() : nullptr;
 }
@@ -512,15 +513,34 @@ UStaticMesh* UStaticMeshToolTarget::GetStaticMesh() const
 bool UStaticMeshToolTargetFactory::CanBuildTarget(UObject* SourceObject, const FToolTargetTypeRequirements& Requirements) const
 {
 	const UStaticMesh* StaticMesh = GetValid(Cast<UStaticMesh>(SourceObject));
-	return StaticMesh && !StaticMesh->IsUnreachable() && StaticMesh->IsValidLowLevel()
-		&& !StaticMesh->GetOutermost()->bIsCookedForEditor
-		&& UStaticMeshToolTarget::HasNonGeneratedLOD(StaticMesh, EditingLOD)
-		&& Requirements.AreSatisfiedBy(UStaticMeshToolTarget::StaticClass());
+	bool bValid = StaticMesh && !StaticMesh->IsUnreachable() && StaticMesh->IsValidLowLevel()
+			&& !StaticMesh->GetOutermost()->bIsCookedForEditor
+			&& UStaticMeshToolTarget::HasNonGeneratedLOD(StaticMesh, EditingLOD);
+	if (!bValid)
+	{
+		return false;
+	}
+	if (CanWriteToSource(SourceObject))
+	{
+		return Requirements.AreSatisfiedBy(UStaticMeshToolTarget::StaticClass());
+	}
+	else
+	{
+		return Requirements.AreSatisfiedBy(UStaticMeshReadOnlyToolTarget::StaticClass());
+	}
 }
 
 UToolTarget* UStaticMeshToolTargetFactory::BuildTarget(UObject* SourceObject, const FToolTargetTypeRequirements& Requirements)
 {
-	UStaticMeshToolTarget* Target = NewObject<UStaticMeshToolTarget>();// TODO: Should we set an outer here?
+	UStaticMeshReadOnlyToolTarget* Target;
+	if (CanWriteToSource(SourceObject))
+	{ 
+		Target = NewObject<UStaticMeshToolTarget>();
+	}
+	else
+	{
+		Target = NewObject<UStaticMeshReadOnlyToolTarget>();
+	}
 	Target->StaticMesh = Cast<UStaticMesh>(SourceObject);
 	Target->SetEditingLOD(EditingLOD);
 	checkSlow(Target->StaticMesh.IsValid() && Requirements.AreSatisfiedBy(Target));
@@ -534,5 +554,15 @@ void UStaticMeshToolTargetFactory::SetActiveEditingLOD(EMeshLODIdentifier NewEdi
 	EditingLOD = NewEditingLOD;
 }
 
+
+bool UStaticMeshToolTargetFactory::CanWriteToSource(const UObject* SourceObject)
+{
+	if (const UStaticMesh* StaticMesh = GetValid(Cast<UStaticMesh>(SourceObject)))
+	{
+		return !StaticMesh->GetPathName().StartsWith(TEXT("/Engine/"));
+	}
+	return false;
+}
+
 #undef LOCTEXT_NAMESPACE
 
diff --git a/Engine/Plugins/Runtime/MeshModelingToolset/Source/ModelingComponentsEditorOnly/Public/ToolTargets/SkeletalMeshComponentToolTarget.h b/Engine/Plugins/Runtime/MeshModelingToolset/Source/ModelingComponentsEditorOnly/Public/ToolTargets/SkeletalMeshComponentToolTarget.h
--- a/Engine/Plugins/Runtime/MeshModelingToolset/Source/ModelingComponentsEditorOnly/Public/ToolTargets/SkeletalMeshComponentToolTarget.h
+++ b/Engine/Plugins/Runtime/MeshModelingToolset/Source/ModelingComponentsEditorOnly/Public/ToolTargets/SkeletalMeshComponentToolTarget.h
@@ -19,8 +19,8 @@ class USkeletalMesh;
 
 
 /**
- * A tool target backed by a read-only skeletal mesh component that can provide and take a mesh
- * description.
+ * A tool target backed by a skeletal mesh component that can provide and take a mesh
+ * description, without 'committer' interfaces for writing to the source skeletal mesh asset.
  */
 UCLASS(Transient)
 class MODELINGCOMPONENTSEDITORONLY_API USkeletalMeshComponentReadOnlyToolTarget :
@@ -64,6 +64,7 @@ protected:
 	
 	// So that the tool target factory can poke into Component.
 	friend class USkeletalMeshComponentReadOnlyToolTargetFactory;
+	friend class USkeletalMeshComponentToolTargetFactory;
 };
 
 
@@ -119,4 +120,6 @@ public:
 	virtual bool CanBuildTarget(UObject* SourceObject, const FToolTargetTypeRequirements& TargetTypeInfo) const override;
 
 	virtual UToolTarget* BuildTarget(UObject* SourceObject, const FToolTargetTypeRequirements& TargetTypeInfo) override;
+
+	static bool CanWriteToSource(const UObject* SourceObject);
 };
diff --git a/Engine/Plugins/Runtime/MeshModelingToolset/Source/ModelingComponentsEditorOnly/Public/ToolTargets/SkeletalMeshToolTarget.h b/Engine/Plugins/Runtime/MeshModelingToolset/Source/ModelingComponentsEditorOnly/Public/ToolTargets/SkeletalMeshToolTarget.h
--- a/Engine/Plugins/Runtime/MeshModelingToolset/Source/ModelingComponentsEditorOnly/Public/ToolTargets/SkeletalMeshToolTarget.h
+++ b/Engine/Plugins/Runtime/MeshModelingToolset/Source/ModelingComponentsEditorOnly/Public/ToolTargets/SkeletalMeshToolTarget.h
@@ -65,6 +65,7 @@ protected:
 	
 	// So that the tool target factory can poke into Component.
 	friend class USkeletalMeshReadOnlyToolTargetFactory;
+	friend class USkeletalMeshToolTargetFactory;
 	friend class USkeletalMeshComponentReadOnlyToolTarget;
 	friend class USkeletalMeshComponentToolTarget;
 
@@ -135,4 +136,6 @@ public:
 	virtual bool CanBuildTarget(UObject* SourceObject, const FToolTargetTypeRequirements& TargetTypeInfo) const override;
 
 	virtual UToolTarget* BuildTarget(UObject* SourceObject, const FToolTargetTypeRequirements& TargetTypeInfo) override;
+
+	static bool CanWriteToSource(const UObject* SourceObject);
 };
diff --git a/Engine/Plugins/Runtime/MeshModelingToolset/Source/ModelingComponentsEditorOnly/Public/ToolTargets/StaticMeshComponentToolTarget.h b/Engine/Plugins/Runtime/MeshModelingToolset/Source/ModelingComponentsEditorOnly/Public/ToolTargets/StaticMeshComponentToolTarget.h
--- a/Engine/Plugins/Runtime/MeshModelingToolset/Source/ModelingComponentsEditorOnly/Public/ToolTargets/StaticMeshComponentToolTarget.h
+++ b/Engine/Plugins/Runtime/MeshModelingToolset/Source/ModelingComponentsEditorOnly/Public/ToolTargets/StaticMeshComponentToolTarget.h
@@ -21,12 +21,12 @@ class UStaticMesh;
 
 /**
  * A tool target backed by a static mesh component that can provide and take a mesh
- * description.
+ * description, without 'committer' interfaces for writing to the source static mesh asset.
  */
 UCLASS(Transient)
-class MODELINGCOMPONENTSEDITORONLY_API UStaticMeshComponentToolTarget : public UPrimitiveComponentToolTarget,
-	public IMeshDescriptionCommitter, public IMeshDescriptionProvider, public IMaterialProvider, public IStaticMeshBackedTarget,
-	public IDynamicMeshProvider, public IDynamicMeshCommitter, public IPhysicsDataSource
+class MODELINGCOMPONENTSEDITORONLY_API UStaticMeshComponentReadOnlyToolTarget : public UPrimitiveComponentToolTarget,
+	public IMeshDescriptionProvider, public IMaterialProvider, public IStaticMeshBackedTarget,
+	public IDynamicMeshProvider, public IPhysicsDataSource
 {
 	GENERATED_BODY()
 
@@ -54,17 +54,12 @@ public:
 	virtual EMeshLODIdentifier GetMeshDescriptionLOD() const override { return EditingLOD; }
 	virtual TArray<int32> GetPolygonGroupToMaterialIndexMap() const override;
 
-	// IMeshDescriptionCommitter implementation
-	virtual void CommitMeshDescription(const FCommitter& Committer, const FCommitMeshParameters& CommitParams = FCommitMeshParameters()) override;
-	using IMeshDescriptionCommitter::CommitMeshDescription; // unhide the other overload
+
 
 	// IDynamicMeshProvider
 	virtual UE::Geometry::FDynamicMesh3 GetDynamicMesh() override;
 	virtual UE::Geometry::FDynamicMesh3 GetDynamicMesh(const FGetMeshParameters& InGetMeshParams) override;
 
-	// IDynamicMeshCommitter
-	virtual void CommitDynamicMesh(const UE::Geometry::FDynamicMesh3& Mesh, const FDynamicMeshCommitInfo& CommitInfo) override;
-	using IDynamicMeshCommitter::CommitDynamicMesh; // unhide the other overload
 
 	// IMaterialProvider implementation
 	virtual int32 GetNumMaterials() const override;
@@ -92,6 +87,26 @@ private:
 	UStaticMeshToolTarget::FMeshDescriptionCache CachedMeshDescriptions;
 };
 
+/**
+ * A tool target backed by a static mesh component that can provide and take a mesh
+ * description.
+ */
+UCLASS(Transient)
+class MODELINGCOMPONENTSEDITORONLY_API UStaticMeshComponentToolTarget : 
+	public UStaticMeshComponentReadOnlyToolTarget,
+	public IMeshDescriptionCommitter, public IDynamicMeshCommitter
+{
+	GENERATED_BODY()
+
+	// IDynamicMeshCommitter
+	virtual void CommitDynamicMesh(const UE::Geometry::FDynamicMesh3& Mesh, const FDynamicMeshCommitInfo& CommitInfo) override;
+	using IDynamicMeshCommitter::CommitDynamicMesh; // unhide the other overload
+	
+	// IMeshDescriptionCommitter implementation
+	virtual void CommitMeshDescription(const FCommitter& Committer, const FCommitMeshParameters& CommitParams = FCommitMeshParameters()) override;
+	using IMeshDescriptionCommitter::CommitMeshDescription; // unhide the other overload
+};
+
 
 /** Factory for UStaticMeshComponentToolTarget to be used by the target manager. */
 UCLASS(Transient)
@@ -108,7 +123,12 @@ public:
 	virtual EMeshLODIdentifier GetActiveEditingLOD() const { return EditingLOD; }
 	virtual void SetActiveEditingLOD(EMeshLODIdentifier NewEditingLOD);
 
+	static bool CanWriteToSource(const UObject* Source);
+
 protected:
 	// LOD to edit, default is to edit LOD0
 	EMeshLODIdentifier EditingLOD = EMeshLODIdentifier::LOD0;
+
+private:
+	static const UStaticMesh* SourceToStaticMesh(const UObject* SourceObject);
 };
\ No newline at end of file
diff --git a/Engine/Plugins/Runtime/MeshModelingToolset/Source/ModelingComponentsEditorOnly/Public/ToolTargets/StaticMeshToolTarget.h b/Engine/Plugins/Runtime/MeshModelingToolset/Source/ModelingComponentsEditorOnly/Public/ToolTargets/StaticMeshToolTarget.h
--- a/Engine/Plugins/Runtime/MeshModelingToolset/Source/ModelingComponentsEditorOnly/Public/ToolTargets/StaticMeshToolTarget.h
+++ b/Engine/Plugins/Runtime/MeshModelingToolset/Source/ModelingComponentsEditorOnly/Public/ToolTargets/StaticMeshToolTarget.h
@@ -19,20 +19,18 @@ class UStaticMesh;
 
 /**
  * A tool target backed by a static mesh asset that can provide and take a mesh
- * description.
+ * description, without 'committer' interfaces for writing to the source static mesh asset.
  * 
  * This is a special tool target that refers to the underlying asset (in this case a static mesh), rather than indirectly through a component.
  * This type of target is used in cases, such as opening an asset through the content browser, when there is no component available.
  */
 UCLASS(Transient)
-class MODELINGCOMPONENTSEDITORONLY_API UStaticMeshToolTarget : 
+class MODELINGCOMPONENTSEDITORONLY_API UStaticMeshReadOnlyToolTarget : 
 	public UToolTarget,
-	public IMeshDescriptionCommitter, 
 	public IMeshDescriptionProvider, 
 	public IMaterialProvider, 
 	public IStaticMeshBackedTarget,
-	public IDynamicMeshProvider, 
-	public IDynamicMeshCommitter
+	public IDynamicMeshProvider
 {
 	GENERATED_BODY()
 
@@ -60,17 +58,11 @@ public:
 	// helper to get the material mapping returned by GetPolygonGroupToMaterialIndexMap for any static mesh
 	static TArray<int32> MapSectionToMaterialID(const UStaticMesh* Mesh, EMeshLODIdentifier EditingLOD);
 
-	// IMeshDescritpionCommitter implementation
-	virtual void CommitMeshDescription(const FCommitter& Committer, const FCommitMeshParameters& CommitParams = FCommitMeshParameters()) override;
-	using IMeshDescriptionCommitter::CommitMeshDescription; // unhide the other overload
 
 	// IDynamicMeshProvider
 	virtual UE::Geometry::FDynamicMesh3 GetDynamicMesh() override;
 	virtual UE::Geometry::FDynamicMesh3 GetDynamicMesh(const FGetMeshParameters& InGetMeshParams) override;
 
-	// IDynamicMeshCommitter
-	virtual void CommitDynamicMesh(const UE::Geometry::FDynamicMesh3& Mesh, const FDynamicMeshCommitInfo& CommitInfo) override;
-	using IDynamicMeshCommitter::CommitDynamicMesh; // unhide the other overload
 
 	// IMaterialProvider implementation
 	virtual int32 GetNumMaterials() const override;
@@ -101,7 +93,7 @@ protected:
 
 	friend class UStaticMeshToolTargetFactory;
 
-	friend class UStaticMeshComponentToolTarget;
+	friend class UStaticMeshComponentReadOnlyToolTarget;
 
 	friend class UStaticMeshComponentToolTargetFactory;
 
@@ -115,8 +107,6 @@ protected:
 	static EMeshLODIdentifier GetValidEditingLOD(const UStaticMesh* StaticMesh, 
 		EMeshLODIdentifier RequestedEditingLOD);
 	static TArray<EMeshLODIdentifier> GetAvailableLODs(const UStaticMesh* StaticMesh, bool bSkipAutoGenerated);
-	static void CommitMeshDescription(UStaticMesh* StaticMesh, 
-		const FCommitter& Committer, EMeshLODIdentifier EditingLODIn);
 	static void GetMaterialSet(const UStaticMesh* StaticMesh, 
 		FComponentMaterialSet& MaterialSetOut, bool bPreferAssetMaterials);
 	static bool CommitMaterialSetUpdate(UStaticMesh* SkeletalMesh, 
@@ -128,6 +118,37 @@ private:
 	FMeshDescriptionCache CachedMeshDescriptions;
 };
 
+/**
+ * A tool target backed by a static mesh asset that can provide and take a mesh
+ * description.
+ *
+ * This is a special tool target that refers to the underlying asset (in this case a static mesh), rather than indirectly through a component.
+ * This type of target is used in cases, such as opening an asset through the content browser, when there is no component available.
+ */
+UCLASS(Transient)
+class MODELINGCOMPONENTSEDITORONLY_API UStaticMeshToolTarget : 
+	public UStaticMeshReadOnlyToolTarget,
+	public IMeshDescriptionCommitter,
+	public IDynamicMeshCommitter
+{
+	GENERATED_BODY()
+	
+public:
+	// IMeshDescriptionCommitter implementation
+	virtual void CommitMeshDescription(const FCommitter& Committer, const FCommitMeshParameters& CommitParams = FCommitMeshParameters()) override;
+	using IMeshDescriptionCommitter::CommitMeshDescription; // unhide the other overload
+
+	// IDynamicMeshCommitter
+	virtual void CommitDynamicMesh(const UE::Geometry::FDynamicMesh3& Mesh, const FDynamicMeshCommitInfo& CommitInfo) override;
+	using IDynamicMeshCommitter::CommitDynamicMesh; // unhide the other overload
+
+protected:
+	friend class UStaticMeshComponentToolTarget;
+
+	static void CommitMeshDescription(UStaticMesh* StaticMesh,
+		const FCommitter& Committer, EMeshLODIdentifier EditingLODIn);
+};
+
 
 /** Factory for UStaticMeshToolTarget to be used by the target manager. */
 UCLASS(Transient)
@@ -147,6 +168,8 @@ public:
 	virtual EMeshLODIdentifier GetActiveEditingLOD() const { return EditingLOD; }
 	virtual void SetActiveEditingLOD(EMeshLODIdentifier NewEditingLOD);
 
+	static bool CanWriteToSource(const UObject* Source);
+
 protected:
 	// LOD to edit, default is to edit LOD0
 	EMeshLODIdentifier EditingLOD = EMeshLODIdentifier::LOD0;
