re-enable duplicate tool for engine assets

diff --git a/Engine/Plugins/Runtime/MeshModelingToolset/Source/MeshModelingTools/Private/CombineMeshesTool.cpp b/Engine/Plugins/Runtime/MeshModelingToolset/Source/MeshModelingTools/Private/CombineMeshesTool.cpp
--- a/Engine/Plugins/Runtime/MeshModelingToolset/Source/MeshModelingTools/Private/CombineMeshesTool.cpp
+++ b/Engine/Plugins/Runtime/MeshModelingToolset/Source/MeshModelingTools/Private/CombineMeshesTool.cpp
@@ -14,8 +14,8 @@
 #include "ShapeApproximation/SimpleShapeSet3.h"
 
 #include "TargetInterfaces/MaterialProvider.h"
-#include "TargetInterfaces/MeshDescriptionCommitter.h"
-#include "TargetInterfaces/MeshDescriptionProvider.h"
+#include "TargetInterfaces/DynamicMeshProvider.h"
+#include "TargetInterfaces/DynamicMeshCommitter.h"
 #include "TargetInterfaces/PrimitiveComponentBackedTarget.h"
 #include "ToolTargetManager.h"
 #include "ModelingToolTargetUtil.h"
@@ -84,6 +84,22 @@ UMultiSelectionMeshEditingTool* UCombineMeshesToolBuilder::CreateNewTool(const F
 	return NewTool;
 }
 
+const FToolTargetTypeRequirements& UCombineMeshesToolBuilder::GetTargetRequirements() const
+{
+	static FToolTargetTypeRequirements CombineTypeRequirements({
+		UMaterialProvider::StaticClass(),
+		UDynamicMeshProvider::StaticClass(),
+		UPrimitiveComponentBackedTarget::StaticClass(),
+		UDynamicMeshCommitter::StaticClass()
+		});
+	static FToolTargetTypeRequirements DuplicateTypeRequirements({
+		UMaterialProvider::StaticClass(),
+		UDynamicMeshProvider::StaticClass(),
+		UPrimitiveComponentBackedTarget::StaticClass()
+		});
+	return bIsDuplicateTool ? DuplicateTypeRequirements : CombineTypeRequirements;
+}
+
 
 /*
  * Tool
diff --git a/Engine/Plugins/Runtime/MeshModelingToolset/Source/MeshModelingTools/Public/CombineMeshesTool.h b/Engine/Plugins/Runtime/MeshModelingToolset/Source/MeshModelingTools/Public/CombineMeshesTool.h
--- a/Engine/Plugins/Runtime/MeshModelingToolset/Source/MeshModelingTools/Public/CombineMeshesTool.h
+++ b/Engine/Plugins/Runtime/MeshModelingToolset/Source/MeshModelingTools/Public/CombineMeshesTool.h
@@ -29,6 +29,10 @@ public:
 
 	virtual bool CanBuildTool(const FToolBuilderState& SceneState) const override;
 	virtual UMultiSelectionMeshEditingTool* CreateNewTool(const FToolBuilderState& SceneState) const override;
+
+protected:
+	virtual const FToolTargetTypeRequirements& GetTargetRequirements() const override;
+
 };
 
 
