set visible defaults for line set and point set components "add lines" and "add points" method, and make incolor an auto-create-ref-term so it can use the default value

diff --git a/Engine/Plugins/Runtime/MeshModelingToolset/Source/ModelingComponents/Public/Drawing/LineSetComponent.h b/Engine/Plugins/Runtime/MeshModelingToolset/Source/ModelingComponents/Public/Drawing/LineSetComponent.h
--- a/Engine/Plugins/Runtime/MeshModelingToolset/Source/ModelingComponents/Public/Drawing/LineSetComponent.h
+++ b/Engine/Plugins/Runtime/MeshModelingToolset/Source/ModelingComponents/Public/Drawing/LineSetComponent.h
@@ -67,12 +67,12 @@ public:
 	}
 
 	/** Add lines to be rendered using the component. */
-	UFUNCTION(BlueprintCallable, Category="Line Set")
+	UFUNCTION(BlueprintCallable, Category="Line Set", meta = (AutoCreateRefTerm = "InColor"))
 	int32 AddLines(
 		const TArray<FVector>& InStart,
 		const TArray<FVector>& InEnd,
-		const FColor& InColor,
-		const float InThickness,
+		const FColor& InColor = FColor::Black,
+		const float InThickness = 1.0f,
 		const float InDepthBias = 0.0f);
 
 	/** Insert a line with the given ID to the overlay */
diff --git a/Engine/Plugins/Runtime/MeshModelingToolset/Source/ModelingComponents/Public/Drawing/PointSetComponent.h b/Engine/Plugins/Runtime/MeshModelingToolset/Source/ModelingComponents/Public/Drawing/PointSetComponent.h
--- a/Engine/Plugins/Runtime/MeshModelingToolset/Source/ModelingComponents/Public/Drawing/PointSetComponent.h
+++ b/Engine/Plugins/Runtime/MeshModelingToolset/Source/ModelingComponents/Public/Drawing/PointSetComponent.h
@@ -77,11 +77,11 @@ public:
 	 * Add points to be rendered using the component. 
 	 * @return Number of added points
 	 */
-	UFUNCTION(BlueprintCallable, Category = "Point Set")
+	UFUNCTION(BlueprintCallable, Category = "Point Set", meta = (AutoCreateRefTerm = "InColor"))
 	int32 AddPoints(
 		const TArray<FVector>& Positions,
-		const FColor& InColor,
-		const float InSize,
+		const FColor& InColor = FColor::Black,
+		const float InSize = 2.0f,
 		const float InDepthBias = 0.0f
 	);
 
