Modeling Tools: Preserve Mesh Element Selection after path extrude/extrude polygon tool incompletion or cancellation

diff --git a/Engine/Plugins/Experimental/MeshModelingToolsetExp/Source/MeshModelingToolsExp/Private/DrawPolyPathTool.cpp b/Engine/Plugins/Experimental/MeshModelingToolsetExp/Source/MeshModelingToolsExp/Private/DrawPolyPathTool.cpp
--- a/Engine/Plugins/Experimental/MeshModelingToolsetExp/Source/MeshModelingToolsExp/Private/DrawPolyPathTool.cpp
+++ b/Engine/Plugins/Experimental/MeshModelingToolsetExp/Source/MeshModelingToolsExp/Private/DrawPolyPathTool.cpp
@@ -268,6 +268,10 @@ void UDrawPolyPathTool::Shutdown(EToolShutdownType ShutdownType)
 		CurHeight = TransformProps->ExtrudeMode == EDrawPolyPathExtrudeMode::Flat ? 0.0 : TransformProps->ExtrudeHeight;
 		OnCompleteExtrudeHeight();
 	}
+	else
+	{
+		bRestoreInputSelection = true;
+	}
 
 	if (bHasSavedWidth)
 	{
@@ -379,6 +383,11 @@ void UDrawPolyPathTool::OnUpdateModifierState(int ModifierID, bool bIsOn)
 	}
 }
 
+bool UDrawPolyPathTool::IsInputSelectionValidOnOutput()
+{
+	return bRestoreInputSelection;
+}
+
 FInputRayHit UDrawPolyPathTool::IsHitByClick(const FInputDeviceRay& ClickPos)
 {
 	FHitResult OutHit;
diff --git a/Engine/Plugins/Experimental/MeshModelingToolsetExp/Source/MeshModelingToolsExp/Public/DrawPolyPathTool.h b/Engine/Plugins/Experimental/MeshModelingToolsetExp/Source/MeshModelingToolsExp/Public/DrawPolyPathTool.h
--- a/Engine/Plugins/Experimental/MeshModelingToolsetExp/Source/MeshModelingToolsExp/Public/DrawPolyPathTool.h
+++ b/Engine/Plugins/Experimental/MeshModelingToolsetExp/Source/MeshModelingToolsExp/Public/DrawPolyPathTool.h
@@ -5,6 +5,7 @@
 #include "CoreMinimal.h"
 #include "UObject/NoExportTypes.h"
 #include "InteractiveToolBuilder.h"
+#include "InteractiveToolQueryInterfaces.h"
 #include "Mechanics/PlaneDistanceFromHitMechanic.h"
 #include "Mechanics/SpatialCurveDistanceMechanic.h"
 #include "Mechanics/CollectSurfacePathMechanic.h"
@@ -159,7 +160,7 @@ public:
  *
  */
 UCLASS()
-class MESHMODELINGTOOLSEXP_API UDrawPolyPathTool : public UInteractiveTool, public IClickBehaviorTarget, public IHoverBehaviorTarget
+class MESHMODELINGTOOLSEXP_API UDrawPolyPathTool : public UInteractiveTool, public IClickBehaviorTarget, public IHoverBehaviorTarget, public IInteractiveToolManageGeometrySelectionAPI
 {
 	GENERATED_BODY()
 public:
@@ -197,6 +198,9 @@ public:
 		DrawPlaneWorld = Frame;
 	}
 
+	// IInteractiveToolManageGeometrySelectionAPI - enables restoration of previous selection when applicable
+	virtual bool IsInputSelectionValidOnOutput() override;
+
 protected:
 	UWorld* TargetWorld;
 
@@ -303,6 +307,10 @@ protected:
 	int32 CurrentCurveTimestamp = 1;
 	void UndoCurrentOperation(EState DestinationState);
 	bool CheckInCurve(int32 Timestamp) const { return CurrentCurveTimestamp == Timestamp; }
+private:
+	// if the drawn path is incomplete upon 'Accept' of the tool, no mesh should be created and the
+	// previous mesh element selection from before entering the tool should be restored
+	bool bRestoreInputSelection = true;
 };
 
 
diff --git a/Engine/Plugins/Runtime/MeshModelingToolset/Source/MeshModelingTools/Private/DrawPolygonTool.cpp b/Engine/Plugins/Runtime/MeshModelingToolset/Source/MeshModelingTools/Private/DrawPolygonTool.cpp
--- a/Engine/Plugins/Runtime/MeshModelingToolset/Source/MeshModelingTools/Private/DrawPolygonTool.cpp
+++ b/Engine/Plugins/Runtime/MeshModelingToolset/Source/MeshModelingTools/Private/DrawPolygonTool.cpp
@@ -230,7 +230,10 @@ void UDrawPolygonTool::ApplyUndoPoints(const TArray<FVector3d>& ClickPointsIn, c
 	}
 }
 
-
+bool UDrawPolygonTool::IsInputSelectionValidOnOutput()
+{
+	return bRestoreInputSelection;
+}
 
 
 void UDrawPolygonTool::OnTick(float DeltaTime)
@@ -1064,6 +1067,7 @@ void UDrawPolygonTool::EmitCurrentPolygon()
 	FCreateMeshObjectResult Result = UE::Modeling::CreateMeshObject(GetToolManager(), MoveTemp(NewMeshObjectParams));
 	if (Result.IsOK() && Result.NewActor != nullptr)
 	{
+		bRestoreInputSelection = false;
 		ToolSelectionUtil::SetNewActorSelection(GetToolManager(), Result.NewActor);
 	}
 
diff --git a/Engine/Plugins/Runtime/MeshModelingToolset/Source/MeshModelingTools/Public/DrawPolygonTool.h b/Engine/Plugins/Runtime/MeshModelingToolset/Source/MeshModelingTools/Public/DrawPolygonTool.h
--- a/Engine/Plugins/Runtime/MeshModelingToolset/Source/MeshModelingTools/Public/DrawPolygonTool.h
+++ b/Engine/Plugins/Runtime/MeshModelingToolset/Source/MeshModelingTools/Public/DrawPolygonTool.h
@@ -6,6 +6,7 @@
 #include "UObject/NoExportTypes.h"
 #include "InteractiveGizmo.h"
 #include "InteractiveToolBuilder.h"
+#include "InteractiveToolQueryInterfaces.h"
 #include "BaseTools/SingleClickTool.h"
 #include "PreviewMesh.h"
 #include "DynamicMesh/DynamicMeshAABBTree3.h"
@@ -173,7 +174,7 @@ public:
  * This tool allows the user to draw and extrude 2D polygons
  */
 UCLASS()
-class MESHMODELINGTOOLS_API UDrawPolygonTool : public UInteractiveTool, public IClickSequenceBehaviorTarget
+class MESHMODELINGTOOLS_API UDrawPolygonTool : public UInteractiveTool, public IClickSequenceBehaviorTarget, public IInteractiveToolManageGeometrySelectionAPI
 {
 	GENERATED_BODY()
 public:
@@ -221,6 +222,9 @@ public:
 		InitialDrawFrame = Frame;
 	}
 
+	// IInteractiveToolManageGeometrySelectionAPI - enables restoration of previous selection when applicable
+	virtual bool IsInputSelectionValidOnOutput() override;
+
 
 protected:
 	// flags used to identify modifier keys/buttons
@@ -323,6 +327,10 @@ protected:
 	int32 CurrentCurveTimestamp = 1;
 	void UndoCurrentOperation(const TArray<FVector3d>& ClickPointsIn, const TArray<FVector3d>& PolygonVerticesIn);
 	bool CheckInCurve(int32 Timestamp) const { return CurrentCurveTimestamp == Timestamp; }
+private:
+	// if the drawn path is incomplete upon 'Accept' of the tool, no mesh should be created and the
+	// previous mesh element selection from before entering the tool should be restored
+	bool bRestoreInputSelection = true;
 };
 
 
