MeshModelingTools: Ignore hits of the gizmo itself when using the drag alignment mechanic so that it can work properly with gizmo components that have collision geometry.

diff --git a/Engine/Plugins/Runtime/MeshModelingToolset/Source/ModelingComponents/Private/Mechanics/DragAlignmentMechanic.cpp b/Engine/Plugins/Runtime/MeshModelingToolset/Source/ModelingComponents/Private/Mechanics/DragAlignmentMechanic.cpp
--- a/Engine/Plugins/Runtime/MeshModelingToolset/Source/ModelingComponents/Private/Mechanics/DragAlignmentMechanic.cpp
+++ b/Engine/Plugins/Runtime/MeshModelingToolset/Source/ModelingComponents/Private/Mechanics/DragAlignmentMechanic.cpp
@@ -7,6 +7,7 @@
 #include "BaseGizmos/RepositionableTransformGizmo.h"
 #include "BaseGizmos/CombinedTransformGizmo.h"
 #include "BaseGizmos/TransformProxy.h"
+#include "Components/PrimitiveComponent.h"
 #include "Operations/GroupTopologyDeformer.h"
 #include "SceneManagement.h" //FPrimitiveDrawInterface
 #include "ToolDataVisualizer.h"
@@ -27,37 +28,50 @@ void FDragAlignmentBase::AddToGizmo(UCombinedTransformGizmo* TransformGizmo, con
 		return;
 	}
 
-	// If we have components to ignore/include, we need a copy of the array so that the alignment
-	// functions can use them later.
-	TSharedPtr<TArray<const UPrimitiveComponent*>> ComponentsToIgnorePersistent;
+	// We need a copy of the components to ignore, potentially to be shared across multiple lambdas. For pivot
+	//  adjustment, the ignored components will be just those of the gizmo itself, but for regular gizmo operation
+	//  we add ComponentsToIgnoreInAlignment.
+	TSharedPtr<TArray<const UPrimitiveComponent*>> GizmoComponentsToIgnore = MakeShared<TArray<const UPrimitiveComponent*>>();
+	if (ensure(TransformGizmo->GetGizmoActor()))
+	{
+		TransformGizmo->GetGizmoActor()->GetComponents<const UPrimitiveComponent*>(*GizmoComponentsToIgnore);
+	}
+
+	TSharedPtr<TArray<const UPrimitiveComponent*>> CombinedComponentsToIgnore = GizmoComponentsToIgnore;
 	if (ComponentsToIgnoreInAlignment)
 	{
-		ComponentsToIgnorePersistent = MakeShared<TArray<const UPrimitiveComponent*>>(*ComponentsToIgnoreInAlignment);
+		// If we had some extra components to ignore, we create a separate combined array so that
+		//  GizmoComponentsToIgnore continues to only include gizmo components for pivot alignment.
+		CombinedComponentsToIgnore = MakeShared<TArray<const UPrimitiveComponent*>>(*ComponentsToIgnoreInAlignment);
+		CombinedComponentsToIgnore->Append(*GizmoComponentsToIgnore);
 	}
-	TSharedPtr<TArray<const UPrimitiveComponent*>> ComponentsToIncludePersistent;
+
+	// Similarly, create an array to hold the extra component to include in raycasting.
+	TSharedPtr<TArray<const UPrimitiveComponent*>> ComponentsToInclude;
 	if (InvisibleComponentsToIncludeInAlignment)
 	{
-		ComponentsToIncludePersistent = MakeShared<TArray<const UPrimitiveComponent*>>(*InvisibleComponentsToIncludeInAlignment);
+		ComponentsToInclude = MakeShared<TArray<const UPrimitiveComponent*>>(*InvisibleComponentsToIncludeInAlignment);
 	}
 
+
 	// Set the alignment functions.
 	TransformGizmo->SetWorldAlignmentFunctions(
 		[this]() { return GetAlignmentModeEnabled(); },
-		[this, ComponentsToIgnorePersistent, ComponentsToIncludePersistent](const FRay& WorldRay, FVector& OutputPoint) {
-			return CastRay(WorldRay, OutputPoint, ComponentsToIgnorePersistent.Get(), ComponentsToIncludePersistent.Get(), true);
+		[this, CombinedComponentsToIgnore, ComponentsToInclude](const FRay& WorldRay, FVector& OutputPoint) {
+			return CastRay(WorldRay, OutputPoint, CombinedComponentsToIgnore.Get(), ComponentsToInclude.Get(), true);
 		});
 
-	// If we're adding to a repositionable gizmo, we probably don't want to ignore components
-	// when repositioning the pivot, since ignoring components is mainly used to avoid hitting
-	// the actual object, which we actually do want to hit in moving moving the pivot.
+	// If we're adding to a repositionable gizmo, we probably don't want to ignore components other
+	// than the gizmo components themselves, since ignoring components is mainly used to avoid hitting
+	// the moved object, which we actually do want to hit in moving moving the pivot.
 	// If we someday do want to ignore things in a pivot reposition function, we'll have to
 	// make this caller-configurable somehow.
 	if (URepositionableTransformGizmo* RepositionableGizmo = Cast<URepositionableTransformGizmo>(TransformGizmo))
 	{
 		RepositionableGizmo->SetPivotAlignmentFunctions(
 			[this]() { return GetAlignmentModeEnabled(); },
-			[this, ComponentsToIncludePersistent](const FRay& WorldRay, FVector& OutputPoint) {
-				return CastRay(WorldRay, OutputPoint, nullptr, ComponentsToIncludePersistent.Get(), false); // don't ignore anything
+			[this, GizmoComponentsToIgnore, ComponentsToInclude](const FRay& WorldRay, FVector& OutputPoint) {
+				return CastRay(WorldRay, OutputPoint, GizmoComponentsToIgnore.Get(), ComponentsToInclude.Get(), /*bUseFilter*/ false);
 			});
 	}
 
