ModelingComponents: Add support for the Actor creation API to create a default, empty actor when no template actor is supplied.

diff --git a/Engine/Plugins/Runtime/MeshModelingToolset/Source/ModelingComponentsEditorOnly/Private/EditorModelingObjectsCreationAPI.cpp b/Engine/Plugins/Runtime/MeshModelingToolset/Source/ModelingComponentsEditorOnly/Private/EditorModelingObjectsCreationAPI.cpp
--- a/Engine/Plugins/Runtime/MeshModelingToolset/Source/ModelingComponentsEditorOnly/Private/EditorModelingObjectsCreationAPI.cpp
+++ b/Engine/Plugins/Runtime/MeshModelingToolset/Source/ModelingComponentsEditorOnly/Private/EditorModelingObjectsCreationAPI.cpp
@@ -31,6 +31,7 @@
 #include "Components/DynamicMeshComponent.h"
 
 #include "ActorFactories/ActorFactory.h"
+#include "ActorFactories/ActorFactoryEmptyActor.h"
 #include "AssetSelection.h"
 
 #include UE_INLINE_GENERATED_CPP_BY_NAME(EditorModelingObjectsCreationAPI)
@@ -521,6 +522,24 @@ FCreateActorResult UEditorModelingObjectsCreationAPI::CreateNewActor(FCreateActo
 		}
 	}
 
+	// If we don't have an asset template, we're going to just create a blank, empty actor instance.
+	if (GEditor && !CreateActorParams.TemplateAsset)
+	{
+		UActorFactoryEmptyActor* EmptyActorFactory = NewObject<UActorFactoryEmptyActor>();
+		FAssetData AssetData(EmptyActorFactory->GetDefaultActorClass(FAssetData()));
+		FActorSpawnParameters SpawnParams;
+		SpawnParams.Name = FName(CreateActorParams.BaseName);
+		SpawnParams.NameMode = FActorSpawnParameters::ESpawnActorNameMode::Requested;
+		NewActor = EmptyActorFactory->CreateActor(AssetData.GetAsset(),
+			CreateActorParams.TargetWorld->GetCurrentLevel(),
+			CreateActorParams.Transform,
+			SpawnParams);
+
+		// This is also editor-only: it's the label that shows up in the hierarchy
+		FActorLabelUtilities::SetActorLabelUnique(NewActor, CreateActorParams.BaseName);
+	}
+
+
 	if (!NewActor)
 	{
 		return FCreateActorResult{ ECreateModelingObjectResult::Failed_ActorCreationFailed };
