ModelingComponents: Add new API for the creation of a component within an existing actor

diff --git a/Engine/Plugins/Runtime/MeshModelingToolset/Source/ModelingComponents/Private/ModelingObjectsCreationAPI.cpp b/Engine/Plugins/Runtime/MeshModelingToolset/Source/ModelingComponents/Private/ModelingObjectsCreationAPI.cpp
--- a/Engine/Plugins/Runtime/MeshModelingToolset/Source/ModelingComponents/Private/ModelingObjectsCreationAPI.cpp
+++ b/Engine/Plugins/Runtime/MeshModelingToolset/Source/ModelingComponents/Private/ModelingObjectsCreationAPI.cpp
@@ -153,7 +153,25 @@ FCreateActorResult UE::Modeling::CreateNewActor(UInteractiveToolManager* ToolMan
 
 
 
-
+FCreateComponentResult UE::Modeling::CreateNewComponentOnActor(UInteractiveToolManager* ToolManager, FCreateComponentParams&& CreateComponentParams)
+{
+	if (ensure(ToolManager))
+	{
+		UModelingObjectsCreationAPI* UseAPI = ToolManager->GetContextObjectStore()->FindContext<UModelingObjectsCreationAPI>();
+		if (UseAPI)
+		{
+			if (UseAPI->HasMoveVariants())
+			{
+				return UseAPI->CreateNewComponentOnActor(MoveTemp(CreateComponentParams));
+			}
+			else
+			{
+				return UseAPI->CreateNewComponentOnActor(CreateComponentParams);
+			}
+		}
+	}
+	return FCreateComponentResult{ ECreateModelingObjectResult::Failed_NoAPIFound };
+}
 
 
 
diff --git a/Engine/Plugins/Runtime/MeshModelingToolset/Source/ModelingComponents/Public/ModelingObjectsCreationAPI.h b/Engine/Plugins/Runtime/MeshModelingToolset/Source/ModelingComponents/Public/ModelingObjectsCreationAPI.h
--- a/Engine/Plugins/Runtime/MeshModelingToolset/Source/ModelingComponents/Public/ModelingObjectsCreationAPI.h
+++ b/Engine/Plugins/Runtime/MeshModelingToolset/Source/ModelingComponents/Public/ModelingObjectsCreationAPI.h
@@ -37,6 +37,7 @@ enum class ECreateModelingObjectResult : uint8
 	Failed_AssetCreationFailed,
 	Failed_ActorCreationFailed,
 	Failed_InvalidMaterial,
+	Failed_InvalidActor
 };
 
 /**
@@ -446,8 +447,55 @@ struct MODELINGCOMPONENTS_API FCreateActorResult
 	bool IsOK() const { return ResultCode == ECreateModelingObjectResult::Ok; }
 };
 
+/**
+ * FCreateComponentParams is a collection of input data intended to be passed to
+ * UModelingObjectsCreationAPI::CreateNewComponentOnActor().
+ */
+USTRUCT(BlueprintType)
+struct MODELINGCOMPONENTS_API FCreateComponentParams
+{
+	GENERATED_BODY()
+
+	/** An actor to host the new component */
+	UPROPERTY(Category = "CreateComponentParams", EditAnywhere)
+	TObjectPtr<AActor> HostActor = nullptr;
+
+	/** A component class to instantiate a new component instance of. */
+	UPROPERTY(Category = "CreateComponentParams", EditAnywhere)
+	TObjectPtr<UClass> ComponentClass;
+
+	/** The base name of the new Component */
+	UPROPERTY(Category = "CreateComponentParams", EditAnywhere)
+	FString BaseName;
+
+	/** If true, and if the component class is derived from USceneComponent, set this component as the root of the Actor. */
+	UPROPERTY(Category = "CreateComponentParams", EditAnywhere)
+	bool bSetAsRoot = false;
+
+	/** If true, wrap the creation operation in a transaction. */
+	UPROPERTY(Category = "CreateComponentParams", EditAnywhere)
+	bool bTransact;
+};
+
+/**
+ * FCreateComponentResult is returned by UModelingObjectsCreationAPI::CreateNewComponentOnActor()
+ * to indicate success/failure and provide information about created components
+ */
+USTRUCT(BlueprintType)
+struct MODELINGCOMPONENTS_API FCreateComponentResult
+{
+	GENERATED_BODY()
+
+	/** Success/Failure status for the requested operation */
+	UPROPERTY(Category = "CreateComponentResult", VisibleAnywhere)
+	ECreateModelingObjectResult ResultCode = ECreateModelingObjectResult::Ok;
 
+	/** A pointer to a newly-created Component */
+	UPROPERTY(Category = "CreateComponentResult", VisibleAnywhere)
+	TObjectPtr<UActorComponent> NewComponent = nullptr;
 
+	bool IsOK() const { return ResultCode == ECreateModelingObjectResult::Ok; }
+};
 
 
 /**
@@ -503,6 +551,13 @@ public:
 	UFUNCTION(BlueprintCallable, Category = "Modeling Objects")
 	virtual FCreateActorResult CreateNewActor(const FCreateActorParams& CreateActorParams) { return FCreateActorResult{ ECreateModelingObjectResult::Failed_Unknown }; }
 
+	/**
+	 * Create a new component on the specified actor of the requested class.
+	 * @return a results data structure, containing a result code and information about the new component created.
+	 */
+	UFUNCTION(BlueprintCallable, Category = "Modeling Objects")
+	virtual FCreateComponentResult CreateNewComponentOnActor(const FCreateComponentParams& CreateComponentParams) { return FCreateComponentResult{ ECreateModelingObjectResult::Failed_Unknown }; }
+
 	//
 	// Non-UFunction variants that support std::move operators
 	//
@@ -517,6 +572,7 @@ public:
 	virtual FCreateTextureObjectResult CreateTextureObject(FCreateTextureObjectParams&& CreateTexParams) { return FCreateTextureObjectResult{ ECreateModelingObjectResult::Failed_Unknown }; }
 	virtual FCreateMaterialObjectResult CreateMaterialObject(FCreateMaterialObjectParams&& CreateMaterialParams) { return FCreateMaterialObjectResult{ ECreateModelingObjectResult::Failed_Unknown }; }
 	virtual FCreateActorResult CreateNewActor(FCreateActorParams&& CreateActorParams) { return FCreateActorResult{ ECreateModelingObjectResult::Failed_Unknown }; }
+	virtual FCreateComponentResult CreateNewComponentOnActor(FCreateComponentParams&& CreateComponentParams) { return FCreateComponentResult{ ECreateModelingObjectResult::Failed_Unknown }; }
 
 };
 
@@ -564,7 +620,13 @@ MODELINGCOMPONENTS_API FCreateMaterialObjectResult CreateMaterialObject(UInterac
  */
 MODELINGCOMPONENTS_API FCreateActorResult CreateNewActor(UInteractiveToolManager* ToolManager, FCreateActorParams&& CreateActorParams);
 
-
+/**
+ * Create a new component on an actor based on the data in CreateComponentParams.
+ * This is a convenience function that will try to locate a UModelingObjectsCreationAPI instance in the ToolManager's ContextObjectStore,
+ * and then call UModelingObjectsCreationAPI::CreateNewComponentOnActor()
+ * @return a results data structure, containing a result code and information about any new objects created
+ */
+MODELINGCOMPONENTS_API FCreateComponentResult CreateNewComponentOnActor(UInteractiveToolManager* ToolManager, FCreateComponentParams&& CreateComponentParams);
 
 /**
  * Strip the appended auto-generated _UUID suffix on the given string, if we can detect one
diff --git a/Engine/Plugins/Runtime/MeshModelingToolset/Source/ModelingComponentsEditorOnly/Private/EditorModelingObjectsCreationAPI.cpp b/Engine/Plugins/Runtime/MeshModelingToolset/Source/ModelingComponentsEditorOnly/Private/EditorModelingObjectsCreationAPI.cpp
--- a/Engine/Plugins/Runtime/MeshModelingToolset/Source/ModelingComponentsEditorOnly/Private/EditorModelingObjectsCreationAPI.cpp
+++ b/Engine/Plugins/Runtime/MeshModelingToolset/Source/ModelingComponentsEditorOnly/Private/EditorModelingObjectsCreationAPI.cpp
@@ -32,6 +32,7 @@
 
 #include "ActorFactories/ActorFactory.h"
 #include "ActorFactories/ActorFactoryEmptyActor.h"
+#include "Kismet2/ComponentEditorUtils.h"
 #include "AssetSelection.h"
 
 #include UE_INLINE_GENERATED_CPP_BY_NAME(EditorModelingObjectsCreationAPI)
@@ -113,6 +114,11 @@ FCreateActorResult UEditorModelingObjectsCreationAPI::CreateNewActor(const FCrea
 	return CreateNewActor(MoveTemp(LocalParams));
 }
 
+FCreateComponentResult UEditorModelingObjectsCreationAPI::CreateNewComponentOnActor(const FCreateComponentParams& CreateComponentParams)
+{
+	FCreateComponentParams LocalParams = CreateComponentParams;
+	return CreateNewComponentOnActor(MoveTemp(LocalParams));		
+}
 
 
 FCreateMeshObjectResult UEditorModelingObjectsCreationAPI::CreateMeshObject(FCreateMeshObjectParams&& CreateMeshParams)
@@ -554,7 +560,59 @@ FCreateActorResult UEditorModelingObjectsCreationAPI::CreateNewActor(FCreateActo
 	return ResultOut;
 }
 
+FCreateComponentResult UEditorModelingObjectsCreationAPI::CreateNewComponentOnActor(FCreateComponentParams&& CreateComponentParams)
+{
+	FCreateComponentResult Result;
+
+	if (!ensure(CreateComponentParams.HostActor))
+	{
+		Result.ResultCode = ECreateModelingObjectResult::Failed_InvalidActor;
+		return Result;
+	}
+
+	if (CreateComponentParams.bTransact)
+	{
+		CreateComponentParams.HostActor->Modify();
+	}
+
+	FName NewComponentName = *FComponentEditorUtils::GenerateValidVariableName(CreateComponentParams.ComponentClass, CreateComponentParams.HostActor);
+
+	if (!CreateComponentParams.BaseName.IsEmpty())
+	{
+		NewComponentName = FName(CreateComponentParams.BaseName);
+	}
 
+	// Note that the RF_Transactional is important here for the spline to undo/redo properly in the future
+	UActorComponent* Component = NewObject<UActorComponent>(CreateComponentParams.HostActor, CreateComponentParams.ComponentClass, NewComponentName,
+		CreateComponentParams.bTransact ? RF_Transactional : RF_NoFlags);
+	
+	USceneComponent* SceneComponent = StaticCast<USceneComponent*>(Component);
+	if (SceneComponent) {
+		if (CreateComponentParams.bSetAsRoot)
+		{
+			CreateComponentParams.HostActor->SetRootComponent(SceneComponent);
+		}
+		else
+		{
+			SceneComponent->SetupAttachment(CreateComponentParams.HostActor->GetRootComponent());
+		}
+	}
+
+	Component->OnComponentCreated();
+	CreateComponentParams.HostActor->AddInstanceComponent(Component);
+	Component->RegisterComponent();
+	if (SceneComponent) {
+		SceneComponent->ResetRelativeTransform();
+	}
+	CreateComponentParams.HostActor->PostEditChange();
+
+	Result.ResultCode = ECreateModelingObjectResult::Ok;
+	Result.NewComponent = Component;
+
+	OnModelingComponentCreated.Broadcast(Result);
+
+	return Result;
+}
 
 ECreateModelingObjectResult UEditorModelingObjectsCreationAPI::GetNewAssetPath(
 	FString& OutNewAssetPath,
diff --git a/Engine/Plugins/Runtime/MeshModelingToolset/Source/ModelingComponentsEditorOnly/Public/EditorModelingObjectsCreationAPI.h b/Engine/Plugins/Runtime/MeshModelingToolset/Source/ModelingComponentsEditorOnly/Public/EditorModelingObjectsCreationAPI.h
--- a/Engine/Plugins/Runtime/MeshModelingToolset/Source/ModelingComponentsEditorOnly/Public/EditorModelingObjectsCreationAPI.h
+++ b/Engine/Plugins/Runtime/MeshModelingToolset/Source/ModelingComponentsEditorOnly/Public/EditorModelingObjectsCreationAPI.h
@@ -43,6 +43,10 @@ public:
 	// UFUNCTION(BlueprintCallable, Category = "Modeling Objects")
 	virtual FCreateActorResult CreateNewActor(const FCreateActorParams& CreateActorParams) override;
 
+	// UFUNCTION(BlueprintCallable, Category = "Modeling Objects")
+	virtual FCreateComponentResult CreateNewComponentOnActor(const FCreateComponentParams& CreateComponentParams) override;
+
+
 	//
 	// Non-UFunction variants that support std::move operators
 	//
@@ -57,6 +61,7 @@ public:
 
 	virtual FCreateActorResult CreateNewActor(FCreateActorParams&& CreateActorParams) override;
 
+	virtual FCreateComponentResult CreateNewComponentOnActor(FCreateComponentParams&& CreateComponentParams) override;
 
 	//
 	// Callbacks that editor can hook into to handle asset creation
@@ -77,6 +82,9 @@ public:
 	DECLARE_MULTICAST_DELEGATE_OneParam(FModelingActorCreatedSignature, const FCreateActorResult& CreatedInfo);
 	FModelingActorCreatedSignature OnModelingActorCreated;
 
+	DECLARE_MULTICAST_DELEGATE_OneParam(FModelingComponentCreatedSignature, const FCreateComponentResult& CreatedInfo);
+	FModelingComponentCreatedSignature OnModelingComponentCreated;
+
 	//
 	// Utility functions to handle registration/unregistration
 	//
