ModelingTools: Promote experimental modeling tools to beta.
Removed SkeletalMeshModelingTools dependency on MeshModelingToolsExp.

diff --git a/Engine/Plugins/Experimental/Animation/SkeletalMeshModelingTools/Source/SkeletalMeshModelingTools/SkeletalMeshModelingTools.Build.cs b/Engine/Plugins/Experimental/Animation/SkeletalMeshModelingTools/Source/SkeletalMeshModelingTools/SkeletalMeshModelingTools.Build.cs
--- a/Engine/Plugins/Experimental/Animation/SkeletalMeshModelingTools/Source/SkeletalMeshModelingTools/SkeletalMeshModelingTools.Build.cs
+++ b/Engine/Plugins/Experimental/Animation/SkeletalMeshModelingTools/Source/SkeletalMeshModelingTools/SkeletalMeshModelingTools.Build.cs
@@ -19,8 +19,6 @@ public class SkeletalMeshModelingTools : ModuleRules
 				"InteractiveToolsFramework",
 				"MeshModelingTools",
 				"MeshModelingToolsEditorOnly",
-				"MeshModelingToolsEditorOnlyExp",
-				"MeshModelingToolsExp",
 				"ModelingComponentsEditorOnly",
 				"ModelingToolsEditorMode",
 				"SkeletalMeshEditor",
diff --git a/Engine/Plugins/Experimental/CharacterFXEditor/ExampleCharacterFXEditor/ExampleCharacterFXEditor.uplugin b/Engine/Plugins/Experimental/CharacterFXEditor/ExampleCharacterFXEditor/ExampleCharacterFXEditor.uplugin
--- a/Engine/Plugins/Experimental/CharacterFXEditor/ExampleCharacterFXEditor/ExampleCharacterFXEditor.uplugin
+++ b/Engine/Plugins/Experimental/CharacterFXEditor/ExampleCharacterFXEditor/ExampleCharacterFXEditor.uplugin
@@ -26,6 +26,10 @@
 			"Name": "BaseCharacterFXEditor",
 			"Enabled": true
 		},
+		{
+			"Name": "MeshModelingToolset",
+			"Enabled": true
+		},
 		{
 			"Name": "MeshModelingToolsetExp",
 			"Enabled": true
diff --git a/Engine/Plugins/Experimental/CharacterFXEditor/ExampleCharacterFXEditor/Source/ExampleCharacterFXEditor/ExampleCharacterFXEditor.Build.cs b/Engine/Plugins/Experimental/CharacterFXEditor/ExampleCharacterFXEditor/Source/ExampleCharacterFXEditor/ExampleCharacterFXEditor.Build.cs
--- a/Engine/Plugins/Experimental/CharacterFXEditor/ExampleCharacterFXEditor/Source/ExampleCharacterFXEditor/ExampleCharacterFXEditor.Build.cs
+++ b/Engine/Plugins/Experimental/CharacterFXEditor/ExampleCharacterFXEditor/Source/ExampleCharacterFXEditor/ExampleCharacterFXEditor.Build.cs
@@ -50,6 +50,7 @@ public class ExampleCharacterFXEditor : ModuleRules
 				"GeometryFramework",
 				"InputCore",
 				"InteractiveToolsFramework",
+				"MeshModelingToolsEditorOnly",
 				"MeshModelingToolsEditorOnlyExp",
 				"ModelingComponentsEditorOnly",
 				"ModelingComponents",
diff --git a/Engine/Plugins/Experimental/Dataflow/Source/DataflowEditor/DataflowEditor.Build.cs b/Engine/Plugins/Experimental/Dataflow/Source/DataflowEditor/DataflowEditor.Build.cs
--- a/Engine/Plugins/Experimental/Dataflow/Source/DataflowEditor/DataflowEditor.Build.cs
+++ b/Engine/Plugins/Experimental/Dataflow/Source/DataflowEditor/DataflowEditor.Build.cs
@@ -48,6 +48,8 @@ namespace UnrealBuildTool.Rules
 					"LevelEditor",
 					"MeshDescription",
 					"MeshConversion",
+					"MeshModelingTools",
+					"MeshModelingToolsEditorOnly",
 					"MeshModelingToolsEditorOnlyExp",
 					"MeshModelingToolsExp",
 					"ModelingComponentsEditorOnly",
diff --git a/Engine/Plugins/Experimental/Enterprise/DataprepGeometryOperations/DataprepGeometryOperations.uplugin b/Engine/Plugins/Experimental/Enterprise/DataprepGeometryOperations/DataprepGeometryOperations.uplugin
--- a/Engine/Plugins/Experimental/Enterprise/DataprepGeometryOperations/DataprepGeometryOperations.uplugin
+++ b/Engine/Plugins/Experimental/Enterprise/DataprepGeometryOperations/DataprepGeometryOperations.uplugin
@@ -31,6 +31,10 @@
 			"Name": "EditorScriptingUtilities",
 			"Enabled": true
 		},
+		{
+			"Name": "MeshModelingToolset",
+			"Enabled": true
+		},
 		{
 			"Name": "MeshModelingToolsetExp",
 			"Enabled": true
diff --git a/Engine/Plugins/Experimental/MeshModelingToolsetExp/MeshModelingToolsetExp.uplugin b/Engine/Plugins/Experimental/MeshModelingToolsetExp/MeshModelingToolsetExp.uplugin
--- a/Engine/Plugins/Experimental/MeshModelingToolsetExp/MeshModelingToolsetExp.uplugin
+++ b/Engine/Plugins/Experimental/MeshModelingToolsetExp/MeshModelingToolsetExp.uplugin
@@ -40,11 +40,6 @@
 			"Name": "ModelingUI",
 			"Type": "Runtime",
 			"LoadingPhase": "Default"
-		},
-		{
-			"Name": "SkeletalMeshModifiers",
-			"Type": "Editor",
-			"LoadingPhase": "Default"
 		}
 	],
 	"Plugins": [
diff --git a/Engine/Plugins/Experimental/MeshModelingToolsetExp/Source/MeshModelingToolsEditorOnlyExp/MeshModelingToolsEditorOnlyExp.Build.cs b/Engine/Plugins/Experimental/MeshModelingToolsetExp/Source/MeshModelingToolsEditorOnlyExp/MeshModelingToolsEditorOnlyExp.Build.cs
--- a/Engine/Plugins/Experimental/MeshModelingToolsetExp/Source/MeshModelingToolsEditorOnlyExp/MeshModelingToolsEditorOnlyExp.Build.cs
+++ b/Engine/Plugins/Experimental/MeshModelingToolsetExp/Source/MeshModelingToolsEditorOnlyExp/MeshModelingToolsEditorOnlyExp.Build.cs
@@ -40,8 +40,6 @@ public class MeshModelingToolsEditorOnlyExp : ModuleRules
 				"ModelingComponentsEditorOnly",
 				"ModelingOperators",
 				"ModelingOperatorsEditorOnly",
-				"SkeletalMeshUtilitiesCommon",
-				"SkeletalMeshModifiers"
 				// ... add other public dependencies that you statically link with here ...
 			}
 		);
diff --git a/Engine/Plugins/Experimental/MeshModelingToolsetExp/Source/MeshModelingToolsEditorOnlyExp/Private/MeshModelingToolsEditorOnlyExp.cpp b/Engine/Plugins/Experimental/MeshModelingToolsetExp/Source/MeshModelingToolsEditorOnlyExp/Private/MeshModelingToolsEditorOnlyExp.cpp
--- a/Engine/Plugins/Experimental/MeshModelingToolsetExp/Source/MeshModelingToolsEditorOnlyExp/Private/MeshModelingToolsEditorOnlyExp.cpp
+++ b/Engine/Plugins/Experimental/MeshModelingToolsetExp/Source/MeshModelingToolsEditorOnlyExp/Private/MeshModelingToolsEditorOnlyExp.cpp
@@ -4,8 +4,6 @@
 
 #define LOCTEXT_NAMESPACE "FMeshModelingToolsEditorOnlyExpModule"
 
-DEFINE_LOG_CATEGORY(LogMeshModelingToolsEditor);
-
 void FMeshModelingToolsEditorOnlyExpModule::StartupModule()
 {
 }
diff --git a/Engine/Plugins/Experimental/MeshModelingToolsetExp/Source/MeshModelingToolsEditorOnlyExp/Public/MeshModelingToolsEditorOnlyExp.h b/Engine/Plugins/Experimental/MeshModelingToolsetExp/Source/MeshModelingToolsEditorOnlyExp/Public/MeshModelingToolsEditorOnlyExp.h
--- a/Engine/Plugins/Experimental/MeshModelingToolsetExp/Source/MeshModelingToolsEditorOnlyExp/Public/MeshModelingToolsEditorOnlyExp.h
+++ b/Engine/Plugins/Experimental/MeshModelingToolsetExp/Source/MeshModelingToolsEditorOnlyExp/Public/MeshModelingToolsEditorOnlyExp.h
@@ -5,8 +5,6 @@
 #include "CoreMinimal.h"
 #include "Modules/ModuleManager.h"
 
-DECLARE_LOG_CATEGORY_EXTERN(LogMeshModelingToolsEditor, Log, All);
-
 class FMeshModelingToolsEditorOnlyExpModule : public IModuleInterface
 {
 public:
diff --git a/Engine/Plugins/Runtime/MeshModelingToolset/MeshModelingToolset.uplugin b/Engine/Plugins/Runtime/MeshModelingToolset/MeshModelingToolset.uplugin
--- a/Engine/Plugins/Runtime/MeshModelingToolset/MeshModelingToolset.uplugin
+++ b/Engine/Plugins/Runtime/MeshModelingToolset/MeshModelingToolset.uplugin
@@ -44,6 +44,11 @@
 			"Name": "ModelingOperatorsEditorOnly",
 			"Type": "Editor",
 			"LoadingPhase": "Default"
+		},
+		{
+			"Name": "SkeletalMeshModifiers",
+			"Type": "Editor",
+			"LoadingPhase": "Default"
 		}
 	],
 	"Plugins": [
diff --git a/Engine/Plugins/Runtime/MeshModelingToolset/Source/MeshModelingTools/Public/MeshAttributePaintTool.h b/Engine/Plugins/Runtime/MeshModelingToolset/Source/MeshModelingTools/Public/MeshAttributePaintTool.h
new file mode 100644
--- /dev/null
+++ b/Engine/Plugins/Runtime/MeshModelingToolset/Source/MeshModelingTools/Public/MeshAttributePaintTool.h
@@ -0,0 +1,350 @@
+// Copyright Epic Games, Inc. All Rights Reserved.
+
+#pragma once
+
+#include "DynamicMeshBrushTool.h"
+#include "BaseTools/MeshSurfacePointMeshEditingTool.h"
+#include "Changes/ValueWatcher.h"
+#include "Changes/IndexedAttributeChange.h"
+#include "DynamicMesh/DynamicVerticesOctree3.h"
+#include "MeshDescription.h"
+#include "Selections/GeometrySelection.h"
+
+#include "MeshAttributePaintTool.generated.h"
+
+struct FMeshDescription;
+class UMeshAttributePaintTool;
+
+/**
+ * Maps float values to linear color ramp.
+ */
+class FFloatAttributeColorMapper
+{
+public:
+	virtual ~FFloatAttributeColorMapper() {}
+
+	FLinearColor LowColor = FLinearColor(0.9f, 0.9f, 0.9f, 1.0f);
+	FLinearColor HighColor = FLinearColor(1.0f, 0.0f, 0.0f, 1.0f);
+
+	virtual FLinearColor ToColor(float Value)
+	{
+		float t = FMath::Clamp(Value, 0.0f, 1.0f);
+		return FLinearColor(
+			FMathf::Lerp(LowColor.R, HighColor.R, t),
+			FMathf::Lerp(LowColor.G, HighColor.G, t),
+			FMathf::Lerp(LowColor.B, HighColor.B, t),
+			1.0f);
+	}
+
+	template<typename VectorType>
+	VectorType ToColor(float Value)
+	{
+		FLinearColor Color = ToColor(Value);
+		return VectorType(Color.R, Color.G, Color.B);
+	}
+};
+
+
+/**
+ * Abstract interface to a single-channel indexed floating-point attribute
+ */
+class IMeshVertexAttributeAdapter
+{
+public:
+	virtual ~IMeshVertexAttributeAdapter() {}
+
+	virtual int ElementNum() const = 0;
+	virtual float GetValue(int32 Index) const = 0;
+	virtual void SetValue(int32 Index, float Value) = 0;
+	virtual UE::Geometry::FInterval1f GetValueRange() = 0;
+};
+
+
+
+/**
+ * Abstract interface to a set of single-channel indexed floating-point attributes
+ */
+class IMeshVertexAttributeSource
+{
+public:
+	virtual ~IMeshVertexAttributeSource() {}
+
+	virtual TArray<FName> GetAttributeList() = 0;
+	virtual TUniquePtr<IMeshVertexAttributeAdapter> GetAttribute(FName AttributeName) = 0;
+	/** @return number of indices in each attribute */
+	virtual int32 GetAttributeElementNum() = 0;
+};
+
+
+
+
+
+
+/**
+ * Tool Builder for Attribute Paint Tool
+ */
+UCLASS()
+class MESHMODELINGTOOLS_API UMeshAttributePaintToolBuilder : public UMeshSurfacePointMeshEditingToolBuilder
+{
+	GENERATED_BODY()
+public:
+	/** Optional color map customization */
+	TUniqueFunction<TUniquePtr<FFloatAttributeColorMapper>()> ColorMapFactory;
+
+	virtual UMeshSurfacePointTool* CreateNewTool(const FToolBuilderState& SceneState) const override;
+	virtual void InitializeNewTool(UMeshSurfacePointTool* NewTool, const FToolBuilderState& SceneState) const override;
+	virtual bool CanBuildTool(const FToolBuilderState& SceneState) const override;
+
+protected:
+	virtual const FToolTargetTypeRequirements& GetTargetRequirements() const override;
+};
+
+
+
+UENUM()
+enum class EBrushActionMode
+{
+	Paint,
+	FloodFill
+};
+
+
+/**
+ * Selected-Attribute settings Attribute Paint Tool
+ */
+UCLASS()
+class MESHMODELINGTOOLS_API UMeshAttributePaintBrushOperationProperties : public UInteractiveToolPropertySet
+{
+	GENERATED_BODY()
+public:
+	/** 
+	 * If the tool was started with a mesh element selection, this setting hides everything
+	 *  except that selection, to make painting it easier. Requires that a mesh element
+	 *  selection exist on tool start.
+	 */
+	UPROPERTY(EditAnywhere, Category = Attribute, meta = (EditCondition = "bToolHasSelection", HideEditConditionToggle))
+	bool bIsolateGeometrySelection = false;
+	
+	//~ For the tool to set, to enable/disable bIsolateGeometrySelection
+	UPROPERTY()
+	bool bToolHasSelection = false;
+
+	UPROPERTY(EditAnywhere, Category = Attribute)
+	EBrushActionMode BrushAction = EBrushActionMode::Paint;
+
+};
+
+
+
+
+
+UCLASS()
+class MESHMODELINGTOOLS_API UMeshAttributePaintToolProperties : public UInteractiveToolPropertySet
+{
+	GENERATED_BODY()
+public:
+
+	UPROPERTY(EditAnywhere, Category = "Attribute", meta = (DisplayName = "Selected Attribute", GetOptions = GetAttributeNames))
+	FString Attribute;
+
+	UFUNCTION()
+	const TArray<FString>& GetAttributeNames() { return Attributes; };
+
+	TArray<FString> Attributes;
+
+public:
+	/**
+	* Initialize the internal array of attribute names
+	* @param bInitialize if set, selected Attribute will be reset to the first attribute or empty if there are none.
+	*/
+	void Initialize(const TArray<FName>& AttributeNames, bool bInitialize = false);
+
+	/**
+	 * Verify that the attribute selection is valid
+	 * @param bUpdateIfInvalid if selection is not valid, use attribute at index 0 or empty if there are no attributes
+	 * @return true if selection is in the Attributes array
+	 */
+	bool ValidateSelectedAttribute(bool bUpdateIfInvalid);
+
+	/**
+	 * @return selected attribute index, or -1 if invalid selection
+	 */
+	int32 GetSelectedAttributeIndex();
+};
+
+
+
+
+
+UENUM()
+enum class EMeshAttributePaintToolActions
+{
+	NoAction
+};
+
+
+
+UCLASS()
+class MESHMODELINGTOOLS_API UMeshAttributePaintEditActions : public UInteractiveToolPropertySet
+{
+	GENERATED_BODY()
+
+public:
+	TWeakObjectPtr<UMeshAttributePaintTool> ParentTool;
+
+	void Initialize(UMeshAttributePaintTool* ParentToolIn) { ParentTool = ParentToolIn; }
+
+	void PostAction(EMeshAttributePaintToolActions Action);
+};
+
+
+
+/**
+ * FCommandChange for color map changes
+ */
+class MESHMODELINGTOOLS_API FMeshAttributePaintChange : public TCustomIndexedValuesChange<float, int32>
+{
+public:
+	virtual FString ToString() const override
+	{
+		return FString(TEXT("Paint Attribute"));
+	}
+};
+
+
+
+/**
+ * UMeshAttributePaintTool paints single-channel float attributes on a MeshDescription.
+ * 
+ */
+UCLASS()
+class MESHMODELINGTOOLS_API UMeshAttributePaintTool : public UDynamicMeshBrushTool, public IInteractiveToolManageGeometrySelectionAPI
+{
+	GENERATED_BODY()
+
+public:
+	virtual void SetWorld(UWorld* World);
+	void SetGeometrySelection(const UE::Geometry::FGeometrySelection& SelectionIn);
+
+	virtual void RegisterActions(FInteractiveToolActionSet& ActionSet) override;
+
+	virtual void Setup() override;
+	virtual void OnTick(float DeltaTime) override;
+
+	virtual bool HasCancel() const override { return true; }
+	virtual bool HasAccept() const override { return true; }
+
+	// UBaseBrushTool overrides
+	virtual void OnBeginDrag(const FRay& Ray) override;
+	virtual void OnUpdateDrag(const FRay& Ray) override;
+	virtual void OnEndDrag(const FRay& Ray) override;
+	virtual bool OnUpdateHover(const FInputDeviceRay& DevicePos) override;
+	
+	// UMeshSurfacePointTool
+	virtual bool HitTest(const FRay& Ray, FHitResult& OutHit) override;
+
+	virtual void RequestAction(EMeshAttributePaintToolActions ActionType);
+
+	virtual void SetColorMap(TUniquePtr<FFloatAttributeColorMapper> ColorMap);
+
+	// IInteractiveToolManageGeometrySelectionAPI -- this tool won't update external geometry selection or change selection-relevant mesh IDs
+	virtual bool IsInputSelectionValidOnOutput() override
+	{
+		return true;
+	}
+
+	UE::Geometry::TDynamicVerticesOctree3<FDynamicMesh3>& GetVerticesOctree() { return VerticesOctree; }
+
+protected:
+	virtual void ApplyStamp(const FBrushStampData& Stamp);
+
+
+	struct FStampActionData
+	{
+		TArray<int32> ROIVertices;
+		TArray<float> ROIBefore;
+		TArray<float> ROIAfter;
+	};
+
+
+	virtual void ApplyStamp_Paint(const FBrushStampData& Stamp, FStampActionData& ActionData);
+	virtual void ApplyStamp_FloodFill(const FBrushStampData& Stamp, FStampActionData& ActionData);
+
+	virtual void OnShutdown(EToolShutdownType ShutdownType) override;
+	
+
+protected:
+	UPROPERTY()
+	TObjectPtr<UMeshAttributePaintBrushOperationProperties> BrushActionProps;
+
+	UPROPERTY()
+	TObjectPtr<UMeshAttributePaintToolProperties> AttribProps;
+
+	TValueWatcher<int32> SelectedAttributeWatcher;
+
+	//UPROPERTY()
+	//UMeshAttributePaintEditActions* AttributeEditActions;
+
+protected:
+	UWorld* TargetWorld;
+
+	bool bInRemoveStroke = false;
+	bool bInSmoothStroke = false;
+	FBrushStampData StartStamp;
+	FBrushStampData LastStamp;
+	bool bStampPending;
+
+	TUniquePtr<FMeshDescription> EditedMesh;
+
+	double CalculateBrushFalloff(double Distance);
+	UE::Geometry::TDynamicVerticesOctree3<FDynamicMesh3> VerticesOctree;
+	TArray<int> PreviewBrushROI;
+	void CalculateVertexROI(const FBrushStampData& Stamp, TArray<int>& VertexROI);
+
+	TUniquePtr<FFloatAttributeColorMapper> ColorMapper;
+	TUniquePtr<IMeshVertexAttributeSource> AttributeSource;
+
+	struct FAttributeData
+	{
+		FName Name;
+		TUniquePtr<IMeshVertexAttributeAdapter> Attribute;
+		TArray<float> CurrentValues;
+		TArray<float> InitialValues;
+	};
+	TArray<FAttributeData> Attributes;
+	int32 AttributeBufferCount;
+	int32 CurrentAttributeIndex;
+	UE::Geometry::FInterval1f CurrentValueRange;
+
+	// actions
+
+	bool bHavePendingAction = false;
+	EMeshAttributePaintToolActions PendingAction;
+	virtual void ApplyAction(EMeshAttributePaintToolActions ActionType);
+
+	bool bVisibleAttributeValid = false;
+	int32 PendingNewSelectedIndex = -1;
+	
+	void InitializeAttributes();
+	void StoreCurrentAttribute();
+	void UpdateVisibleAttribute();
+	void UpdateSelectedAttribute(int32 NewSelectedIndex);
+
+	TUniquePtr<TIndexedValuesChangeBuilder<float, FMeshAttributePaintChange>> ActiveChangeBuilder;
+	void BeginChange();
+	TUniquePtr<FMeshAttributePaintChange> EndChange();
+	void ExternalUpdateValues(int32 AttribIndex, const TArray<int32>& VertexIndices, const TArray<float>& NewValues);
+
+private:
+	TOptional<UE::Geometry::FGeometrySelection> GeometrySelection;
+	TSet<int32> SelectionTids;
+	TSet<int32> SelectionVids;
+	bool ShouldFilterTriangles() const;
+};
+
+
+
+
+
+
+
diff --git a/Engine/Plugins/Runtime/MeshModelingToolset/Source/MeshModelingTools/Public/Sculpting/MeshSculptToolBase.h b/Engine/Plugins/Runtime/MeshModelingToolset/Source/MeshModelingTools/Public/Sculpting/MeshSculptToolBase.h
new file mode 100644
--- /dev/null
+++ b/Engine/Plugins/Runtime/MeshModelingToolset/Source/MeshModelingTools/Public/Sculpting/MeshSculptToolBase.h
@@ -0,0 +1,616 @@
+// Copyright Epic Games, Inc. All Rights Reserved.
+
+#pragma once
+
+#include "CoreMinimal.h"
+#include "BaseTools/MeshSurfacePointTool.h"
+#include "BaseTools/BaseBrushTool.h"
+#include "Sculpting/MeshBrushOpBase.h"
+#include "BoxTypes.h"
+#include "Properties/MeshMaterialProperties.h"
+#include "Changes/ValueWatcher.h"
+#include "Components/BaseDynamicMeshComponent.h"
+#include "TransactionUtil.h"
+#include "MeshSculptToolBase.generated.h"
+
+
+class ULocalTwoAxisPropertyEditInputBehavior;
+class UMaterialInstanceDynamic;
+class UCombinedTransformGizmo;
+class UTransformProxy;
+class UPreviewMesh;
+
+/**
+ * Type of Brush Size currently active in FBrushToolRadius
+ */
+UENUM()
+enum class EBrushToolSizeType : uint8
+{
+	/** Brush size is a dimensionless scale relative to the target object size */
+	Adaptive = 0,
+	/** Brush size is defined in world dimensions */
+	World = 1
+};
+
+/**
+ * FBrushToolRadius is used to define the size of 3D "brushes" used in (eg) sculpting tools.
+ * The brush size can be defined in various ways.
+ */
+USTRUCT()
+struct MESHMODELINGTOOLS_API FBrushToolRadius
+{
+	GENERATED_BODY()
+
+	/** Specify the type of brush size currently in use */
+	UPROPERTY(EditAnywhere, Category = Brush)
+	EBrushToolSizeType SizeType = EBrushToolSizeType::Adaptive;
+
+	/** Adaptive brush size is used to interpolate between an object-specific minimum and maximum brush size */
+	UPROPERTY(EditAnywhere, Category = Brush, AdvancedDisplay, meta = (EditCondition = "SizeType == EBrushToolSizeType::Adaptive",
+		DisplayName = "Size", UIMin = "0.0", UIMax = "1.0", ClampMin = "0.0", ClampMax = "10.0"))
+	float AdaptiveSize = 0.25;
+
+	/** World brush size is a dimension in world coordinates */
+	UPROPERTY(EditAnywhere, Category = Brush, AdvancedDisplay, meta = (EditCondition = "SizeType == EBrushToolSizeType::World",
+		DisplayName = "World Radius", UIMin = "1.0", UIMax = "1000.0", ClampMin = "0.1", ClampMax = "50000.0"))
+	float WorldRadius = 100.0;
+
+	/**
+	 * WorldSizeRange defines the min/max dimensions for Adaptive brush size
+	 */
+	TInterval<float> WorldSizeRange = TInterval<float>(1.0f, 1000.0f);
+
+	//
+	// util functions
+	//
+
+	/** Set the WorldSizeRange value and optionally clamp the WorldRadius based on this new range */
+	void InitializeWorldSizeRange(TInterval<float> Range, bool bValidateWorldRadius = true);
+	/** Return the set/calculated world-space radius for the current settings */
+	float GetWorldRadius() const;
+	/** Increase the current radius dimension by a fixed step (or a smaller fixed step) */
+	void IncreaseRadius(bool bSmallStep);
+	/** Decrease the current radius dimension by a fixed step (or a smaller fixed step) */
+	void DecreaseRadius(bool bSmallStep);
+};
+
+ 
+/** Mesh Sculpting Brush Types */
+UENUM()
+enum class EMeshSculptFalloffType : uint8
+{
+	Smooth = 0,
+	Linear = 1,
+	Inverse = 2,
+	Round = 3,
+	BoxSmooth = 4,
+	BoxLinear = 5,
+	BoxInverse = 6,
+	BoxRound = 7,
+
+	LastValue UMETA(Hidden)
+};
+
+
+
+
+
+UCLASS()
+class MESHMODELINGTOOLS_API USculptBrushProperties : public UInteractiveToolPropertySet
+{
+	GENERATED_BODY()
+public:
+
+	UPROPERTY(EditAnywhere, Category = Brush, meta = (DisplayPriority = 1, ModelingQuickEdit))
+	FBrushToolRadius BrushSize;
+
+	/** Amount of falloff to apply (0.0 - 1.0) */
+	UPROPERTY(EditAnywhere, Category = Brush, meta = (DisplayName = "Falloff", UIMin = "0.0", UIMax = "1.0", ClampMin = "0.0", ClampMax = "1.0", HideEditConditionToggle, EditConditionHides, EditCondition="bShowFalloff", DisplayPriority = 3))
+	float BrushFalloffAmount;
+
+	/** If false, then BrushFalloffAmount will not be shown in DetailsView panels (otherwise no effect) */
+	UPROPERTY(meta = (TransientToolProperty))
+	bool bShowFalloff = true;
+
+	/** Depth of Brush into surface along view ray or surface normal, depending on the Active Brush Type */
+	UPROPERTY(EditAnywhere, Category = Brush, meta = (UIMin = "-0.5", UIMax = "0.5", ClampMin = "-1.0", ClampMax = "1.0", DisplayPriority = 5, HideEditConditionToggle, EditConditionHides, EditCondition = "bShowPerBrushProps"))
+	float Depth = 0;
+
+	/** Allow the Brush to hit the back-side of the mesh */
+	UPROPERTY(EditAnywhere, Category = Brush, meta = (DisplayPriority = 6))
+	bool bHitBackFaces = true;
+
+	/** Brush stamps are applied at this time interval. 0 for a single stamp, 1 for continuous stamps, 0.5 is a stamp every half-second */
+	UPROPERTY(EditAnywhere, Category = Brush, meta = (DisplayName = "Flow", UIMin = "0.0", UIMax = "1.0", ClampMin = "0.0", ClampMax = "1.0", DisplayPriority = 7, HideEditConditionToggle, EditConditionHides, EditCondition = "bShowFlowRate"))
+	float FlowRate = 1.0f;
+
+	/** Space out stamp centers at distances Spacing*BrushDiameter along the stroke (so Spacing of 1 means that stamps will be adjacent but non-overlapping). Zero spacing means continuous stamps. */
+	UPROPERTY(EditAnywhere, Category = Brush, meta = (UIMin = "0.0", UIMax = "4.0", ClampMin = "0.0", ClampMax = "1000.0", DisplayPriority = 8, HideEditConditionToggle, EditConditionHides, EditCondition = "bShowSpacing"))
+	float Spacing = 0.0f;
+
+	/** Lazy brush smooths out the brush path by averaging the cursor positions */
+	UPROPERTY(EditAnywhere, Category = Brush, meta = (DisplayName = "Lazy", UIMin = "0", UIMax = "1.0", ClampMin = "0", ClampMax = "1.0", DisplayPriority = 9, HideEditConditionToggle, EditConditionHides, EditCondition = "bShowLazyness"))
+	float Lazyness = 0;
+
+	/**  */
+	UPROPERTY( meta = (TransientToolProperty))
+	bool bShowPerBrushProps = true;
+
+	/**  */
+	UPROPERTY(meta = (TransientToolProperty))
+	bool bShowLazyness = true;
+
+	/**  */
+	UPROPERTY(meta = (TransientToolProperty))
+	bool bShowFlowRate = true;
+
+	UPROPERTY(meta = (TransientToolProperty))
+	bool bShowSpacing = true;
+};
+
+
+
+
+UCLASS()
+class MESHMODELINGTOOLS_API UKelvinBrushProperties : public UInteractiveToolPropertySet
+{
+	GENERATED_BODY()
+
+public:
+	/** Brush Fall off as fraction of brush size*/
+	UPROPERTY(EditAnywhere, Category = Kelvin, meta = (UIMin = "0.0", ClampMin = "0.0"))
+	float FallOffDistance = 1.f;
+
+	/** How much the mesh resists shear */
+	UPROPERTY(EditAnywhere, Category = Kelvin, meta = (UIMin = "0.0", ClampMin = "0.0"))
+	float Stiffness = 1.f;
+
+	/** How compressible the spatial region is: 1 - 2 x Poisson ratio */
+	UPROPERTY(EditAnywhere, Category = Kelvin, meta = (UIMin = "0.0", UIMax = "0.1", ClampMin = "0.0", ClampMax = "1.0"))
+	float Incompressiblity = 1.f;
+
+	/** Integration steps*/
+	UPROPERTY(EditAnywhere, Category = Kelvin, meta = (UIMin = "0.0", UIMax = "100", ClampMin = "0.0", ClampMax = "100"))
+	int BrushSteps = 3;
+};
+
+
+
+
+UCLASS()
+class MESHMODELINGTOOLS_API UWorkPlaneProperties : public UInteractiveToolPropertySet
+{
+	GENERATED_BODY()
+
+public:
+	UPROPERTY( meta = (TransientToolProperty) )
+	bool bPropertySetEnabled = true;
+
+	/** Toggle whether Work Plane Positioning Gizmo is visible */
+	UPROPERTY(EditAnywhere, Category = TargetPlane, meta = (HideEditConditionToggle, EditCondition = "bPropertySetEnabled == true"))
+	bool bShowGizmo = true;
+
+	UPROPERTY(EditAnywhere, Category = TargetPlane, AdvancedDisplay, meta = (HideEditConditionToggle, EditCondition = "bPropertySetEnabled == true"))
+	FVector Position = FVector::ZeroVector;
+
+	UPROPERTY(EditAnywhere, Category = TargetPlane, AdvancedDisplay, meta = (HideEditConditionToggle, EditCondition = "bPropertySetEnabled == true"))
+	FQuat Rotation = FQuat::Identity;
+
+	// Recenter the gizmo around the target position (without changing work plane), if it is "too far" (> 10 meters + max bounds dim) from that position currently
+	void RecenterGizmoIfFar(FVector CenterPosition, double BoundsMaxDim, double TooFarDistance = 1000)
+	{
+		double DistanceTolSq = (BoundsMaxDim + TooFarDistance) * (BoundsMaxDim + TooFarDistance);
+		if (FVector::DistSquared(CenterPosition, Position) > DistanceTolSq)
+		{
+			FVector Normal = Rotation.GetAxisZ();
+			Position = CenterPosition - (CenterPosition - Position).ProjectOnToNormal(Normal);
+		}
+	}
+};
+
+
+
+
+UCLASS()
+class MESHMODELINGTOOLS_API USculptMaxBrushProperties : public UInteractiveToolPropertySet
+{
+	GENERATED_BODY()
+public:
+	/** Specify maximum displacement height (relative to brush size) */
+	UPROPERTY(EditAnywhere, Category = SculptMaxBrush, meta = (UIMin = "0.0", UIMax = "1.0", ClampMin = "0.0", ClampMax = "1.0"))
+	float MaxHeight = 0.5;
+
+	/** Use maximum height from last brush stroke, regardless of brush size. Note that spatial brush falloff still applies.  */
+	UPROPERTY(EditAnywhere, Category = SculptMaxBrush)
+	bool bFreezeCurrentHeight = false;
+};
+
+
+
+
+/**
+ * Base Tool for mesh sculpting tools, provides some shared functionality
+ */
+UCLASS()
+class MESHMODELINGTOOLS_API UMeshSculptToolBase : public UMeshSurfacePointTool
+{
+	GENERATED_BODY()
+protected:
+	using FFrame3d = UE::Geometry::FFrame3d;
+public:
+
+	virtual void RegisterActions(FInteractiveToolActionSet& ActionSet) override;
+
+	virtual void SetWorld(UWorld* World);
+
+	virtual void Setup() override;
+	virtual void Shutdown(EToolShutdownType ShutdownType) override;
+	virtual void Render(IToolsContextRenderAPI* RenderAPI) override;
+
+	// UMeshSurfacePointTool API
+	virtual bool HitTest(const FRay& Ray, FHitResult& OutHit) override;
+	virtual void OnBeginDrag(const FRay& Ray) override;
+	virtual void OnUpdateDrag(const FRay& Ray) override;
+	virtual void OnEndDrag(const FRay& Ray) override;
+	virtual void OnCancelDrag() override;
+	// end UMeshSurfacePointTool API
+
+protected:
+	virtual void OnTick(float DeltaTime) override;
+
+	
+	virtual void OnCompleteSetup();
+	virtual void OnBeginStroke(const FRay& WorldRay) { check(false); }
+	virtual void OnEndStroke() { check(false); }
+	virtual void OnCancelStroke() { check(false); }
+
+public:
+	/** Properties that control brush size/etc */
+	UPROPERTY()
+	TObjectPtr<USculptBrushProperties> BrushProperties;
+
+	/** Properties for 3D workplane / gizmo */
+	UPROPERTY()
+	TObjectPtr<UWorkPlaneProperties> GizmoProperties;
+
+
+protected:
+	UWorld* TargetWorld;		// required to spawn UPreviewMesh/etc
+	FViewCameraState CameraState;
+
+	/** Initial transformation on target mesh */
+	UE::Geometry::FTransformSRT3d InitialTargetTransform;
+	/** Active transformation on target mesh, includes baked scale */
+	UE::Geometry::FTransformSRT3d CurTargetTransform;
+
+	FRay3d GetLocalRay(const FRay& WorldRay) const;
+
+
+	/**
+	 * Subclass must implement this and return relevant rendering component
+	 */
+	virtual UBaseDynamicMeshComponent* GetSculptMeshComponent() { check(false); return nullptr; }
+
+	virtual FDynamicMesh3* GetSculptMesh() { return GetSculptMeshComponent()->GetMesh(); }
+	virtual const FDynamicMesh3* GetSculptMesh() const { return const_cast<UMeshSculptToolBase*>(this)->GetSculptMeshComponent()->GetMesh(); }
+
+	virtual FDynamicMesh3* GetBaseMesh() { check(false); return nullptr; }
+	virtual const FDynamicMesh3* GetBaseMesh() const { check(false); return nullptr; }
+
+	// For any subclass where this returns false, BrushProperties will not be automatically saved/restored, so the class won't use BrushProperties changes made in other tools.
+	virtual bool SharesBrushPropertiesChanges() const { return true; }
+
+
+	/**
+	 * Subclass calls this to set up editing component
+	 */
+	void InitializeSculptMeshComponent(UBaseDynamicMeshComponent* Component, AActor* Actor);
+
+
+	/**
+	 * Subclass can override this to change what results are written.
+	 * Default is to apply a default vertex positions update to the target object.
+	 */
+	virtual void CommitResult(UBaseDynamicMeshComponent* Component, bool bModifiedTopology);
+
+
+	//
+	// Brush Types
+	//
+public:
+	struct FBrushTypeInfo
+	{
+		FText Name;
+		int32 Identifier;
+	};
+	const TArray<FBrushTypeInfo>& GetRegisteredPrimaryBrushTypes() const { return RegisteredPrimaryBrushTypes; }
+	const TArray<FBrushTypeInfo>& GetRegisteredSecondaryBrushTypes() const { return RegisteredSecondaryBrushTypes; }
+
+protected:
+	TArray<FBrushTypeInfo> RegisteredPrimaryBrushTypes;
+	TArray<FBrushTypeInfo> RegisteredSecondaryBrushTypes;
+
+	UPROPERTY()
+	TMap<int32, TObjectPtr<UMeshSculptBrushOpProps>> BrushOpPropSets;
+
+	TMap<int32, TUniquePtr<FMeshSculptBrushOpFactory>> BrushOpFactories;
+
+	UPROPERTY()
+	TMap<int32, TObjectPtr<UMeshSculptBrushOpProps>> SecondaryBrushOpPropSets;
+
+	TMap<int32, TUniquePtr<FMeshSculptBrushOpFactory>> SecondaryBrushOpFactories;
+
+	void RegisterBrushType(int32 Identifier, FText Name, TUniquePtr<FMeshSculptBrushOpFactory> Factory, UMeshSculptBrushOpProps* PropSet);
+	void RegisterSecondaryBrushType(int32 Identifier, FText Name, TUniquePtr<FMeshSculptBrushOpFactory> Factory, UMeshSculptBrushOpProps* PropSet);
+
+	virtual void SaveAllBrushTypeProperties(UInteractiveTool* SaveFromTool);
+	virtual void RestoreAllBrushTypeProperties(UInteractiveTool* RestoreToTool);
+
+protected:
+	TUniquePtr<FMeshSculptBrushOp> PrimaryBrushOp;
+	UMeshSculptBrushOpProps* PrimaryVisiblePropSet = nullptr;		// BrushOpPropSets prevents GC of this
+
+	TUniquePtr<FMeshSculptBrushOp> SecondaryBrushOp;
+	UMeshSculptBrushOpProps* SecondaryVisiblePropSet = nullptr;
+
+	bool bBrushOpPropsVisible = true;
+
+	void SetActivePrimaryBrushType(int32 Identifier);
+	void SetActiveSecondaryBrushType(int32 Identifier);
+	virtual TUniquePtr<FMeshSculptBrushOp>& GetActiveBrushOp();
+	void SetBrushOpPropsVisibility(bool bVisible);
+
+	//
+	// Falloff types
+	//
+public:
+	struct FFalloffTypeInfo
+	{
+		FText Name;
+		FString StringIdentifier;
+		int32 Identifier;
+	};
+	const TArray<FFalloffTypeInfo>& GetRegisteredPrimaryFalloffTypes() const { return RegisteredPrimaryFalloffTypes; }
+
+	/** Set the active falloff type for the primary brush */
+	virtual void SetPrimaryFalloffType(EMeshSculptFalloffType Falloff);
+
+protected:
+	TSharedPtr<FMeshSculptFallofFunc> PrimaryFalloff;
+
+	TArray<FFalloffTypeInfo> RegisteredPrimaryFalloffTypes;
+	void RegisterStandardFalloffTypes();
+
+
+	//
+	// Brush Size
+	//
+protected:
+	UE::Geometry::FInterval1d BrushRelativeSizeRange;
+	double CurrentBrushRadius = 1.0;
+	virtual void InitializeBrushSizeRange(const UE::Geometry::FAxisAlignedBox3d& TargetBounds);
+	virtual void CalculateBrushRadius();
+	virtual double GetCurrentBrushRadius() const { return CurrentBrushRadius; }
+
+	double CurrentBrushFalloff = 0.5;
+	virtual double GetCurrentBrushFalloff() const { return CurrentBrushFalloff; }
+
+	double ActivePressure = 1.0;
+	virtual double GetActivePressure() const { return ActivePressure; }
+
+	virtual double GetCurrentBrushStrength();
+	virtual double GetCurrentBrushDepth();
+
+public:
+	virtual void IncreaseBrushRadiusAction();
+	virtual void DecreaseBrushRadiusAction();
+	virtual void IncreaseBrushRadiusSmallStepAction();
+	virtual void DecreaseBrushRadiusSmallStepAction();
+
+
+	// client currently needs to implement these...
+	virtual void IncreaseBrushSpeedAction() {}
+	virtual void DecreaseBrushSpeedAction() {}
+	virtual void NextBrushModeAction() {}
+	virtual void PreviousBrushModeAction() {}
+
+
+public:
+	// IInteractiveToolCameraFocusAPI override to focus on brush w/ 'F'
+	virtual FBox GetWorldSpaceFocusBox() override;
+
+
+
+	//
+	// Brush/Stroke stuff
+	//
+protected:
+	FFrame3d LastBrushFrameWorld;
+	FFrame3d LastBrushFrameLocal;
+	int32 LastBrushTriangleID = IndexConstants::InvalidID;
+
+	const FFrame3d& GetBrushFrameWorld() const { return LastBrushFrameWorld; }
+	const FFrame3d& GetBrushFrameLocal() const { return LastBrushFrameLocal; }
+	int32 GetBrushTriangleID() const { return LastBrushTriangleID; }
+	void UpdateBrushFrameWorld(const FVector3d& NewPosition, const FVector3d& NewNormal);
+	void AlignBrushToView();
+
+	bool GetBrushCanHitBackFaces() const { return BrushProperties->bHitBackFaces; }
+
+	/** @return hit triangle at ray position - subclass must implement this */
+	virtual int32 FindHitSculptMeshTriangle(const FRay3d& LocalRay) { check(false); return IndexConstants::InvalidID; }
+	/** @return hit triangle at ray position - subclass should implement this for most brushes */
+	virtual int32 FindHitTargetMeshTriangle(const FRay3d& LocalRay) { check(false); return IndexConstants::InvalidID;	}
+
+	virtual bool UpdateBrushPositionOnActivePlane(const FRay& WorldRay);
+	virtual bool UpdateBrushPositionOnTargetMesh(const FRay& WorldRay, bool bFallbackToViewPlane);
+	virtual bool UpdateBrushPositionOnSculptMesh(const FRay& WorldRay, bool bFallbackToViewPlane);
+
+
+
+	//
+	// Brush Target Plane is plane that some brushes move on
+	//
+protected:
+	FFrame3d ActiveBrushTargetPlaneWorld;
+	virtual void UpdateBrushTargetPlaneFromHit(const FRay& WorldRay, const FHitResult& Hit);
+
+
+	//
+	// Stroke Modifiers
+	//
+protected:
+	bool bInStroke = false;
+	bool bSmoothing = false;
+	bool bInvert = false;
+	virtual void SaveActiveStrokeModifiers();
+	virtual bool InStroke() const { return bInStroke; }
+	virtual bool GetInSmoothingStroke() const { return bSmoothing; }
+	virtual bool GetInInvertStroke() const { return bInvert; }
+
+	// when in a stroke, this function determines when a new stamp should be emitted, based on spacing and flow rate settings
+	virtual void UpdateStampPendingState();
+
+	// for tracking stroke time and length, to apply spacing and flow rate settings
+	double ActiveStrokeTime = 0.0;
+	double ActiveStrokePathArcLen = 0.0;
+	int LastFlowTimeStamp = 0;
+	int LastSpacingTimestamp = 0;
+	virtual void ResetStrokeTime();
+	virtual void AccumulateStrokeTime(float DeltaTime);
+
+	//
+	// Stamps
+	//
+protected:
+	bool bIsStampPending = false;
+	FRay PendingStampRay;
+	FSculptBrushStamp HoverStamp;
+	FSculptBrushStamp CurrentStamp;
+	FSculptBrushStamp LastStamp;
+	virtual void UpdateHoverStamp(const FFrame3d& StampFrameWorld);
+	virtual bool IsStampPending() const { return bIsStampPending; }
+	virtual const FRay& GetPendingStampRayWorld() const { return PendingStampRay;  }
+	// Temporal Flow Rate defines the frequency of stamp placement. 1 is max rate, 0 is no stamps. Defaults to BrushProperties->FlowRate, but subclasses can re-use that setting for other things
+	virtual float GetStampTemporalFlowRate() const;
+
+
+	//
+	// Stamp ROI Plane is a plane used by some brush ops
+	//
+protected:
+	FFrame3d StampRegionPlane;
+	virtual FFrame3d ComputeStampRegionPlane(const FFrame3d& StampFrame, const TArray<int32>& StampTriangles, bool bIgnoreDepth, bool bViewAligned, bool bInvDistFalloff = true);
+	virtual FFrame3d ComputeStampRegionPlane(const FFrame3d& StampFrame, const TSet<int32>& StampTriangles, bool bIgnoreDepth, bool bViewAligned, bool bInvDistFalloff = true);
+
+
+	// Stroke plane is a plane used by some brush ops
+	//
+protected:
+	FFrame3d StrokePlane;
+	virtual const FFrame3d& GetCurrentStrokeReferencePlane() const { return StrokePlane; }
+
+	virtual void UpdateStrokeReferencePlaneForROI(const FFrame3d& StampFrame, const TArray<int32>& TriangleROI, bool bViewAligned);
+	virtual void UpdateStrokeReferencePlaneFromWorkPlane();
+
+
+	//
+	// Display / Material
+	//
+public:
+	UPROPERTY()
+	TObjectPtr<UMeshEditingViewProperties> ViewProperties;
+
+	UPROPERTY()
+	TObjectPtr<UMaterialInstanceDynamic> ActiveOverrideMaterial;
+
+protected:
+	virtual void SetViewPropertiesEnabled(bool bNewValue);
+	virtual void UpdateWireframeVisibility(bool bNewValue);
+	virtual void UpdateMaterialMode(EMeshEditingMaterialModes NewMode);
+	virtual void UpdateFlatShadingSetting(bool bNewValue);
+	virtual void UpdateColorSetting(FLinearColor NewColor);
+	virtual void UpdateTransparentColorSetting(FLinearColor NewColor);
+	virtual void UpdateOpacitySetting(double Opacity);
+	virtual void UpdateTwoSidedSetting(bool bOn);
+	virtual void UpdateCustomMaterial(TWeakObjectPtr<UMaterialInterface> NewMaterial);
+	virtual void UpdateImageSetting(UTexture2D* NewImage);
+
+
+
+	//
+	// brush indicator
+	//
+protected:
+	// subclasses should call this to create indicator in their ::Setup()
+	virtual void InitializeIndicator();
+
+	// Called by InitializeIndicator to create a mesh for the brush ROI indicator. Default is sphere.
+	virtual UPreviewMesh* MakeBrushIndicatorMesh(UObject* Parent, UWorld* World);
+
+	virtual void ConfigureIndicator(bool bVolumetric);
+	virtual bool GetIsVolumetricIndicator();
+
+	virtual void SetIndicatorVisibility(bool bVisible);
+	virtual bool GetIndicatorVisibility() const;
+
+protected:
+	UPROPERTY()
+	TObjectPtr<UBrushStampIndicator> BrushIndicator;
+
+	UPROPERTY()
+	bool bIsVolumetricIndicator;
+
+
+	UPROPERTY()
+	TObjectPtr<UMaterialInstanceDynamic> BrushIndicatorMaterial;
+
+	UPROPERTY()
+	TObjectPtr<UPreviewMesh> BrushIndicatorMesh;
+
+	
+	TWeakObjectPtr<ULocalTwoAxisPropertyEditInputBehavior> BrushEditBehavior;
+
+
+	//
+	// Work Plane
+	//
+public:
+	// plane gizmo
+	UPROPERTY()
+	TObjectPtr<UCombinedTransformGizmo> PlaneTransformGizmo;
+
+	UPROPERTY()
+	TObjectPtr<UTransformProxy> PlaneTransformProxy;
+
+protected:
+	virtual void UpdateWorkPlane();
+	virtual bool ShowWorkPlane() const { return false; };
+	// Only relevant when ShowWorkPlane() returns true. Determines whether the plane
+	//  grid lines are drawn.
+	bool bDrawWorkPlaneGridLines = true;
+
+protected:
+	TValueWatcher<FVector> GizmoPositionWatcher;
+	TValueWatcher<FQuat> GizmoRotationWatcher;
+
+	virtual void UpdateGizmoFromProperties();
+	virtual void PlaneTransformChanged(UTransformProxy* Proxy, FTransform Transform);
+
+	enum class EPendingWorkPlaneUpdate
+	{
+		NoUpdatePending,
+		MoveToHitPositionNormal,
+		MoveToHitPosition,
+		MoveToHitPositionViewAligned
+	};
+	EPendingWorkPlaneUpdate PendingWorkPlaneUpdate;
+	virtual void SetFixedSculptPlaneFromWorldPos(const FVector& Position, const FVector& Normal, EPendingWorkPlaneUpdate UpdateType);
+	virtual void UpdateFixedSculptPlanePosition(const FVector& Position);
+	virtual void UpdateFixedSculptPlaneRotation(const FQuat& Rotation);
+	virtual void UpdateFixedPlaneGizmoVisibility(bool bVisible);
+
+protected:
+	UE::TransactionUtil::FLongTransactionTracker LongTransactions;
+
+};
\ No newline at end of file
diff --git a/Engine/Plugins/Runtime/MeshModelingToolset/Source/MeshModelingToolsEditorOnly/MeshModelingToolsEditorOnly.Build.cs b/Engine/Plugins/Runtime/MeshModelingToolset/Source/MeshModelingToolsEditorOnly/MeshModelingToolsEditorOnly.Build.cs
--- a/Engine/Plugins/Runtime/MeshModelingToolset/Source/MeshModelingToolsEditorOnly/MeshModelingToolsEditorOnly.Build.cs
+++ b/Engine/Plugins/Runtime/MeshModelingToolset/Source/MeshModelingToolsEditorOnly/MeshModelingToolsEditorOnly.Build.cs
@@ -13,13 +13,21 @@ public class MeshModelingToolsEditorOnly : ModuleRules
 			{
 				"Core",
 				"InteractiveToolsFramework",
+				"MeshDescription",
+				"StaticMeshDescription",
+				"SkeletalMeshDescription",
 				"GeometryCore",
+				"GeometryFramework",
+				"DynamicMesh",
+				"MeshConversion",
 				"MeshModelingTools",
 				"ModelingComponents",
-				"ModelingOperatorsEditorOnly",
+				"ModelingComponentsEditorOnly",
 				"ModelingOperators",
+				"ModelingOperatorsEditorOnly",
 				"PropertyEditor",
-				"DynamicMesh",
+				"SkeletalMeshUtilitiesCommon",
+				"SkeletalMeshModifiers"
 				// ... add other public dependencies that you statically link with here ...
 			}
 		);
@@ -28,6 +36,16 @@ public class MeshModelingToolsEditorOnly : ModuleRules
 			new string[]
 			{
 				"CoreUObject",
+				"Engine",
+				"RenderCore",
+				"InputCore",
+				"UnrealEd",
+				"Persona",
+				"AdvancedPreviewScene",
+				"SceneOutliner",
+				"Slate",
+				"Json",
+				"ApplicationCore"
 				// ... add private dependencies that you statically link with here ...	
 			}
 		);
diff --git a/Engine/Plugins/Runtime/MeshModelingToolset/Source/MeshModelingToolsEditorOnly/Private/MeshModelingToolsEditorOnly.cpp b/Engine/Plugins/Runtime/MeshModelingToolset/Source/MeshModelingToolsEditorOnly/Private/MeshModelingToolsEditorOnly.cpp
--- a/Engine/Plugins/Runtime/MeshModelingToolset/Source/MeshModelingToolsEditorOnly/Private/MeshModelingToolsEditorOnly.cpp
+++ b/Engine/Plugins/Runtime/MeshModelingToolset/Source/MeshModelingToolsEditorOnly/Private/MeshModelingToolsEditorOnly.cpp
@@ -10,6 +10,8 @@
 
 #define LOCTEXT_NAMESPACE "FMeshModelingToolsEditorOnlyModule"
 
+DEFINE_LOG_CATEGORY(LogMeshModelingToolsEditor);
+
 void FMeshModelingToolsEditorOnlyModule::StartupModule()
 {
 	// This code will execute after your module is loaded into memory; the exact timing is specified in the .uplugin file per-module
diff --git a/Engine/Plugins/Runtime/MeshModelingToolset/Source/MeshModelingToolsEditorOnly/Public/MeshModelingToolsEditorOnly.h b/Engine/Plugins/Runtime/MeshModelingToolset/Source/MeshModelingToolsEditorOnly/Public/MeshModelingToolsEditorOnly.h
--- a/Engine/Plugins/Runtime/MeshModelingToolset/Source/MeshModelingToolsEditorOnly/Public/MeshModelingToolsEditorOnly.h
+++ b/Engine/Plugins/Runtime/MeshModelingToolset/Source/MeshModelingToolsEditorOnly/Public/MeshModelingToolsEditorOnly.h
@@ -5,6 +5,8 @@
 #include "CoreMinimal.h"
 #include "Modules/ModuleManager.h"
 
+DECLARE_LOG_CATEGORY_EXTERN(LogMeshModelingToolsEditor, Log, All);
+
 class FMeshModelingToolsEditorOnlyModule : public IModuleInterface
 {
 public:
