Removed several global "using namespace" statements

diff --git a/Engine/Plugins/Runtime/MeshModelingToolset/Source/ModelingComponents/Private/Mechanics/CollisionPrimitivesMechanic.cpp b/Engine/Plugins/Runtime/MeshModelingToolset/Source/ModelingComponents/Private/Mechanics/CollisionPrimitivesMechanic.cpp
--- a/Engine/Plugins/Runtime/MeshModelingToolset/Source/ModelingComponents/Private/Mechanics/CollisionPrimitivesMechanic.cpp
+++ b/Engine/Plugins/Runtime/MeshModelingToolset/Source/ModelingComponents/Private/Mechanics/CollisionPrimitivesMechanic.cpp
@@ -19,8 +19,6 @@
 
 #include UE_INLINE_GENERATED_CPP_BY_NAME(CollisionPrimitivesMechanic)
 
-using namespace UE::Geometry;
-
 #define LOCTEXT_NAMESPACE "UCollisionPrimitivesMechanic"
 
 namespace CollisionPrimitivesMechanicLocals
@@ -81,7 +79,7 @@ void UCollisionPrimitivesMechanic::Setup(UInteractiveTool* ParentToolIn)
 	SegmentsThickness = 1.0f;
 
 	float UseSegmentRadius = 5.0f;
-	GeometrySetToleranceTest = [UseSegmentRadius](const FVector3d& A, const FVector3d& B) { return DistanceSquared(A, B) < UseSegmentRadius * UseSegmentRadius; };
+	GeometrySetToleranceTest = [UseSegmentRadius](const FVector3d& A, const FVector3d& B) { return UE::Geometry::DistanceSquared(A, B) < UseSegmentRadius * UseSegmentRadius; };
 
 	UInteractiveGizmoManager* GizmoManager = GetParentTool()->GetToolManager()->GetPairedGizmoManager();
 	TranslateTransformProxy = NewObject<UTransformProxy>(this);
@@ -156,7 +154,7 @@ void UCollisionPrimitivesMechanic::Setup(UInteractiveTool* ParentToolIn)
 }
 
 void UCollisionPrimitivesMechanic::Initialize(TSharedPtr<FPhysicsDataCollection> InPhysicsData,
-											   const FAxisAlignedBox3d& MeshBoundsIn,
+											   const UE::Geometry::FAxisAlignedBox3d& MeshBoundsIn,
 											   const FTransform3d& InLocalToWorldTransform)
 {
 	LocalToWorldTransform = InLocalToWorldTransform;
@@ -238,7 +236,7 @@ void UCollisionPrimitivesMechanic::RebuildDrawables(bool bRegenerateCurveLists)
 			A_World = Transform.TransformPosition(A);
 			B_World = Transform.TransformPosition(B);
 			LinesOut.Add(FRenderableLine(A_World, B_World, Color, SegmentsThickness, DepthBias));
-			GeometrySet.AddCurve(CurveIndex, FPolyline3d(A_World, B_World));
+			GeometrySet.AddCurve(CurveIndex, UE::Geometry::FPolyline3d(A_World, B_World));
 
 			PrimitiveToCurveLookup.FindOrAdd(PrimIndex).Add(CurveIndex);
 			CurveToPrimitiveLookup.FindOrAdd(CurveIndex) = PrimIndex;
@@ -254,7 +252,7 @@ void UCollisionPrimitivesMechanic::RebuildDrawables(bool bRegenerateCurveLists)
 
 			DrawnPrimitiveEdges->SetLineStart(CurveIndex, A_World);
 			DrawnPrimitiveEdges->SetLineEnd(CurveIndex, B_World);
-			GeometrySet.UpdateCurve(CurveIndex, FPolyline3d(A_World, B_World));
+			GeometrySet.UpdateCurve(CurveIndex, UE::Geometry::FPolyline3d(A_World, B_World));
 
 			ensure(PrimitiveToCurveLookup[PrimIndex].Contains(CurveIndex));
 			ensure(CurveToPrimitiveLookup[CurveIndex] == PrimIndex);
@@ -286,15 +284,15 @@ void UCollisionPrimitivesMechanic::RebuildDrawables(bool bRegenerateCurveLists)
 					RenderData.Transform = ElemTransform;
 				}
 				float Radius = PhysicsData->ExternalScale3D.GetAbsMin() * Sphere.Radius;
-				UE::Geometry::GenerateCircleSegments<float>(CircleSteps, Radius, FVector3f::Zero(), FVector3f::UnitX(), FVector3f::UnitY(), FTransformSRT3f(),
+				UE::Geometry::GenerateCircleSegments<float>(CircleSteps, Radius, FVector3f::Zero(), FVector3f::UnitX(), FVector3f::UnitY(), UE::Geometry::FTransformSRT3f(),
 					[&](const FVector3f& A, const FVector3f& B) {
 						StoreCurve(RenderData, LinesOut, PrimIndex, GobalCurveIndex++, ElemTransform * LocalToWorldTransform, (FVector)A, (FVector)B, Color);
 					});
-				UE::Geometry::GenerateCircleSegments<float>(CircleSteps, Radius, FVector3f::Zero(), FVector3f::UnitX(), FVector3f::UnitZ(), FTransformSRT3f(),
+				UE::Geometry::GenerateCircleSegments<float>(CircleSteps, Radius, FVector3f::Zero(), FVector3f::UnitX(), FVector3f::UnitZ(), UE::Geometry::FTransformSRT3f(),
 					[&](const FVector3f& A, const FVector3f& B) {
 						StoreCurve(RenderData, LinesOut, PrimIndex, GobalCurveIndex++, ElemTransform * LocalToWorldTransform, (FVector)A, (FVector)B, Color);
 					});
-				UE::Geometry::GenerateCircleSegments<float>(CircleSteps, Radius, FVector3f::Zero(), FVector3f::UnitY(), FVector3f::UnitZ(), FTransformSRT3f(),
+				UE::Geometry::GenerateCircleSegments<float>(CircleSteps, Radius, FVector3f::Zero(), FVector3f::UnitY(), FVector3f::UnitZ(), UE::Geometry::FTransformSRT3f(),
 					[&](const FVector3f& A, const FVector3f& B) { 
 						StoreCurve(RenderData, LinesOut, PrimIndex, GobalCurveIndex++, ElemTransform * LocalToWorldTransform, (FVector)A, (FVector)B, Color);
 					});
@@ -332,7 +330,7 @@ void UCollisionPrimitivesMechanic::RebuildDrawables(bool bRegenerateCurveLists)
 					PhysicsData->ExternalScale3D.X * Box.X * 0.5f,
 					PhysicsData->ExternalScale3D.Y * Box.Y * 0.5f,
 					PhysicsData->ExternalScale3D.Z * Box.Z * 0.5f);
-				UE::Geometry::GenerateBoxSegments<float>(HalfDimensions, FVector3f::Zero(), FVector3f::UnitX(), FVector3f::UnitY(), FVector3f::UnitZ(), FTransformSRT3f(),
+				UE::Geometry::GenerateBoxSegments<float>(HalfDimensions, FVector3f::Zero(), FVector3f::UnitX(), FVector3f::UnitY(), FVector3f::UnitZ(), UE::Geometry::FTransformSRT3f(),
 					[&](const FVector3f& A, const FVector3f& B) {
 						StoreCurve(RenderData, LinesOut, PrimIndex, GobalCurveIndex++, ElemTransform* LocalToWorldTransform, (FVector)A, (FVector)B, Color);
 					});
@@ -370,31 +368,31 @@ void UCollisionPrimitivesMechanic::RebuildDrawables(bool bRegenerateCurveLists)
 				FVector3f Top(0, 0, HalfLength), Bottom(0, 0, -HalfLength);
 
 				// top and bottom circles
-				UE::Geometry::GenerateCircleSegments<float>(CircleSteps, Radius, Top, FVector3f::UnitX(), FVector3f::UnitY(), FTransformSRT3f(),
+				UE::Geometry::GenerateCircleSegments<float>(CircleSteps, Radius, Top, FVector3f::UnitX(), FVector3f::UnitY(), UE::Geometry::FTransformSRT3f(),
 					[&](const FVector3f& A, const FVector3f& B) { 
 						StoreCurve(RenderData, LinesOut, PrimIndex, GobalCurveIndex++, ElemTransform* LocalToWorldTransform, (FVector)A, (FVector)B, Color);
 					});
-				UE::Geometry::GenerateCircleSegments<float>(CircleSteps, Radius, Bottom, FVector3f::UnitX(), FVector3f::UnitY(), FTransformSRT3f(),
+				UE::Geometry::GenerateCircleSegments<float>(CircleSteps, Radius, Bottom, FVector3f::UnitX(), FVector3f::UnitY(), UE::Geometry::FTransformSRT3f(),
 					[&](const FVector3f& A, const FVector3f& B) { 
 						StoreCurve(RenderData, LinesOut, PrimIndex, GobalCurveIndex++, ElemTransform* LocalToWorldTransform, (FVector)A, (FVector)B, Color);
 					});
 
 				// top dome
-				UE::Geometry::GenerateArcSegments<float>(CircleSteps, Radius, 0.0, PI, Top, FVector3f::UnitY(), FVector3f::UnitZ(), FTransformSRT3f(),
+				UE::Geometry::GenerateArcSegments<float>(CircleSteps, Radius, 0.0, PI, Top, FVector3f::UnitY(), FVector3f::UnitZ(), UE::Geometry::FTransformSRT3f(),
 					[&](const FVector3f& A, const FVector3f& B) { 
 						StoreCurve(RenderData, LinesOut, PrimIndex, GobalCurveIndex++, ElemTransform* LocalToWorldTransform, (FVector)A, (FVector)B, Color);
 					});
-				UE::Geometry::GenerateArcSegments<float>(CircleSteps, Radius, 0.0, PI, Top, FVector3f::UnitX(), FVector3f::UnitZ(), FTransformSRT3f(),
+				UE::Geometry::GenerateArcSegments<float>(CircleSteps, Radius, 0.0, PI, Top, FVector3f::UnitX(), FVector3f::UnitZ(), UE::Geometry::FTransformSRT3f(),
 					[&](const FVector3f& A, const FVector3f& B) { 
 						StoreCurve(RenderData, LinesOut, PrimIndex, GobalCurveIndex++, ElemTransform* LocalToWorldTransform, (FVector)A, (FVector)B, Color);
 					});
 
 				// bottom dome
-				UE::Geometry::GenerateArcSegments<float>(CircleSteps, Radius, 0.0, -PI, Bottom, FVector3f::UnitY(), FVector3f::UnitZ(), FTransformSRT3f(),
+				UE::Geometry::GenerateArcSegments<float>(CircleSteps, Radius, 0.0, -PI, Bottom, FVector3f::UnitY(), FVector3f::UnitZ(), UE::Geometry::FTransformSRT3f(),
 					[&](const FVector3f& A, const FVector3f& B) { 
 						StoreCurve(RenderData, LinesOut, PrimIndex, GobalCurveIndex++, ElemTransform* LocalToWorldTransform, (FVector)A, (FVector)B, Color);
 					});
-				UE::Geometry::GenerateArcSegments<float>(CircleSteps, Radius, 0.0, -PI, Bottom, FVector3f::UnitX(), FVector3f::UnitZ(), FTransformSRT3f(),
+				UE::Geometry::GenerateArcSegments<float>(CircleSteps, Radius, 0.0, -PI, Bottom, FVector3f::UnitX(), FVector3f::UnitZ(), UE::Geometry::FTransformSRT3f(),
 					[&](const FVector3f& A, const FVector3f& B) { 
 						StoreCurve(RenderData, LinesOut, PrimIndex, GobalCurveIndex++, ElemTransform* LocalToWorldTransform, (FVector)A, (FVector)B, Color);
 					});
@@ -489,7 +487,7 @@ void UCollisionPrimitivesMechanic::RebuildDrawables(bool bRegenerateCurveLists)
 					const FVector3d Center = Box.GetCenter();
 					const FVector3d HalfDimensions = 0.5 * (Box.Max - Box.Min);
 
-					UE::Geometry::GenerateBoxSegments<double>(HalfDimensions, Center, FVector3d::UnitX(), FVector3d::UnitY(), FVector3d::UnitZ(), FTransformSRT3d(),
+					UE::Geometry::GenerateBoxSegments<double>(HalfDimensions, Center, FVector3d::UnitX(), FVector3d::UnitY(), FVector3d::UnitZ(), UE::Geometry::FTransformSRT3d(),
 						[&](const FVector3d& A, const FVector3d& B) {
 							StoreCurve(RenderData, LinesOut, PrimIndex, GobalCurveIndex++, ElemTransform* LocalToWorldTransform, (FVector)A, (FVector)B, Color);
 						});
@@ -567,13 +565,12 @@ void UCollisionPrimitivesMechanic::GizmoTransformChanged(UTransformProxy* Proxy,
 	FTransform LocalTransform = Transform * LocalToWorldTransform.Inverse();
 
 
-
-	FQuaterniond DeltaRotation = (FQuaterniond) ( LocalToWorldTransform.InverseTransformRotation( Transform.GetRotation() ) * LocalToWorldTransform.InverseTransformRotation( GizmoStartRotation) .Inverse() );
+	UE::Geometry::FQuaterniond DeltaRotation = (UE::Geometry::FQuaterniond) ( LocalToWorldTransform.InverseTransformRotation( Transform.GetRotation() ) * LocalToWorldTransform.InverseTransformRotation( GizmoStartRotation) .Inverse() );
 	FVector DeltaScale = Transform.GetScale3D() / GizmoStartScale;
 
-	FTransformSRT3d DeltaTransform;
+	UE::Geometry::FTransformSRT3d DeltaTransform;
 	DeltaTransform.SetScale((FVector3d)DeltaScale);
-	DeltaTransform.SetRotation((FQuaterniond)DeltaRotation);
+	DeltaTransform.SetRotation((UE::Geometry::FQuaterniond)DeltaRotation);
 	DeltaTransform.SetTranslation( (FVector3d)LocalTransform.GetTranslation() - (FVector3d)LocalToWorldTransform.InverseTransformPosition(GizmoStartPosition) );
 
 	FKAggregateGeom& AggGeom = PhysicsData->AggGeom;
@@ -1000,7 +997,7 @@ void UCollisionPrimitivesMechanic::OnDragRectangleChanged(const FCameraRectangle
 	using namespace CollisionPrimitivesMechanicLocals;
 	
 	TSet<int> CurveSelection, PrimitiveSelection;
-	GeometrySet.ParallelFindAllCurvesSatisfying([&](const FPolyline3d& Segment)
+	GeometrySet.ParallelFindAllCurvesSatisfying([&](const UE::Geometry::FPolyline3d& Segment)
 	{
 		return Rectangle.IsProjectedSegmentIntersectingRectangle(Segment.Start(), Segment.End());
 	}, CurveSelection);
diff --git a/Engine/Source/Editor/ContentBrowser/Private/Menus/AssetViewMenus.cpp b/Engine/Source/Editor/ContentBrowser/Private/Menus/AssetViewMenus.cpp
--- a/Engine/Source/Editor/ContentBrowser/Private/Menus/AssetViewMenus.cpp
+++ b/Engine/Source/Editor/ContentBrowser/Private/Menus/AssetViewMenus.cpp
@@ -60,12 +60,12 @@ namespace UE::Editor::ContentBrowser
 	}
 }
 
-using namespace UE::Editor::ContentBrowser;
-
 FDelayedAutoRegisterHelper SAssetView::AssetViewOptionsMenuRegistration(
 	EDelayedRegisterRunPhase::EndOfEngineInit,
 	[]
 	{
+		using namespace UE::Editor::ContentBrowser;
+
 		FToolMenuOwnerScoped OwnerScoped(UE_MODULE_NAME);
 
 		static const FName AssetViewOptionsMenuName("ContentBrowser.AssetViewOptions");
@@ -520,6 +520,8 @@ FDelayedAutoRegisterHelper SAssetView::ToolBarMenuExtensionRegistration(
 	EDelayedRegisterRunPhase::EndOfEngineInit,
 	[]
 	{
+		using namespace UE::Editor::ContentBrowser;
+
 		// @todo: Remove when old style is removed, all extensions are for the new style so early-out if it's not enabled
 		if (!IsNewStyleEnabled())
 		{
diff --git a/Engine/Source/Editor/ContentBrowser/Private/Widgets/SContentBrowserSourceTree.cpp b/Engine/Source/Editor/ContentBrowser/Private/Widgets/SContentBrowserSourceTree.cpp
--- a/Engine/Source/Editor/ContentBrowser/Private/Widgets/SContentBrowserSourceTree.cpp
+++ b/Engine/Source/Editor/ContentBrowser/Private/Widgets/SContentBrowserSourceTree.cpp
@@ -21,8 +21,6 @@
 
 #define LOCTEXT_NAMESPACE "ContentBrowser"
 
-using namespace UE::ContentBrowser::Private;
-
 namespace UE::Editor::ContentBrowser::Private
 {
 	static constexpr float MinBodyHeight = 88.0f;
@@ -170,6 +168,8 @@ namespace UE::Editor::ContentBrowser::Private
 
 	void SContentBrowserSourceTree::Construct(const FArguments& InArgs)
 	{
+		using namespace UE::ContentBrowser::Private;
+
 		constexpr float SourceTreeSectionPadding = 2.0f;
 		constexpr float SourceTreeSectionHandleSize = 8.0f;
 		constexpr float SourceTreeHeaderHeight = 36.0f;
@@ -222,6 +222,8 @@ namespace UE::Editor::ContentBrowser::Private
 		const TSharedPtr<FSourcesSearch>& InSearch,
 		TSharedRef<IScrollableWidget> InBody)
 	{
+		using namespace UE::ContentBrowser::Private;
+
 		check(!InId.IsNone());
 
 		Id = InId;
