make the collision primitives mechanic use interval gizmos instead of scale to control  capsule length/radius and box width/depth/height, so that the translation can change between world and local space without breaking the shape parameters UI

+ make the interval gizmo support (positive and negative) intervals on each major axis, rather than hard-coding only three intervals

diff --git a/Engine/Plugins/Runtime/MeshModelingToolset/Source/ModelingComponents/Private/Mechanics/CollisionPrimitivesMechanic.cpp b/Engine/Plugins/Runtime/MeshModelingToolset/Source/ModelingComponents/Private/Mechanics/CollisionPrimitivesMechanic.cpp
--- a/Engine/Plugins/Runtime/MeshModelingToolset/Source/ModelingComponents/Private/Mechanics/CollisionPrimitivesMechanic.cpp
+++ b/Engine/Plugins/Runtime/MeshModelingToolset/Source/ModelingComponents/Private/Mechanics/CollisionPrimitivesMechanic.cpp
@@ -4,6 +4,7 @@
 
 #include "BaseBehaviors/SingleClickOrDragBehavior.h"
 #include "BaseBehaviors/MouseHoverBehavior.h"
+#include "BaseGizmos/IntervalGizmo.h"
 #include "BaseGizmos/TransformGizmoUtil.h"
 #include "BaseGizmos/TransformProxy.h"
 #include "Drawing/PreviewGeometryActor.h"
@@ -96,21 +97,47 @@ void UCollisionPrimitivesMechanic::Setup(UInteractiveTool* ParentToolIn)
 		ETransformGizmoSubElements::TranslateAllAxes | ETransformGizmoSubElements::TranslateAllPlanes | ETransformGizmoSubElements::ScaleUniform, this);
 
 	BoxTransformGizmo = UE::TransformGizmoUtil::CreateCustomTransformGizmo(GizmoManager,
-		ETransformGizmoSubElements::FullTranslateRotateScale, this);
+		ETransformGizmoSubElements::TranslateAllAxes | ETransformGizmoSubElements::TranslateAllPlanes | ETransformGizmoSubElements::RotateAllAxes, this);
+
+	// create sources for the interval parameters
+	BoxXIntervalSource = NewObject< UGizmoLocalFloatParameterSource >(this);
+	BoxYIntervalSource = NewObject< UGizmoLocalFloatParameterSource >(this);
+	BoxZIntervalSource = NewObject< UGizmoLocalFloatParameterSource >(this);
+
+	CapsuleRadiusIntervalSource = NewObject< UGizmoLocalFloatParameterSource >(this);
+	CapsuleLengthIntervalSource = NewObject< UGizmoLocalFloatParameterSource >(this);
+
+	// add the interval gizmo
+	GizmoManager->RegisterGizmoType(UIntervalGizmo::GizmoName, NewObject<UIntervalGizmoBuilder>());
+	BoxIntervalGizmo = GizmoManager->CreateGizmo<UIntervalGizmo>(UIntervalGizmo::GizmoName, TEXT("BoxSizeInterval"), this);
+
+	UIntervalGizmo::FParameterSources BoxIntervalSources;
+	BoxIntervalSources.RightInterval = BoxXIntervalSource;
+	BoxIntervalSources.ForwardInterval = BoxYIntervalSource;
+	BoxIntervalSources.UpInterval = BoxZIntervalSource;
+	BoxIntervalGizmo->SetActiveTarget(BoxTransformProxy, BoxIntervalSources);
 
 	CapsuleTransformGizmo = UE::TransformGizmoUtil::CreateCustomTransformGizmo(GizmoManager,
-		ETransformGizmoSubElements::TranslateAllAxes | ETransformGizmoSubElements::TranslateAllPlanes | ETransformGizmoSubElements::ScaleAxisX | ETransformGizmoSubElements::ScaleAxisZ | ETransformGizmoSubElements::RotateAxisY | ETransformGizmoSubElements::RotateAxisX, this);
+		ETransformGizmoSubElements::TranslateAllAxes | ETransformGizmoSubElements::TranslateAllPlanes | ETransformGizmoSubElements::RotateAllAxes, this);
+
+	CapsuleIntervalGizmo = GizmoManager->CreateGizmo<UIntervalGizmo>(UIntervalGizmo::GizmoName, TEXT("CapsuleSizeInterval"), this);
+	UIntervalGizmo::FParameterSources CapsuleIntervalSources;
+	CapsuleIntervalSources.RightInterval = CapsuleRadiusIntervalSource;
+	CapsuleIntervalSources.UpInterval = CapsuleLengthIntervalSource;
+	CapsuleIntervalGizmo->SetActiveTarget(CapsuleTransformProxy, CapsuleIntervalSources);
 
 	FullTransformGizmo = UE::TransformGizmoUtil::CreateCustomTransformGizmo(GizmoManager,
 		ETransformGizmoSubElements::FullTranslateRotateScale, this);
 
+	const bool bShouldUseContextCoordinateSystem = true;
+
 	TranslateTransformProxy->OnTransformChanged.AddUObject(this, &UCollisionPrimitivesMechanic::GizmoTransformChanged);
 	TranslateTransformProxy->OnBeginTransformEdit.AddUObject(this, &UCollisionPrimitivesMechanic::GizmoTransformStarted);
 	TranslateTransformProxy->OnEndTransformEdit.AddUObject(this, &UCollisionPrimitivesMechanic::GizmoTransformEnded);
 
 	TranslateTransformGizmo->SetActiveTarget(TranslateTransformProxy);
 	TranslateTransformGizmo->SetVisibility(false);
-	TranslateTransformGizmo->bUseContextCoordinateSystem = false;
+	TranslateTransformGizmo->bUseContextCoordinateSystem = bShouldUseContextCoordinateSystem;
 	TranslateTransformGizmo->CurrentCoordinateSystem = EToolContextCoordinateSystem::Local;
 
 	SphereTransformProxy->OnTransformChanged.AddUObject(this, &UCollisionPrimitivesMechanic::GizmoTransformChanged);
@@ -119,26 +146,30 @@ void UCollisionPrimitivesMechanic::Setup(UInteractiveTool* ParentToolIn)
 
 	SphereTransformGizmo->SetActiveTarget(SphereTransformProxy);
 	SphereTransformGizmo->SetVisibility(false);
-	SphereTransformGizmo->bUseContextCoordinateSystem = false;
+	SphereTransformGizmo->bUseContextCoordinateSystem = bShouldUseContextCoordinateSystem;
 	SphereTransformGizmo->CurrentCoordinateSystem = EToolContextCoordinateSystem::Local;
 
 	BoxTransformProxy->OnTransformChanged.AddUObject(this, &UCollisionPrimitivesMechanic::GizmoTransformChanged);
 	BoxTransformProxy->OnBeginTransformEdit.AddUObject(this, &UCollisionPrimitivesMechanic::GizmoTransformStarted);
 	BoxTransformProxy->OnEndTransformEdit.AddUObject(this, &UCollisionPrimitivesMechanic::GizmoTransformEnded);
+	BoxIntervalGizmo->OnIntervalChanged.AddUObject(this, &UCollisionPrimitivesMechanic::IntervalGizmoValueChanged);
 
 	BoxTransformGizmo->SetActiveTarget(BoxTransformProxy);
 	BoxTransformGizmo->SetVisibility(false);
-	BoxTransformGizmo->bUseContextCoordinateSystem = false;
+	BoxTransformGizmo->bUseContextCoordinateSystem = bShouldUseContextCoordinateSystem;
 	BoxTransformGizmo->CurrentCoordinateSystem = EToolContextCoordinateSystem::Local;
+	BoxIntervalGizmo->SetVisibility(false);
 
 	CapsuleTransformProxy->OnTransformChanged.AddUObject(this, &UCollisionPrimitivesMechanic::GizmoTransformChanged);
 	CapsuleTransformProxy->OnBeginTransformEdit.AddUObject(this, &UCollisionPrimitivesMechanic::GizmoTransformStarted);
 	CapsuleTransformProxy->OnEndTransformEdit.AddUObject(this, &UCollisionPrimitivesMechanic::GizmoTransformEnded);
+	CapsuleIntervalGizmo->OnIntervalChanged.AddUObject(this, &UCollisionPrimitivesMechanic::IntervalGizmoValueChanged);
 
 	CapsuleTransformGizmo->SetActiveTarget(CapsuleTransformProxy);
 	CapsuleTransformGizmo->SetVisibility(false);
-	CapsuleTransformGizmo->bUseContextCoordinateSystem = false;
+	CapsuleTransformGizmo->bUseContextCoordinateSystem = bShouldUseContextCoordinateSystem;
 	CapsuleTransformGizmo->CurrentCoordinateSystem = EToolContextCoordinateSystem::Local;
+	CapsuleIntervalGizmo->SetVisibility(false);
 
 	FullTransformProxy->OnTransformChanged.AddUObject(this, &UCollisionPrimitivesMechanic::GizmoTransformChanged);
 	FullTransformProxy->OnBeginTransformEdit.AddUObject(this, &UCollisionPrimitivesMechanic::GizmoTransformStarted);
@@ -146,7 +177,7 @@ void UCollisionPrimitivesMechanic::Setup(UInteractiveTool* ParentToolIn)
 
 	FullTransformGizmo->SetActiveTarget(FullTransformProxy);
 	FullTransformGizmo->SetVisibility(false);
-	FullTransformGizmo->bUseContextCoordinateSystem = false;
+	FullTransformGizmo->bUseContextCoordinateSystem = bShouldUseContextCoordinateSystem;
 	FullTransformGizmo->CurrentCoordinateSystem = EToolContextCoordinateSystem::Local;
 
 	// Hold the active proxy, so we can capture undo information from whatever gizmo got moved by the user.
@@ -199,6 +230,7 @@ void UCollisionPrimitivesMechanic::Shutdown()
 	// Calls shutdown on gizmo and destroys it.
 	UInteractiveGizmoManager* GizmoManager = GetParentTool()->GetToolManager()->GetPairedGizmoManager();
 	GizmoManager->DestroyAllGizmosByOwner(this);
+	GizmoManager->DeregisterGizmoType(UIntervalGizmo::GizmoName);
 }
 
 
@@ -538,6 +570,45 @@ void UCollisionPrimitivesMechanic::UpdateDrawables()
 }
 
 
+void UCollisionPrimitivesMechanic::IntervalGizmoValueChanged(UIntervalGizmo* IntervalGizmo, const FVector& Direction, float Value)
+{
+	if (SelectedPrimitiveIDs.Num() == 0 || !bGizmoBeingDragged)
+	{
+		return;
+	}
+
+	FKAggregateGeom& AggGeom = PhysicsData->AggGeom;
+
+	for (int32 PrimitiveID : SelectedPrimitiveIDs)
+	{
+		FKShapeElem* Prim = AggGeom.GetElement(PrimitiveID);
+
+		switch (Prim->GetShapeType())
+		{
+		case EAggCollisionShape::Sphere:
+			break;
+		case EAggCollisionShape::Box:
+			SetBoxShapeFromIntervals(StaticCast<FKBoxElem*>(Prim));
+			break;
+		case EAggCollisionShape::Sphyl:
+			SetCapsuleShapeFromIntervals(StaticCast<FKSphylElem*>(Prim));
+			break;
+		case EAggCollisionShape::Convex:
+			break;
+		case EAggCollisionShape::LevelSet:
+			break;
+		default:
+			ensure(false); // If we encounter a non-supported primitive type,
+			// we should catch this and figure out how to handle it when needed.
+		};
+	}
+
+	RebuildDrawables(false);
+
+	OnCollisionGeometryChanged.Broadcast();
+}
+
+
 void UCollisionPrimitivesMechanic::GizmoTransformStarted(UTransformProxy* Proxy)
 {
 	ParentTool->GetToolManager()->BeginUndoTransaction(CollisionPrimitivesMechanicLocals::CollisionPrimitiveMovementTransactionText);
@@ -571,7 +642,7 @@ void UCollisionPrimitivesMechanic::GizmoTransformChanged(UTransformProxy* Proxy,
 	UE::Geometry::FTransformSRT3d DeltaTransform;
 	DeltaTransform.SetScale((FVector3d)DeltaScale);
 	DeltaTransform.SetRotation((UE::Geometry::FQuaterniond)DeltaRotation);
-	DeltaTransform.SetTranslation( (FVector3d)LocalTransform.GetTranslation() - (FVector3d)LocalToWorldTransform.InverseTransformPosition(GizmoStartPosition) );
+	DeltaTransform.SetTranslation( (FVector3d)LocalToWorldTransform.InverseTransformPosition(Transform.GetLocation()) - (FVector3d)LocalToWorldTransform.InverseTransformPosition(GizmoStartPosition));
 
 	FKAggregateGeom& AggGeom = PhysicsData->AggGeom;
 	FKAggregateGeom& PreTransformAggGeom = *PrimitivePreTransform;
@@ -597,9 +668,7 @@ void UCollisionPrimitivesMechanic::GizmoTransformChanged(UTransformProxy* Proxy,
 			ElemTM = StaticCast<FKBoxElem*>(PreTransformPrim)->GetTransform();
 			ElemTM.Accumulate(DeltaTransform, ScalarRegister(1.0));
 			ElemTM.NormalizeRotation();
-			StaticCast<FKBoxElem*>(Prim)->X = StaticCast<FKBoxElem*>(PreTransformPrim)->X * ElemTM.GetScale3D().X;
-			StaticCast<FKBoxElem*>(Prim)->Y = StaticCast<FKBoxElem*>(PreTransformPrim)->Y * ElemTM.GetScale3D().Y;
-			StaticCast<FKBoxElem*>(Prim)->Z = StaticCast<FKBoxElem*>(PreTransformPrim)->Z * ElemTM.GetScale3D().Z;
+			SetBoxShapeFromIntervals(StaticCast<FKBoxElem*>(Prim));
 			ElemTM.SetScale3D(FVector::One());
 			StaticCast<FKBoxElem*>(Prim)->SetTransform(ElemTM);
 			break;
@@ -607,8 +676,7 @@ void UCollisionPrimitivesMechanic::GizmoTransformChanged(UTransformProxy* Proxy,
 			ElemTM = StaticCast<FKSphylElem*>(PreTransformPrim)->GetTransform();
 			ElemTM.Accumulate(DeltaTransform, ScalarRegister(1.0));
 			ElemTM.NormalizeRotation();
-			StaticCast<FKSphylElem*>(Prim)->Radius = StaticCast<FKSphylElem*>(PreTransformPrim)->Radius * ElemTM.GetScale3D().GetAbs().X;
-			StaticCast<FKSphylElem*>(Prim)->Length = StaticCast<FKSphylElem*>(PreTransformPrim)->Length * ElemTM.GetScale3D().GetAbs().Z;
+			SetCapsuleShapeFromIntervals(StaticCast<FKSphylElem*>(Prim));
 			ElemTM.SetScale3D(FVector::One());
 			StaticCast<FKSphylElem*>(Prim)->SetTransform(ElemTM);
 			break;
@@ -635,6 +703,29 @@ void UCollisionPrimitivesMechanic::GizmoTransformChanged(UTransformProxy* Proxy,
 	OnCollisionGeometryChanged.Broadcast();
 }
 
+void UCollisionPrimitivesMechanic::SetBoxShapeFromIntervals(FKBoxElem* BoxElem) const
+{
+	FVector ExternalScale = GetSafeAbsScale(LocalToWorldTransform.GetScale3D());
+	BoxElem->X = FMath::Abs(BoxXIntervalSource->GetParameter() * 2 / ExternalScale.X);
+	BoxElem->Y = FMath::Abs(BoxYIntervalSource->GetParameter() * 2 / ExternalScale.Y);
+	BoxElem->Z = FMath::Abs(BoxZIntervalSource->GetParameter() * 2 / ExternalScale.Z);
+}
+
+void UCollisionPrimitivesMechanic::SetCapsuleShapeFromIntervals(FKSphylElem* CapsuleElem) const
+{
+	FVector ExternalScale = GetSafeAbsScale(LocalToWorldTransform.GetScale3D());
+	CapsuleElem->Radius = FMath::Abs(CapsuleRadiusIntervalSource->GetParameter() / FMath::Max(ExternalScale.X, ExternalScale.Y));
+	CapsuleElem->Length = FMath::Abs(CapsuleLengthIntervalSource->GetParameter() * 2 / ExternalScale.Z);
+}
+
+FVector3d UCollisionPrimitivesMechanic::GetSafeAbsScale(FVector3d Scale3D) const
+{
+	Scale3D.X = FMath::Max(FMathd::ZeroTolerance, FMath::Abs(Scale3D.X));
+	Scale3D.Y = FMath::Max(FMathd::ZeroTolerance, FMath::Abs(Scale3D.Y));
+	Scale3D.Z = FMath::Max(FMathd::ZeroTolerance, FMath::Abs(Scale3D.Z));
+	return Scale3D;
+}
+
 void UCollisionPrimitivesMechanic::GizmoTransformEnded(UTransformProxy* Proxy)
 {
 	FKAggregateGeom& AggGeom = PhysicsData->AggGeom;
@@ -650,14 +741,11 @@ void UCollisionPrimitivesMechanic::GizmoTransformEnded(UTransformProxy* Proxy)
 			break;
 		case EAggCollisionShape::Box:
 			StaticCast<FKBoxElem*>(Prim)->SetTransform(ElemTM);
-			StaticCast<FKBoxElem*>(Prim)->X *= ElemTM.GetScale3D().X;
-			StaticCast<FKBoxElem*>(Prim)->Y *= ElemTM.GetScale3D().Y;
-			StaticCast<FKBoxElem*>(Prim)->Z *= ElemTM.GetScale3D().Z;
+			SetBoxShapeFromIntervals(StaticCast<FKBoxElem*>(Prim));
 			break;
 		case EAggCollisionShape::Sphyl:
 			StaticCast<FKSphylElem*>(Prim)->SetTransform(ElemTM);
-			StaticCast<FKSphylElem*>(Prim)->Radius *= ElemTM.GetScale3D().GetAbs().X;
-			StaticCast<FKSphylElem*>(Prim)->Length *= ElemTM.GetScale3D().GetAbs().Z;
+			SetCapsuleShapeFromIntervals(StaticCast<FKSphylElem*>(Prim));
 			break;
 		case EAggCollisionShape::Convex:
 			StaticCast<FKConvexElem*>(Prim)->SetTransform(ElemTM);
@@ -808,10 +896,29 @@ void UCollisionPrimitivesMechanic::UpdateGizmoLocation()
 
 	FQuat GizmoRotation = LocalToWorldTransform.GetRotation();
 
-	// We are handling one selected primitive specially because multiple selected primitives don't get rotation controls
+	// We are handling one selected primitive specially because multiple selected primitives don't get rotation or interval controls
 	if (SelectedPrimitiveIDs.Num() == 1)
 	{
-		GizmoRotation = (PrimitiveRenderData[SelectedPrimitiveIDs.Array()[0]].Transform * LocalToWorldTransform).GetRotation();
+		const int32 SelectedPrimitiveID = *SelectedPrimitiveIDs.CreateConstIterator();
+		GizmoRotation = (PrimitiveRenderData[SelectedPrimitiveID].Transform * LocalToWorldTransform).GetRotation();
+
+		FKShapeElem* Shape = PhysicsData->AggGeom.GetElement(SelectedPrimitiveID);
+		if (Shape->GetShapeType() == EAggCollisionShape::Box)
+		{
+			FKBoxElem* BoxElem = StaticCast<FKBoxElem*>(Shape);
+			FVector ExternalScale = GetSafeAbsScale(LocalToWorldTransform.GetScale3D());
+			BoxXIntervalSource->SetParameter(BoxElem->X * .5 * ExternalScale.X);
+			BoxYIntervalSource->SetParameter(BoxElem->Y * .5 * ExternalScale.Y);
+			BoxZIntervalSource->SetParameter(BoxElem->Z * .5 * ExternalScale.Z);
+		}
+		if (Shape->GetShapeType() == EAggCollisionShape::Sphyl)
+		{
+			FKSphylElem* CapsuleElem = StaticCast<FKSphylElem*>(Shape);
+			FVector ExternalScaleAbs = GetSafeAbsScale(LocalToWorldTransform.GetScale3D());
+			const float RadiusScale = FMath::Max(ExternalScaleAbs.X, ExternalScaleAbs.Y);
+			CapsuleRadiusIntervalSource->SetParameter(CapsuleElem->Radius * RadiusScale);
+			CapsuleLengthIntervalSource->SetParameter(CapsuleElem->Length * .5 * ExternalScaleAbs.Z);
+		}
 	}
 
 	UpdateGizmoVisibility();
@@ -850,14 +957,18 @@ void UCollisionPrimitivesMechanic::UpdateGizmoVisibility()
 	{
 		Gizmo.SetVisibility(false);
 	});
+	// Note: Interval gizmos are not transform gizmos so we handle them separately from the above 'forall' helper
+	BoxIntervalGizmo->SetVisibility(false);
+	CapsuleIntervalGizmo->SetVisibility(false);
 	CurrentActiveProxy = TranslateTransformProxy;
 
 	if(!(bIsDraggingRectangle || SelectedPrimitiveIDs.IsEmpty() || (ShouldHideGizmo.IsBound() && ShouldHideGizmo.Execute())))
 	{
-		// We are handling one selected primitive specially because each primitive type gets it's own defined gizmo
+		// We are handling one selected primitive specially because each primitive type gets its own defined gizmo
 		if (SelectedPrimitiveIDs.Num() == 1)
 		{
-			switch (PrimitiveRenderData[SelectedPrimitiveIDs.Array()[0]].ShapeType)
+			const int32 PrimitiveID = *SelectedPrimitiveIDs.CreateConstIterator();
+			switch (PrimitiveRenderData[PrimitiveID].ShapeType)
 			{
 			case EAggCollisionShape::Sphere:
 				SphereTransformGizmo->SetVisibility(true);
@@ -865,10 +976,12 @@ void UCollisionPrimitivesMechanic::UpdateGizmoVisibility()
 				break;
 			case EAggCollisionShape::Box:
 				BoxTransformGizmo->SetVisibility(true);
+				BoxIntervalGizmo->SetVisibility(true);
 				CurrentActiveProxy = BoxTransformProxy;
 				break;
 			case EAggCollisionShape::Sphyl:
 				CapsuleTransformGizmo->SetVisibility(true);
+				CapsuleIntervalGizmo->SetVisibility(true);
 				CurrentActiveProxy = CapsuleTransformProxy;
 				break;
 			case EAggCollisionShape::Convex:
diff --git a/Engine/Plugins/Runtime/MeshModelingToolset/Source/ModelingComponents/Public/Mechanics/CollisionPrimitivesMechanic.h b/Engine/Plugins/Runtime/MeshModelingToolset/Source/ModelingComponents/Public/Mechanics/CollisionPrimitivesMechanic.h
--- a/Engine/Plugins/Runtime/MeshModelingToolset/Source/ModelingComponents/Public/Mechanics/CollisionPrimitivesMechanic.h
+++ b/Engine/Plugins/Runtime/MeshModelingToolset/Source/ModelingComponents/Public/Mechanics/CollisionPrimitivesMechanic.h
@@ -20,9 +20,13 @@ class ULineSetComponent;
 class UMouseHoverBehavior;
 class USingleClickInputBehavior;
 class UCombinedTransformGizmo;
+class UIntervalGizmo;
+class UGizmoLocalFloatParameterSource;
 class UTransformProxy;
 class FPhysicsDataCollection;
 struct FKAggregateGeom;
+struct FKBoxElem;
+struct FKSphylElem;
 
 UCLASS()
 class MODELINGCOMPONENTS_API UCollisionPrimitivesMechanic : 
@@ -153,17 +157,36 @@ protected:
 	UPROPERTY()
 	TObjectPtr<UCombinedTransformGizmo> FullTransformGizmo = nullptr;
 
+	UPROPERTY()
+	TObjectPtr<UIntervalGizmo> BoxIntervalGizmo = nullptr;
+	UPROPERTY()
+	TObjectPtr<UGizmoLocalFloatParameterSource> BoxXIntervalSource;
+	UPROPERTY()
+	TObjectPtr<UGizmoLocalFloatParameterSource> BoxYIntervalSource;
+	UPROPERTY()
+	TObjectPtr<UGizmoLocalFloatParameterSource> BoxZIntervalSource;
+
+	UPROPERTY()
+	TObjectPtr<UIntervalGizmo> CapsuleIntervalGizmo = nullptr;
+	UPROPERTY()
+	TObjectPtr<UGizmoLocalFloatParameterSource> CapsuleRadiusIntervalSource;
+	UPROPERTY()
+	TObjectPtr<UGizmoLocalFloatParameterSource> CapsuleLengthIntervalSource;
+
 	// Used to make it easy to tell whether the gizmo was moved by the user or by undo/redo or
 	// some other change that we shouldn't respond to. Basing our movement undo/redo on the
 	// gizmo turns out to be quite a pain, though may someday be easier if the transform proxy
 	// is able to manage arbitrary objects.
 	bool bGizmoBeingDragged = false;
 
-	// Callbacks we'll receive from the gizmo proxy
+	// Callbacks we'll receive from the gizmo proxies
 	void GizmoTransformChanged(UTransformProxy* Proxy, FTransform Transform);
 	void GizmoTransformStarted(UTransformProxy* Proxy);
 	void GizmoTransformEnded(UTransformProxy* Proxy);
 
+	// Callbacks we'll receive from the interval gizmo
+	void IntervalGizmoValueChanged(UIntervalGizmo* IntervalGizmo, const FVector& Direction, float Value);
+
 	// Support for hovering
 	TFunction<bool(const FVector3d&, const FVector3d&)> GeometrySetToleranceTest;
 	int32 HoveredPrimitiveID = -1;
@@ -258,4 +281,10 @@ protected:
 
 private:
 	UE::TransactionUtil::FLongTransactionTracker LongTransactions;
+
+	// Helper to get a 'safe' copy of scale in which no elements are zero, so we can divide by each dimension
+	FVector3d GetSafeAbsScale(FVector3d Scale3D) const;
+	
+	void SetBoxShapeFromIntervals(FKBoxElem* BoxElem) const;
+	void SetCapsuleShapeFromIntervals(FKSphylElem* CapsuleElem) const;
 };
diff --git a/Engine/Source/Runtime/InteractiveToolsFramework/Private/BaseGizmos/IntervalGizmo.cpp b/Engine/Source/Runtime/InteractiveToolsFramework/Private/BaseGizmos/IntervalGizmo.cpp
--- a/Engine/Source/Runtime/InteractiveToolsFramework/Private/BaseGizmos/IntervalGizmo.cpp
+++ b/Engine/Source/Runtime/InteractiveToolsFramework/Private/BaseGizmos/IntervalGizmo.cpp
@@ -147,11 +147,13 @@ AIntervalGizmoActor* AIntervalGizmoActor::ConstructDefaultIntervalGizmo(UWorld*
 
 	const FLinearColor MintGreen(152 / 255.f, 1.f, 152 / 255.f);
 	
-	// add components 
+	// add all possible interval components (note: some may be hidden / unused)
 	NewActor->UpIntervalComponent      = AddDefaultLineHandleComponent(World, NewActor, GizmoViewContext, MintGreen, FVector(0, 1, 0), FVector(0, 0, 1));
 	NewActor->DownIntervalComponent    = AddDefaultLineHandleComponent(World, NewActor, GizmoViewContext, MintGreen, FVector(0, 1, 0), FVector(0, 0, 1));
 	NewActor->ForwardIntervalComponent = AddDefaultLineHandleComponent(World, NewActor, GizmoViewContext, MintGreen, FVector(1, 0, 0), FVector(0, 1, 0));
-
+	NewActor->BackwardIntervalComponent = AddDefaultLineHandleComponent(World, NewActor, GizmoViewContext, MintGreen, FVector(1, 0, 0), FVector(0, 1, 0));
+	NewActor->RightIntervalComponent = AddDefaultLineHandleComponent(World, NewActor, GizmoViewContext, MintGreen, FVector(0, 0, 1), FVector(1, 0, 0));
+	NewActor->LeftIntervalComponent = AddDefaultLineHandleComponent(World, NewActor, GizmoViewContext, MintGreen, FVector(0, 0, 1), FVector(1, 0, 0));
 
 	return NewActor;
 }
@@ -280,21 +282,14 @@ void UIntervalGizmo::Tick(float DeltaTime)
 		GizmoComponent->SetWorldTransform(TargetTransform);
 	}
 	// Update the lengths
+	EnumerateValidIntervals([this](UGizmoLocalFloatParameterSource* Source, UGizmoLineHandleComponent* Component, UGizmoComponentAxisSource* Axis, UE::Geometry::FInterval1f& IntervalRange, FVector3d Direction, float DirectionAxisSign)
 	{
-		if (GizmoActor->UpIntervalComponent != nullptr && UpIntervalSource != nullptr)
-		{
-			GizmoActor->UpIntervalComponent->Length = UpIntervalSource->GetParameter();
-		}
-		if (GizmoActor->DownIntervalComponent != nullptr && DownIntervalSource != nullptr)
-		{
-			GizmoActor->DownIntervalComponent->Length = DownIntervalSource->GetParameter();
-		}
-		if (GizmoActor->ForwardIntervalComponent != nullptr && ForwardIntervalSource != nullptr)
+		if (Component)
 		{
-			GizmoActor->ForwardIntervalComponent->Length = ForwardIntervalSource->GetParameter();
+			Component->Length = Source->GetParameter();
 		}
-	}
-
+	});
+	
 	if (UpdateCoordSystemFunction)
 	{
 		for (UPrimitiveComponent* Component : ActiveComponents)
@@ -303,14 +298,21 @@ void UIntervalGizmo::Tick(float DeltaTime)
 		}
 	}
 }
+void UIntervalGizmo::SetActiveTarget(UTransformProxy* TransformTargetIn, UGizmoLocalFloatParameterSource* UpInterval, UGizmoLocalFloatParameterSource* DownInterval, UGizmoLocalFloatParameterSource* ForwardInterval, IToolContextTransactionProvider* TransactionProvider)
+{
+	FParameterSources Sources;
+	Sources.UpInterval = UpInterval;
+	Sources.DownInterval = DownInterval;
+	Sources.ForwardInterval = ForwardInterval;
+	SetActiveTarget(TransformTargetIn, Sources, TransactionProvider);
+}
 
-void UIntervalGizmo::SetActiveTarget(UTransformProxy* TransformTargetIn, UGizmoLocalFloatParameterSource* UpInterval , UGizmoLocalFloatParameterSource* DownInterval, UGizmoLocalFloatParameterSource* ForwardInterval, IToolContextTransactionProvider* TransactionProvider)
+void UIntervalGizmo::SetActiveTarget(UTransformProxy* TransformTargetIn, const FParameterSources& ParameterSources, IToolContextTransactionProvider* TransactionProvider)
 {
 	if (TransformProxy != nullptr)
 	{
 		ClearActiveTarget();
 		ClearSources();
-
 	}
 
 	// This state target emits an explicit FChange that moves the GizmoActor root component during undo/redo.
@@ -323,20 +325,56 @@ void UIntervalGizmo::SetActiveTarget(UTransformProxy* TransformTargetIn, UGizmoL
 	TransformProxy = TransformTargetIn;
 
 	// parameters and init lengths for each interval
-	UpIntervalSource      = UpInterval;
-	DownIntervalSource    = DownInterval;
-	ForwardIntervalSource = ForwardInterval;
+	UpIntervalSource      = ParameterSources.UpInterval;
+	DownIntervalSource    = ParameterSources.DownInterval;
+	ForwardIntervalSource = ParameterSources.ForwardInterval;
+	BackwardIntervalSource = ParameterSources.BackwardInterval;
+	RightIntervalSource = ParameterSources.RightInterval;
+	LeftIntervalSource = ParameterSources.LeftInterval;
+
+	if (ParameterSources.InitParameterRanges == EDefaultParameterRanges::HalfRange)
+	{
+		EnumerateAllIntervals([](UGizmoLocalFloatParameterSource* Source, UGizmoLocalFloatParameterSource* OppositeSource, UGizmoLineHandleComponent* Component, UE::Geometry::FInterval1f& IntervalRange, float DirectionAxisSign)
+			{
+				IntervalRange = DirectionAxisSign < 0 ? UE::Geometry::FInterval1f(-FLT_MAX, 0.f) : UE::Geometry::FInterval1f(0.f, FLT_MAX);
+			}
+		);
+	}
+	else if (ParameterSources.InitParameterRanges == EDefaultParameterRanges::FullRange)
+	{
+		EnumerateAllIntervals([](UGizmoLocalFloatParameterSource* Source, UGizmoLocalFloatParameterSource* OppositeSource, UGizmoLineHandleComponent* Component, UE::Geometry::FInterval1f& IntervalRange, float DirectionAxisSign)
+			{
+				IntervalRange = UE::Geometry::FInterval1f(-FLT_MAX, FLT_MAX);
+			}
+		);
+	}
+	else // EDefaultParameterRanges::HalfIfMatched
+	{
+		EnumerateAllIntervals([](UGizmoLocalFloatParameterSource* Source, UGizmoLocalFloatParameterSource* OppositeSource, UGizmoLineHandleComponent* Component, UE::Geometry::FInterval1f& IntervalRange, float DirectionAxisSign)
+			{
+				IntervalRange = UE::Geometry::FInterval1f(-FLT_MAX, FLT_MAX);
+				if (OppositeSource) // if opposite interval exists, cut range in half
+				{
+					if (DirectionAxisSign < 0)
+					{
+						IntervalRange.Max = 0;
+					}
+					else
+					{
+						IntervalRange.Min = 0;
+					}
+				}
+			}
+		);
+	}
 
 	// Get the parameter source to notify our delegate of any changes
-	UpIntervalSource->OnParameterChanged.AddWeakLambda(this, [this](IGizmoFloatParameterSource*, FGizmoFloatParameterChange Change) {
-		OnIntervalChanged.Broadcast(this, FVector::ZAxisVector, Change.CurrentValue);
-		});
-	DownIntervalSource->OnParameterChanged.AddWeakLambda(this, [this](IGizmoFloatParameterSource*, FGizmoFloatParameterChange Change) {
-		OnIntervalChanged.Broadcast(this, -FVector::ZAxisVector, -Change.CurrentValue);
-		});
-	ForwardIntervalSource->OnParameterChanged.AddWeakLambda(this, [this](IGizmoFloatParameterSource*, FGizmoFloatParameterChange Change) {
-		OnIntervalChanged.Broadcast(this, FVector::YAxisVector, Change.CurrentValue);
+	EnumerateValidIntervals([this](UGizmoLocalFloatParameterSource* Source, UGizmoLineHandleComponent* Component, UGizmoComponentAxisSource* Axis, UE::Geometry::FInterval1f& IntervalRange, FVector3d Direction, float DirectionAxisSign)
+	{
+		Source->OnParameterChanged.AddWeakLambda(this, [this, Direction, DirectionAxisSign](IGizmoFloatParameterSource*, FGizmoFloatParameterChange Change) {
+			OnIntervalChanged.Broadcast(this, Direction, DirectionAxisSign * Change.CurrentValue);
 		});
+	});
 
 	USceneComponent* GizmoComponent = GizmoActor->GetRootComponent();
 
@@ -363,43 +401,37 @@ void UIntervalGizmo::SetActiveTarget(UTransformProxy* TransformTargetIn, UGizmoL
 	StateTarget = UGizmoTransformChangeStateTarget::Construct(GizmoComponent,
 		LOCTEXT("UIntervalGizmoTransaction", "Interval"), TransactionProvider, this);
 	StateTarget->DependentChangeSources.Add(MakeUnique<FTransformProxyChangeSource>(TransformProxy));
-	StateTarget->DependentChangeSources.Add(MakeUnique<FGizmoFloatParameterChangeSource>(UpIntervalSource));
-	StateTarget->DependentChangeSources.Add(MakeUnique<FGizmoFloatParameterChangeSource>(DownIntervalSource));
-	StateTarget->DependentChangeSources.Add(MakeUnique<FGizmoFloatParameterChangeSource>(ForwardIntervalSource));
+	EnumerateValidIntervals([this](UGizmoLocalFloatParameterSource* Source, UGizmoLineHandleComponent* Component, UGizmoComponentAxisSource* Axis, UE::Geometry::FInterval1f& IntervalRange, FVector3d Direction, float DirectionAxisSign)
+	{
+		StateTarget->DependentChangeSources.Add(MakeUnique<FGizmoFloatParameterChangeSource>(Source));
+	});
 
 	// Have the state target notify us of the start/end of drags
 	StateTarget->DependentChangeSources.Add(MakeUnique<FIntervalGizmoChangeBroadcaster>(this));
 
-	// root component provides local Y/Z axis, identified by AxisIndex
+	// root component provides local X/Y/Z axis, identified by AxisIndex
+	AxisXSource = UGizmoComponentAxisSource::Construct(GizmoComponent, 0, true, this);
 	AxisYSource = UGizmoComponentAxisSource::Construct(GizmoComponent, 1, true, this);
 	AxisZSource = UGizmoComponentAxisSource::Construct(GizmoComponent, 2, true, this);
-
 	
-	
-	if (GizmoActor->UpIntervalComponent != nullptr)
+	EnumerateAllIntervals([](UGizmoLocalFloatParameterSource* Source, UGizmoLocalFloatParameterSource* OppositeSource, UGizmoLineHandleComponent* Component, UE::Geometry::FInterval1f& IntervalRange, float DirectionAxisSign)
 	{
-		float IntervalMin = 0.f;
-		float IntervalMax = FLT_MAX;
-		AddIntervalHandleGizmo(GizmoComponent, GizmoActor->UpIntervalComponent, AxisZSource, UpIntervalSource, IntervalMin, IntervalMax, StateTarget);
-		ActiveComponents.Add(GizmoActor->UpIntervalComponent);
-	}
-
-	if (GizmoActor->DownIntervalComponent != nullptr)
-	{
-		float IntervalMin = -FLT_MAX;
-		float IntervalMax = 0.f;
-		AddIntervalHandleGizmo(GizmoComponent, GizmoActor->DownIntervalComponent, AxisZSource, DownIntervalSource, IntervalMin, IntervalMax, StateTarget);
-		ActiveComponents.Add(GizmoActor->DownIntervalComponent);
-	}
-
-	if (GizmoActor->ForwardIntervalComponent != nullptr)
+		Component->SetVisibility(false);
+	});
+	EnumerateValidIntervals([this, GizmoComponent](UGizmoLocalFloatParameterSource* Source, UGizmoLineHandleComponent* Component, UGizmoComponentAxisSource* Axis, UE::Geometry::FInterval1f& IntervalRange, FVector3d Direction, float DirectionAxisSign)
 	{
-		float IntervalMin = -FLT_MAX;
-		float IntervalMax = FLT_MAX;
-		AddIntervalHandleGizmo(GizmoComponent, GizmoActor->ForwardIntervalComponent, AxisYSource, ForwardIntervalSource, IntervalMin, IntervalMax, StateTarget);
-		ActiveComponents.Add(GizmoActor->ForwardIntervalComponent);
-	}
+		AddIntervalHandleGizmo(GizmoComponent, Component, Axis, Source, IntervalRange.Min, IntervalRange.Max, StateTarget);
+		ActiveComponents.Add(Component);
+		Component->SetVisibility(true);
+	});
+}
 
+void UIntervalGizmo::SetVisibility(bool bVisible)
+{
+	GizmoActor->SetActorHiddenInGame(bVisible == false);
+#if WITH_EDITOR
+	GizmoActor->SetIsTemporarilyHiddenInEditor(bVisible == false);
+#endif
 }
 
 void UIntervalGizmo::ClearSources()
@@ -407,6 +439,9 @@ void UIntervalGizmo::ClearSources()
 	UpIntervalSource = nullptr;
 	DownIntervalSource = nullptr;
 	ForwardIntervalSource = nullptr;
+	BackwardIntervalSource = nullptr;
+	RightIntervalSource = nullptr;
+	LeftIntervalSource = nullptr;
 }
 
 void UIntervalGizmo::ClearActiveTarget()
@@ -415,8 +450,8 @@ void UIntervalGizmo::ClearActiveTarget()
 	{
 		GetGizmoManager()->DestroyGizmo(Gizmo);
 	}
-	ActiveGizmos.SetNum(0);
-	ActiveComponents.SetNum(0);
+	ActiveGizmos.Empty();
+	ActiveComponents.Empty();
 	
 	ClearSources();
 
@@ -470,6 +505,54 @@ UInteractiveGizmo* UIntervalGizmo::AddIntervalHandleGizmo(
 	return IntervalGizmo;
 }
 
+// Call IterFn on each Source/Component combination where the Source is not null. Note IterFn will still be called if the component and/or axis are null.
+void UIntervalGizmo::EnumerateValidIntervals(
+	TFunctionRef<void(UGizmoLocalFloatParameterSource* Source, UGizmoLineHandleComponent* Component, UGizmoComponentAxisSource* Axis, UE::Geometry::FInterval1f& IntervalRange, FVector3d Direction, float DirectionAxisSign)> IterFn
+)
+{
+	constexpr float IntervalMax = FLT_MAX;
+	if (RightIntervalSource)
+	{
+		IterFn(RightIntervalSource, GizmoActor ? GizmoActor->RightIntervalComponent : nullptr, AxisXSource, RightIntervalRange, FVector3d(1, 0, 0), 1.f);
+	}
+	if (LeftIntervalSource)
+	{
+		IterFn(LeftIntervalSource, GizmoActor ? GizmoActor->LeftIntervalComponent : nullptr, AxisXSource, LeftIntervalRange, FVector3d(-1, 0, 0), -1.f);
+	}
+
+	if (ForwardIntervalSource)
+	{
+		IterFn(ForwardIntervalSource, GizmoActor ? GizmoActor->ForwardIntervalComponent : nullptr, AxisYSource, ForwardIntervalRange, FVector3d(0, 1, 0), 1.f);
+	}
+	if (BackwardIntervalSource)
+	{
+		IterFn(BackwardIntervalSource, GizmoActor ? GizmoActor->BackwardIntervalComponent : nullptr, AxisYSource, BackwardIntervalRange, FVector3d(0, -1, 0), -1.f);
+	}
+
+	if (UpIntervalSource)
+	{
+		IterFn(UpIntervalSource, GizmoActor ? GizmoActor->UpIntervalComponent : nullptr, AxisZSource, UpIntervalRange, FVector3d(0, 0, 1), 1.f);
+	}
+	if (DownIntervalSource)
+	{
+		IterFn(DownIntervalSource, GizmoActor ? GizmoActor->DownIntervalComponent : nullptr, AxisZSource, DownIntervalRange, FVector3d(0, 0, -1), -1.f);
+	}
+}
+
+// Call IterFn on each Source/Component combination, including those where the Source is null
+void UIntervalGizmo::EnumerateAllIntervals(TFunctionRef<void(UGizmoLocalFloatParameterSource* Source, UGizmoLocalFloatParameterSource* OppositeSource, UGizmoLineHandleComponent* Component, UE::Geometry::FInterval1f& IntervalRange, float DirectionAxisSign)> IterFn)
+{
+	if (GizmoActor)
+	{
+		IterFn(RightIntervalSource, LeftIntervalSource, GizmoActor->RightIntervalComponent, RightIntervalRange, 1.f);
+		IterFn(LeftIntervalSource, RightIntervalSource, GizmoActor->LeftIntervalComponent, LeftIntervalRange, -1.f);
+		IterFn(ForwardIntervalSource, BackwardIntervalSource, GizmoActor->ForwardIntervalComponent, ForwardIntervalRange, 1.f);
+		IterFn(BackwardIntervalSource, ForwardIntervalSource, GizmoActor->BackwardIntervalComponent, BackwardIntervalRange, -1.f);
+		IterFn(UpIntervalSource, DownIntervalSource, GizmoActor->UpIntervalComponent, UpIntervalRange, 1.f);
+		IterFn(DownIntervalSource, UpIntervalSource, GizmoActor->DownIntervalComponent, DownIntervalRange, -1.f);
+	}
+}
+
 float UGizmoAxisIntervalParameterSource::GetParameter() const
 {
 	return FloatParameterSource->GetParameter();
diff --git a/Engine/Source/Runtime/InteractiveToolsFramework/Public/BaseGizmos/IntervalGizmo.h b/Engine/Source/Runtime/InteractiveToolsFramework/Public/BaseGizmos/IntervalGizmo.h
--- a/Engine/Source/Runtime/InteractiveToolsFramework/Public/BaseGizmos/IntervalGizmo.h
+++ b/Engine/Source/Runtime/InteractiveToolsFramework/Public/BaseGizmos/IntervalGizmo.h
@@ -10,6 +10,7 @@
 #include "BaseGizmos/GizmoActor.h"
 #include "BaseGizmos/TransformProxy.h"
 #include "ParameterSourcesFloat.h"
+#include "BoxTypes.h"
 #include "IntervalGizmo.generated.h"
 
 class UInteractiveGizmoManager;
@@ -29,7 +30,6 @@ class UGizmoViewContext;
  * FIntervalGizmoActorFactory returns an instance of this Actor type (or a subclass).
  *
  * If a particular sub-Gizmo is not required, simply set that UProperty to null.
-
  */
 UCLASS(Transient, NotPlaceable, Hidden, NotBlueprintable, NotBlueprintType, MinimalAPI)
 class AIntervalGizmoActor : public AGizmoActor
@@ -50,11 +50,19 @@ public:
 	UPROPERTY()
 	TObjectPtr<UGizmoLineHandleComponent> ForwardIntervalComponent = nullptr;
 
+	UPROPERTY()
+	TObjectPtr<UGizmoLineHandleComponent> BackwardIntervalComponent = nullptr;
+
+	UPROPERTY()
+	TObjectPtr<UGizmoLineHandleComponent> RightIntervalComponent = nullptr;
+
+	UPROPERTY()
+	TObjectPtr<UGizmoLineHandleComponent> LeftIntervalComponent = nullptr;
 
 public:
 	/**
 	 * Create a new instance of AIntervalGizmoActor and populate the various
-	 * sub-components with standard GizmoXComponent instances suitable for a 3-interval Gizmo
+	 * sub-components with standard GizmoXComponent instances
 	 */
 	static INTERACTIVETOOLSFRAMEWORK_API AIntervalGizmoActor* ConstructDefaultIntervalGizmo(UWorld* World, 
 		UGizmoViewContext* GizmoViewContext);
@@ -156,9 +164,37 @@ public:
 	INTERACTIVETOOLSFRAMEWORK_API virtual void Shutdown() override;
 	INTERACTIVETOOLSFRAMEWORK_API virtual void Tick(float DeltaTime) override;
 
+	enum class EDefaultParameterRanges
+	{
+		// Init all interval ranges to [0, +/- FLT_MAX] in the axis direction
+		HalfRange,
+		// Init all interval ranges to [-FLT_MAX, FLT_MAX]
+		FullRange,
+		// Init interval ranges to 'half range' if 'matched' by the opposite interval on the same axis, or 'full range' otherwise
+		HalfIfMatched
+	};
+
+	// Input parameter sources and related configuration
+	struct FParameterSources
+	{
+		// +/- X axis parameter sources
+		UGizmoLocalFloatParameterSource* RightInterval = nullptr;
+		UGizmoLocalFloatParameterSource* LeftInterval = nullptr;
+		// +/- Y axis parameter sources
+		UGizmoLocalFloatParameterSource* ForwardInterval = nullptr;
+		UGizmoLocalFloatParameterSource* BackwardInterval = nullptr;
+		// +/- Z axis parameter sources
+		UGizmoLocalFloatParameterSource* UpInterval = nullptr;
+		UGizmoLocalFloatParameterSource* DownInterval = nullptr;
+
+		EDefaultParameterRanges InitParameterRanges = EDefaultParameterRanges::HalfIfMatched;
+		
+	};
 
 	INTERACTIVETOOLSFRAMEWORK_API virtual void SetActiveTarget(UTransformProxy* TransformTargetIn, UGizmoLocalFloatParameterSource* UpInterval, UGizmoLocalFloatParameterSource* DownInterval, UGizmoLocalFloatParameterSource* ForwardInterval,
 		                         IToolContextTransactionProvider* TransactionProvider = nullptr);
+	INTERACTIVETOOLSFRAMEWORK_API virtual void SetActiveTarget(UTransformProxy* TransformTargetIn, const FParameterSources& ParameterSources,
+		IToolContextTransactionProvider* TransactionProvider = nullptr);
 
 	/** Sets functions that allow the endpoints of the intervals to be snapped to world geometry when ShouldAlignDestination is true */
 	INTERACTIVETOOLSFRAMEWORK_API virtual void SetWorldAlignmentFunctions(
@@ -185,6 +221,9 @@ public:
 	 */
 	INTERACTIVETOOLSFRAMEWORK_API virtual FTransform GetGizmoTransform() const;
 
+	/** Set visibility of the Gizmo */
+	INTERACTIVETOOLSFRAMEWORK_API virtual void SetVisibility(bool bVisible);
+
 	/** State target is shared across gizmos, and created internally during SetActiveTarget() */
 	UPROPERTY()
 	TObjectPtr<UGizmoTransformChangeStateTarget> StateTarget = nullptr;
@@ -206,9 +245,11 @@ public:
 protected:
 
 	/** GizmoActors will be spawned in this World */
-	UWorld* World;
+	UPROPERTY()
+	TObjectPtr<UWorld> World;
 
-	AIntervalGizmoActor* GizmoActor = nullptr;
+	UPROPERTY()
+	TObjectPtr<AIntervalGizmoActor> GizmoActor = nullptr;
 
 	/** The gizmo tracks the location and orientation of the transform in this TransformProxy. */
 	UPROPERTY()
@@ -222,10 +263,25 @@ protected:
 	UPROPERTY()
 	TArray<TObjectPtr<UInteractiveGizmo>> ActiveGizmos;
 
-	UGizmoLocalFloatParameterSource* UpIntervalSource = nullptr;
-	UGizmoLocalFloatParameterSource* DownIntervalSource = nullptr;
-	UGizmoLocalFloatParameterSource* ForwardIntervalSource = nullptr;
-
+	// +/- Z axis parameter sources
+	UPROPERTY()
+	TObjectPtr<UGizmoLocalFloatParameterSource> UpIntervalSource = nullptr;
+	UPROPERTY()
+	TObjectPtr<UGizmoLocalFloatParameterSource> DownIntervalSource = nullptr;
+	// +/- Y axis parameter sources
+	UPROPERTY()
+	TObjectPtr<UGizmoLocalFloatParameterSource> ForwardIntervalSource = nullptr;
+	UPROPERTY()
+	TObjectPtr<UGizmoLocalFloatParameterSource> BackwardIntervalSource = nullptr;
+	// +/- X axis parameter sources
+	UPROPERTY()
+	TObjectPtr<UGizmoLocalFloatParameterSource> RightIntervalSource = nullptr;
+	UPROPERTY()
+	TObjectPtr<UGizmoLocalFloatParameterSource> LeftIntervalSource = nullptr;
+	
+	/** X-axis source is shared across Gizmos, and created internally during SetActiveTarget() */
+	UPROPERTY()
+	TObjectPtr<UGizmoComponentAxisSource> AxisXSource = nullptr;
 	
 	/** Y-axis source is shared across Gizmos, and created internally during SetActiveTarget() */
 	UPROPERTY()
@@ -262,6 +318,16 @@ protected:
 		float MaxParameter,
 		IGizmoStateTarget* StateTargetIn);
 
+private:
+
+	// Valid ranges for intervals
+	UE::Geometry::FInterval1f UpIntervalRange, DownIntervalRange, ForwardIntervalRange, BackwardIntervalRange, RightIntervalRange, LeftIntervalRange;
+
+	// Call IterFn on each Source/Component combination where the Source is not null. Note IterFn will still be called if the component and/or axis are null.
+	void EnumerateValidIntervals(TFunctionRef<void(UGizmoLocalFloatParameterSource* Source, UGizmoLineHandleComponent* Component, UGizmoComponentAxisSource* Axis, UE::Geometry::FInterval1f& IntervalRange, FVector3d Direction, float DirectionAxisSign)> IterFn);
+
+	// Call IterFn on each interval Component (including those without a Source)
+	void EnumerateAllIntervals(TFunctionRef<void(UGizmoLocalFloatParameterSource* Source, UGizmoLocalFloatParameterSource* OppositeSource, UGizmoLineHandleComponent* Component, UE::Geometry::FInterval1f& IntervalRange, float DirectionAxisSign)> IterFn);
 };
 
 /**
