PR #12314: Modeling Tools: Add "Hit Back Faces" option for mesh selections

diff --git a/Engine/Plugins/Editor/ModelingToolsEditorMode/Source/ModelingToolsEditorMode/Private/ModelingModeToolkit_Toolbars.cpp b/Engine/Plugins/Editor/ModelingToolsEditorMode/Source/ModelingToolsEditorMode/Private/ModelingModeToolkit_Toolbars.cpp
--- a/Engine/Plugins/Editor/ModelingToolsEditorMode/Source/ModelingToolsEditorMode/Private/ModelingModeToolkit_Toolbars.cpp
+++ b/Engine/Plugins/Editor/ModelingToolsEditorMode/Source/ModelingToolsEditorMode/Private/ModelingModeToolkit_Toolbars.cpp
@@ -407,6 +407,18 @@ void MakeSubMenu_Selection_LocalFrameMode(FMenuBuilder& MenuBuilder)
 	);
 }
 
+
+void MakeSubMenu_Selection_SelectionFilters(FMenuBuilder& MenuBuilder)
+{
+	const FModelingToolsManagerCommands& Commands = FModelingToolsManagerCommands::Get();
+	MenuBuilder.AddMenuEntry(
+		Commands.SelectionHitBackFaces,
+		NAME_None,
+		LOCTEXT("Selection_SelectionFilters_HitBackFaces", "Hit Back Faces"), 
+		LOCTEXT("Selection_SelectionFilters_HitBackFaces_Tooltip", "Hit triangles that are facing away from the camera"),
+		FSlateIcon(), NAME_None);
+}
+
 TSharedRef<SWidget> MakeMenu_SelectionConfigSettings(const FModelingToolsEditorModeToolkit* Toolkit)
 {
 	FMenuBuilder MenuBuilder(true, Toolkit->GetToolkitCommands());
@@ -423,6 +435,10 @@ TSharedRef<SWidget> MakeMenu_SelectionConfigSettings(const FModelingToolsEditorM
 	MakeSubMenu_Selection_LocalFrameMode(MenuBuilder);
 	MenuBuilder.EndSection();
 
+	MenuBuilder.BeginSection("Section_SelectionFilters", LOCTEXT("Section_SelectionFilters", "Selection Filters"));
+	MakeSubMenu_Selection_SelectionFilters(MenuBuilder);
+	MenuBuilder.EndSection();
+
 	TSharedRef<SWidget> MenuWidget = MenuBuilder.MakeWidget();
 	return MenuWidget;
 }
diff --git a/Engine/Plugins/Editor/ModelingToolsEditorMode/Source/ModelingToolsEditorMode/Private/ModelingToolsEditorMode.cpp b/Engine/Plugins/Editor/ModelingToolsEditorMode/Source/ModelingToolsEditorMode/Private/ModelingToolsEditorMode.cpp
--- a/Engine/Plugins/Editor/ModelingToolsEditorMode/Source/ModelingToolsEditorMode/Private/ModelingToolsEditorMode.cpp
+++ b/Engine/Plugins/Editor/ModelingToolsEditorMode/Source/ModelingToolsEditorMode/Private/ModelingToolsEditorMode.cpp
@@ -955,6 +955,22 @@ void UModelingToolsEditorMode::Enter()
 	);
 	CommandList->MapAction(ToolManagerCommands.SelectionLocking, ToggleSelectionLocking);
 
+	// Register Selection Hit Back Faces Toggle
+	FUIAction ToggleHitBackFacesAction = FUIAction(
+		FExecuteAction::CreateLambda([this]
+		{
+			UModelingToolsModeCustomizationSettings* ModelingEditorSettings = GetMutableDefault<UModelingToolsModeCustomizationSettings>();
+			ModelingEditorSettings->bMeshSelectionHitBackFaces = !ModelingEditorSettings->bMeshSelectionHitBackFaces;
+			SelectionManager->SetHitBackFaces(ModelingEditorSettings->bMeshSelectionHitBackFaces);
+			ModelingEditorSettings->SaveConfig();
+		}),
+		FCanExecuteAction(),
+		FIsActionChecked::CreateLambda([this]()
+		{
+			return SelectionManager->GetHitBackFaces();
+		}));
+	CommandList->MapAction(ToolManagerCommands.SelectionHitBackFaces, ToggleHitBackFacesAction);
+
 	// register extensions
 	TArray<IModelingModeToolExtension*> Extensions = IModularFeatures::Get().GetModularFeatureImplementations<IModelingModeToolExtension>(
 		IModelingModeToolExtension::GetModularFeatureName());
@@ -1473,6 +1489,15 @@ void UModelingToolsEditorMode::OnBlueprintPreCompile(UBlueprint* Blueprint)
 							FModelingToolsManagerCommands::Get().SelectionLocalFrameMode_Object, CommandList));
 				}
 
+				// Section for Selection Filters
+				{
+					FToolMenuSection& SelectionFilterSection =
+						Submenu->FindOrAddSection("Selection Filters", LOCTEXT("SelectionFiltersLabel", "Selection Filters"));
+					SelectionFilterSection.AddEntry(
+							FToolMenuEntry::InitMenuEntryWithCommandList(
+								FModelingToolsManagerCommands::Get().SelectionHitBackFaces, CommandList));
+				}
+
 				// Section for locking/unlocking target
 				{
 					FToolMenuSection& LockingSection =
diff --git a/Engine/Plugins/Editor/ModelingToolsEditorMode/Source/ModelingToolsEditorMode/Private/ModelingToolsManagerActions.cpp b/Engine/Plugins/Editor/ModelingToolsEditorMode/Source/ModelingToolsEditorMode/Private/ModelingToolsManagerActions.cpp
--- a/Engine/Plugins/Editor/ModelingToolsEditorMode/Source/ModelingToolsEditorMode/Private/ModelingToolsManagerActions.cpp
+++ b/Engine/Plugins/Editor/ModelingToolsEditorMode/Source/ModelingToolsEditorMode/Private/ModelingToolsManagerActions.cpp
@@ -263,6 +263,7 @@ void FModelingToolsManagerCommands::RegisterCommands()
 	REGISTER_MODELING_TOOL_COMMAND(SelectionMeshTypes_Volumes, "Volumes", "Toggle whether Volume mesh elements can be selected in the Viewport");
 	REGISTER_MODELING_TOOL_COMMAND(SelectionMeshTypes_StaticMeshes, "Static Meshes", "Toggle whether Static Mesh mesh elements can be selected in the Viewport");
 
+	REGISTER_MODELING_TOOL_COMMAND(SelectionHitBackFaces, "Hit Back Faces", "Hit triangles that are facing away from the camera");
 	REGISTER_MODELING_TOOL_COMMAND(SelectionLocking, "Lock Target", "Click to Toggle locking on the Selected Object and disallow/allow Mesh Selections");
 	
 	UI_COMMAND(AddToFavorites, "Add to Favorites", "Add to Favorites", EUserInterfaceActionType::Button, FInputChord());
diff --git a/Engine/Plugins/Editor/ModelingToolsEditorMode/Source/ModelingToolsEditorMode/Public/ModelingToolsEditorModeSettings.h b/Engine/Plugins/Editor/ModelingToolsEditorMode/Source/ModelingToolsEditorMode/Public/ModelingToolsEditorModeSettings.h
--- a/Engine/Plugins/Editor/ModelingToolsEditorMode/Source/ModelingToolsEditorMode/Public/ModelingToolsEditorModeSettings.h
+++ b/Engine/Plugins/Editor/ModelingToolsEditorMode/Source/ModelingToolsEditorMode/Public/ModelingToolsEditorModeSettings.h
@@ -388,6 +388,9 @@ public:
 
 	UPROPERTY(config)
 	int32 LastMeshSelectionLocalFrameMode = 0;
+	
+	UPROPERTY(config)
+	bool bMeshSelectionHitBackFaces = true;
 
 public:
 
diff --git a/Engine/Plugins/Editor/ModelingToolsEditorMode/Source/ModelingToolsEditorMode/Public/ModelingToolsManagerActions.h b/Engine/Plugins/Editor/ModelingToolsEditorMode/Source/ModelingToolsEditorMode/Public/ModelingToolsManagerActions.h
--- a/Engine/Plugins/Editor/ModelingToolsEditorMode/Source/ModelingToolsEditorMode/Public/ModelingToolsManagerActions.h
+++ b/Engine/Plugins/Editor/ModelingToolsEditorMode/Source/ModelingToolsEditorMode/Public/ModelingToolsManagerActions.h
@@ -201,6 +201,7 @@ public:
 	TSharedPtr<FUICommandInfo> SelectionMeshTypes_Volumes;
 	TSharedPtr<FUICommandInfo> SelectionMeshTypes_StaticMeshes;
 	TSharedPtr<FUICommandInfo> SelectionLocking;
+	TSharedPtr<FUICommandInfo> SelectionHitBackFaces;
 
 	TSharedPtr<FUICommandInfo> AddToFavorites;
 	TSharedPtr<FUICommandInfo> RemoveFromFavorites;
diff --git a/Engine/Plugins/Runtime/GeometryProcessing/Source/DynamicMesh/Private/Selections/GeometrySelectionUtil.cpp b/Engine/Plugins/Runtime/GeometryProcessing/Source/DynamicMesh/Private/Selections/GeometrySelectionUtil.cpp
--- a/Engine/Plugins/Runtime/GeometryProcessing/Source/DynamicMesh/Private/Selections/GeometrySelectionUtil.cpp
+++ b/Engine/Plugins/Runtime/GeometryProcessing/Source/DynamicMesh/Private/Selections/GeometrySelectionUtil.cpp
@@ -272,8 +272,15 @@ void UE::Geometry::UpdateTriangleSelectionViaRaycast(
 
 	ResultOut.bSelectionMissed = true;
 
+	IMeshSpatial::FQueryOptions SpatialQueryOptions;
+	if (!Editor->GetQueryConfig().bHitBackFaces)
+	{
+		SpatialQueryOptions.TriangleFilterF = [&](int32 Tid) {
+			return ColliderMesh->GetTriNormal(Tid).Dot(LocalRay.Direction) < 0;
+		};
+	}
 	double RayHitT; int32 HitTriangleID; FVector3d HitBaryCoords;
-	if (ColliderMesh->FindNearestHitTriangle(LocalRay, RayHitT, HitTriangleID, HitBaryCoords))
+	if (ColliderMesh->FindNearestHitTriangle(LocalRay, RayHitT, HitTriangleID, HitBaryCoords, SpatialQueryOptions))
 	{
 		HitTriangleID = ColliderMesh->GetSourceTriangleID(HitTriangleID);
 		if (HitTriangleID == IndexConstants::InvalidID)
@@ -346,8 +353,15 @@ void UE::Geometry::UpdateGroupSelectionViaRaycast(
 
 	ResultOut.bSelectionMissed = true;
 
+	IMeshSpatial::FQueryOptions SpatialQueryOptions;
+	if (!Editor->GetQueryConfig().bHitBackFaces)
+	{
+		SpatialQueryOptions.TriangleFilterF = [&](int32 Tid) {
+			return ColliderMesh->GetTriNormal(Tid).Dot(LocalRay.Direction) < 0;
+		};
+	}
 	double RayHitT; int32 HitTriangleID; FVector3d HitBaryCoords;
-	if (ColliderMesh->FindNearestHitTriangle(LocalRay, RayHitT, HitTriangleID, HitBaryCoords))
+	if (ColliderMesh->FindNearestHitTriangle(LocalRay, RayHitT, HitTriangleID, HitBaryCoords, SpatialQueryOptions))
 	{
 		HitTriangleID = ColliderMesh->GetSourceTriangleID(HitTriangleID);
 		if (HitTriangleID == IndexConstants::InvalidID)
diff --git a/Engine/Plugins/Runtime/GeometryProcessing/Source/DynamicMesh/Public/Selections/GeometrySelection.h b/Engine/Plugins/Runtime/GeometryProcessing/Source/DynamicMesh/Public/Selections/GeometrySelection.h
--- a/Engine/Plugins/Runtime/GeometryProcessing/Source/DynamicMesh/Public/Selections/GeometrySelection.h
+++ b/Engine/Plugins/Runtime/GeometryProcessing/Source/DynamicMesh/Public/Selections/GeometrySelection.h
@@ -212,6 +212,12 @@ struct DYNAMICMESH_API FGeometrySelectionHitQueryConfig
 	
 	/** If true, only "visible" elements are considered, otherwise query may return obscured elements (eg, hidden edges occluded by mesh surface) */
 	bool bOnlyVisible = true;
+
+	/**
+	 * When false, triangles that are facing away from the camera are not considered.
+	 * This can be useful when working with inside-out meshes, where we usually want to select the farther, visible wall.
+	 */
+	bool bHitBackFaces = false;
 };
 
 
diff --git a/Engine/Plugins/Runtime/MeshModelingToolset/Source/ModelingComponents/Private/Selection/DynamicMeshSelector.cpp b/Engine/Plugins/Runtime/MeshModelingToolset/Source/ModelingComponents/Private/Selection/DynamicMeshSelector.cpp
--- a/Engine/Plugins/Runtime/MeshModelingToolset/Source/ModelingComponents/Private/Selection/DynamicMeshSelector.cpp
+++ b/Engine/Plugins/Runtime/MeshModelingToolset/Source/ModelingComponents/Private/Selection/DynamicMeshSelector.cpp
@@ -220,10 +220,17 @@ bool FBaseDynamicMeshSelector::RayHitTest(
 
 	// todo: do we need to sometimes ignore surface hits here? maybe for Volumes in (eg) edge mode...
 
+	IMeshSpatial::FQueryOptions SpatialQueryOptions;
+	if (!QueryConfig.bHitBackFaces)
+	{
+		SpatialQueryOptions.TriangleFilterF = [&](int32 Tid) {
+			return GetColliderMesh()->GetTriNormal(Tid).Dot(RayInfo.WorldRay.Direction) < 0;
+		};
+	}
 	FTransformSRT3d WorldTransform = GetWorldTransformFunc();
 	FRay3d LocalRay = WorldTransform.InverseTransformRay(RayInfo.WorldRay);
 	double RayHitT; int32 HitTriangleID; FVector3d HitBaryCoords;
-	if (GetColliderMesh()->FindNearestHitTriangle(LocalRay, RayHitT, HitTriangleID, HitBaryCoords))
+	if (GetColliderMesh()->FindNearestHitTriangle(LocalRay, RayHitT, HitTriangleID, HitBaryCoords, SpatialQueryOptions))
 	{
 		HitResultOut.bHit = true;
 		HitResultOut.HitIdentifier = HitTriangleID;
@@ -351,14 +358,22 @@ void FBaseDynamicMeshSelector::UpdateSelectionViaRaycast_GroupEdges(
 		return ToolSceneQueriesUtil::PointSnapQuery(RayInfo.CameraState,
 			WorldTransform.TransformPosition(A), WorldTransform.TransformPosition(B));
 	};
-
+	
+	IMeshSpatial::FQueryOptions SpatialQueryOptions;
+	if (!SelectionEditor.GetQueryConfig().bHitBackFaces)
+	{
+		SpatialQueryOptions.TriangleFilterF = [&](int32 Tid) {
+			return GetColliderMesh()->GetTriNormal(Tid).Dot(RayInfo.WorldRay.Direction) < 0;
+		};
+	}
+	
 	if (SelectionEditor.GetQueryConfig().bOnlyVisible)
 	{
 		// Compute mesh hit and then find nearest-segment to the mesh hit point. This filters out 'backface' hits
 		// and avoids complexity of determing "nearest visible" segment. However this will not work for segments
 		// on the mesh silhouette, or if the mesh surface is hidden, so this path should be optional...maybe controlled by the the SelectionEditor?
 		double RayHitT; int32 HitTriangleID; FVector3d HitBaryCoords;
-		bool bSurfaceMeshHit = GetColliderMesh()->FindNearestHitTriangle(LocalRay, RayHitT, HitTriangleID, HitBaryCoords);
+		bool bSurfaceMeshHit = GetColliderMesh()->FindNearestHitTriangle(LocalRay, RayHitT, HitTriangleID, HitBaryCoords, SpatialQueryOptions);
 		if (bSurfaceMeshHit)
 		{
 			FVector3d SurfaceHitPos = LocalRay.PointAt(RayHitT);
@@ -387,7 +402,7 @@ void FBaseDynamicMeshSelector::UpdateSelectionViaRaycast_GroupEdges(
 	// find nearest segment to ray
 	FSegmentTree3::FSegment Segment;
 	FSegmentTree3::FRayNearestSegmentInfo SegmentInfo;
-	if (GetGroupEdgeSpatial()->FindNearestVisibleSegmentHitByRay(LocalRay, EdgeHitToleranceTest, Segment, SegmentInfo))
+	if (GetGroupEdgeSpatial()->FindNearestVisibleSegmentHitByRay(LocalRay, EdgeHitToleranceTest, Segment, SegmentInfo, SpatialQueryOptions))
 	{
 		int32 GroupEdgeID = GetGroupTopology()->FindGroupEdgeID(Segment.ID);
 		if (GroupEdgeID >= 0)
@@ -432,6 +447,7 @@ void FBaseDynamicMeshSelector::UpdateSelectionViaShape(
 	FTransformSRT3d WorldTransform = GetWorldTransformFunc();
 	TArray<uint64> InsideElements;
 
+	// todo: support bHitBackFaces
 	// todo: this should be made more efficient using BVHs...
 
 	if ( SelectionEditor.GetTopologyType() == EGeometryTopologyType::Triangle )
diff --git a/Engine/Plugins/Runtime/MeshModelingToolset/Source/ModelingComponents/Private/Selection/GeometrySelectionManager.cpp b/Engine/Plugins/Runtime/MeshModelingToolset/Source/ModelingComponents/Private/Selection/GeometrySelectionManager.cpp
--- a/Engine/Plugins/Runtime/MeshModelingToolset/Source/ModelingComponents/Private/Selection/GeometrySelectionManager.cpp
+++ b/Engine/Plugins/Runtime/MeshModelingToolset/Source/ModelingComponents/Private/Selection/GeometrySelectionManager.cpp
@@ -390,17 +390,29 @@ EGeometryTopologyType UGeometrySelectionManager::GetSelectionTopologyType() cons
 }
 
 
-
-UE::Geometry::FGeometrySelectionHitQueryConfig UGeometrySelectionManager::GetCurrentSelectionQueryConfig() const
+FGeometrySelectionHitQueryConfig UGeometrySelectionManager::GetCurrentSelectionQueryConfig() const
 {
 	FGeometrySelectionHitQueryConfig HitQueryConfig;
 	HitQueryConfig.TopologyType = GetSelectionTopologyType();
 	HitQueryConfig.ElementType = GetSelectionElementType();
 	HitQueryConfig.bOnlyVisible = true;
+	HitQueryConfig.bHitBackFaces = GetHitBackFaces();
 	return HitQueryConfig;
 }
 
 
+void UGeometrySelectionManager::SetHitBackFaces(const bool bInHitBackFaces)
+{
+	bHitBackFaces = bInHitBackFaces;
+	for (const TSharedPtr<FGeometrySelectionTarget>& Target : ActiveTargetReferences)
+	{
+		FGeometrySelectionHitQueryConfig TargetQueryConfig = Target->SelectionEditor->GetQueryConfig();
+		TargetQueryConfig.bHitBackFaces = bHitBackFaces;
+		Target->SelectionEditor->UpdateQueryConfig(TargetQueryConfig, Target->SelectionEditor->GetIsTopologyIDFilteringEnabled());
+	}
+}
+
+
 bool UGeometrySelectionManager::HasSelection() const
 {
 	for (TSharedPtr<FGeometrySelectionTarget> Target : ActiveTargetReferences)
diff --git a/Engine/Plugins/Runtime/MeshModelingToolset/Source/ModelingComponents/Public/Selection/GeometrySelectionManager.h b/Engine/Plugins/Runtime/MeshModelingToolset/Source/ModelingComponents/Public/Selection/GeometrySelectionManager.h
--- a/Engine/Plugins/Runtime/MeshModelingToolset/Source/ModelingComponents/Public/Selection/GeometrySelectionManager.h
+++ b/Engine/Plugins/Runtime/MeshModelingToolset/Source/ModelingComponents/Public/Selection/GeometrySelectionManager.h
@@ -453,6 +453,14 @@ protected:
 
 	UE::Geometry::FGeometrySelectionHitQueryConfig GetCurrentSelectionQueryConfig() const;
 
+public:
+	// Selection Filters
+	void SetHitBackFaces(bool bInHitBackFaces);
+	bool GetHitBackFaces() const { return bHitBackFaces; }
+	
+protected:
+	bool bHitBackFaces = true;
+	
 	// ITF references
 
 	UPROPERTY()
