add error message for scene component conversion due to an unsupported component type

diff --git a/Engine/Plugins/Runtime/MeshModelingToolset/Source/ModelingComponents/Private/ConversionUtils/SceneComponentToDynamicMesh.cpp b/Engine/Plugins/Runtime/MeshModelingToolset/Source/ModelingComponents/Private/ConversionUtils/SceneComponentToDynamicMesh.cpp
--- a/Engine/Plugins/Runtime/MeshModelingToolset/Source/ModelingComponents/Private/ConversionUtils/SceneComponentToDynamicMesh.cpp
+++ b/Engine/Plugins/Runtime/MeshModelingToolset/Source/ModelingComponents/Private/ConversionUtils/SceneComponentToDynamicMesh.cpp
@@ -33,6 +33,7 @@
 #include "StaticMeshAttributes.h"
 #include "StaticMeshLODResourcesToDynamicMesh.h"
 #include "StaticMeshOperations.h"
+#include "Internationalization/Text.h"
 
 #define LOCTEXT_NAMESPACE "ModelingComponents_SceneComponentToDynamicMesh"
 
@@ -819,6 +820,10 @@ bool SceneComponentToDynamicMesh(USceneComponent* Component, const FToMeshOption
 			OutErrorMessage = LOCTEXT("CopyMeshFromComponent_MissingRestCollection", "GeometryCollectionComponent has null Rest Collection object");
 		}
 	}
+	else
+	{
+		OutErrorMessage = FText::FormatOrdered(LOCTEXT("CopyMeshFromComponent_UnsupportedComponentType", "Scene Component \"{0}\" has unsupported type"), FText::FromName(Component->GetFName()));
+	}
 
 	// transform mesh to world
 	if (bSuccess && bTransformToWorld)
