ModelingTools: Fixed crash when running the Subdivide tool using the Loop subdivision scheme on a volume.

diff --git a/Engine/Plugins/Runtime/MeshModelingToolset/Source/ModelingComponentsEditorOnly/Private/Operations/SubdividePoly.cpp b/Engine/Plugins/Runtime/MeshModelingToolset/Source/ModelingComponentsEditorOnly/Private/Operations/SubdividePoly.cpp
--- a/Engine/Plugins/Runtime/MeshModelingToolset/Source/ModelingComponentsEditorOnly/Private/Operations/SubdividePoly.cpp
+++ b/Engine/Plugins/Runtime/MeshModelingToolset/Source/ModelingComponentsEditorOnly/Private/Operations/SubdividePoly.cpp
@@ -525,7 +525,12 @@ bool FSubdividePoly::ComputeTopologySubdivision()
 		Descriptor.numVertsPerFace = NumVertsPerFace.GetData();
 		Descriptor.vertIndicesPerFace = BoundaryVertsPerFace.GetData();
 
-		if (UVComputationMethod == ESubdivisionOutputUVs::Interpolated)
+		const bool bShouldInterpolateUVs = (UVComputationMethod == ESubdivisionOutputUVs::Interpolated)
+			&& OriginalMesh.HasAttributes() 
+			&& OriginalMesh.Attributes()->PrimaryUV() != nullptr
+			&& OriginalMesh.Attributes()->PrimaryUV()->ElementCount() > 0;
+
+		if (bShouldInterpolateUVs)
 		{
 			UVChannel.numValues = BoundaryVertsPerFace.Num();
 
@@ -581,9 +586,10 @@ bool FSubdividePoly::ComputeTopologySubdivision()
 		Descriptor.numVertsPerFace = NumVertsPerFace.GetData();
 		Descriptor.vertIndicesPerFace = BoundaryVertsPerFace.GetData();
 
-		const bool bShouldInterpolateUVs = (UVComputationMethod == ESubdivisionOutputUVs::Interpolated) &&
-			OriginalMesh.HasAttributes() && 
-			(OriginalMesh.Attributes()->PrimaryUV() != nullptr);
+		const bool bShouldInterpolateUVs = (UVComputationMethod == ESubdivisionOutputUVs::Interpolated)
+			&& OriginalMesh.HasAttributes() 
+			&& OriginalMesh.Attributes()->PrimaryUV() != nullptr
+			&& OriginalMesh.Attributes()->PrimaryUV()->ElementCount() > 0;
 
 		if (bShouldInterpolateUVs)
 		{
@@ -833,9 +839,10 @@ bool FSubdividePoly::ComputeSubdividedMesh(FDynamicMesh3& OutMesh)
 	// 
 	TArray<SubdividePolyLocal::SubdUVVertex> RefinedUVs;
 
-	const bool bShouldInterpolateUVs = (UVComputationMethod == ESubdivisionOutputUVs::Interpolated) &&
-		OriginalMesh.HasAttributes() &&
-		(OriginalMesh.Attributes()->PrimaryUV() != nullptr);
+	const bool bShouldInterpolateUVs = (UVComputationMethod == ESubdivisionOutputUVs::Interpolated)
+		&& OriginalMesh.HasAttributes()
+		&& OriginalMesh.Attributes()->PrimaryUV() != nullptr
+		&& OriginalMesh.Attributes()->PrimaryUV()->ElementCount() > 0;
 
 	if (bShouldInterpolateUVs)
 	{
