ModelingTools: Fixed crash when accepting the Duplicate tool on a volume.

diff --git a/Engine/Plugins/Runtime/MeshModelingToolset/Source/MeshModelingTools/Private/CombineMeshesTool.cpp b/Engine/Plugins/Runtime/MeshModelingToolset/Source/MeshModelingTools/Private/CombineMeshesTool.cpp
--- a/Engine/Plugins/Runtime/MeshModelingToolset/Source/MeshModelingTools/Private/CombineMeshesTool.cpp
+++ b/Engine/Plugins/Runtime/MeshModelingToolset/Source/MeshModelingTools/Private/CombineMeshesTool.cpp
@@ -279,10 +279,12 @@ void UCombineMeshesTool::CreateNewAsset()
 			}
 
 			// update material IDs to account for combined material set
-			FDynamicMeshMaterialAttribute* MatAttrib = ComponentDMesh.Attributes()->GetMaterialID();
-			for (int TID : ComponentDMesh.TriangleIndicesItr())
+			if (FDynamicMeshMaterialAttribute* MatAttrib = ComponentDMesh.Attributes()->GetMaterialID())
 			{
-				SetNewMaterialID(ComponentIdx, MatAttrib, TID, MaterialIDRemaps, AllMaterials);
+				for (int TID : ComponentDMesh.TriangleIndicesItr())
+				{
+					SetNewMaterialID(ComponentIdx, MatAttrib, TID, MaterialIDRemaps, AllMaterials);
+				}
 			}
 
 			FDynamicMeshEditor Editor(&AccumulateDMesh);
