ModelingTools: Fixed incorrect action tooltips in the Edit Attributes tool.

diff --git a/Engine/Plugins/Runtime/MeshModelingToolset/Source/MeshModelingToolsEditorOnly/Public/AttributeEditorTool.h b/Engine/Plugins/Runtime/MeshModelingToolset/Source/MeshModelingToolsEditorOnly/Public/AttributeEditorTool.h
--- a/Engine/Plugins/Runtime/MeshModelingToolset/Source/MeshModelingToolsEditorOnly/Public/AttributeEditorTool.h
+++ b/Engine/Plugins/Runtime/MeshModelingToolset/Source/MeshModelingToolsEditorOnly/Public/AttributeEditorTool.h
@@ -172,7 +172,7 @@ public:
 		PostAction(EAttributeEditorToolActions::ClearAllUVs);
 	}
 
-	/** Clear the selected UV layers, setting all UV values to (0,0) */
+	/** Add a new UV layer */
 	UFUNCTION(CallInEditor, Category = UVs, meta = (DisplayPriority = 2))
 	void AddNew()
 	{
@@ -186,7 +186,7 @@ public:
 		PostAction(EAttributeEditorToolActions::DeleteSelectedUVSet);
 	}
 
-	/** Clear the selected UV layers, setting all UV values to (0,0) */
+	/** Duplicate the selected UV layers */
 	UFUNCTION(CallInEditor, Category = UVs, meta = (DisplayPriority = 4))
 	void DuplicateSelected()
 	{
