Avoid blocking UStaticMesh builds when checking validity of the static mesh tool target.

diff --git a/Engine/Plugins/Runtime/MeshModelingToolset/Source/ModelingComponentsEditorOnly/Private/ToolTargets/StaticMeshToolTarget.cpp b/Engine/Plugins/Runtime/MeshModelingToolset/Source/ModelingComponentsEditorOnly/Private/ToolTargets/StaticMeshToolTarget.cpp
--- a/Engine/Plugins/Runtime/MeshModelingToolset/Source/ModelingComponentsEditorOnly/Private/ToolTargets/StaticMeshToolTarget.cpp
+++ b/Engine/Plugins/Runtime/MeshModelingToolset/Source/ModelingComponentsEditorOnly/Private/ToolTargets/StaticMeshToolTarget.cpp
@@ -126,6 +126,14 @@ bool UStaticMeshReadOnlyToolTarget::IsValid(const UStaticMesh* StaticMeshIn, EMe
 		return false;
 	}
 
+	// Early out if the static mesh is still compiling. A significant number of functions checking the state of the static mesh are
+	// blocking on static mesh builds such as `GetNumSourceModels` used below. It is best to avoid long blocking calls in IsValid to
+	// ensure editor interactions are smooth. A tool target can't be properly initialized until the static mesh is finished building anyways.
+	if (StaticMeshIn->IsCompiling())
+	{
+		return false;
+	}
+
 	if (EditingLODIn == EMeshLODIdentifier::HiResSource)
 	{
 		if (StaticMeshIn->IsHiResMeshDescriptionValid() == false)
@@ -147,6 +155,14 @@ bool UStaticMeshReadOnlyToolTarget::HasNonGeneratedLOD(const UStaticMesh* Static
 	{
 		return false;
 	}
+
+	// Early out if the static mesh is still compiling. A significant number of functions checking the state of the static mesh are
+	// blocking on static mesh builds such as `GetNumSourceModels` used below. It is best to avoid long blocking calls in IsValid to
+	// ensure editor interactions are smooth. A tool target can't be properly initialized until the static mesh is finished building anyways.
+	if (StaticMeshIn->IsCompiling())
+	{
+		return false;
+	}
 	
 	if (MeshLOD == EMeshLODIdentifier::Default)
 	{
