Modeling Mode: add quick settings tag to desired properties

diff --git a/Engine/Plugins/Runtime/MeshModelingToolset/Source/MeshModelingTools/Private/Sculpting/KelvinletBrushOp.h b/Engine/Plugins/Runtime/MeshModelingToolset/Source/MeshModelingTools/Private/Sculpting/KelvinletBrushOp.h
--- a/Engine/Plugins/Runtime/MeshModelingToolset/Source/MeshModelingTools/Private/Sculpting/KelvinletBrushOp.h
+++ b/Engine/Plugins/Runtime/MeshModelingToolset/Source/MeshModelingTools/Private/Sculpting/KelvinletBrushOp.h
@@ -194,7 +194,7 @@ class MESHMODELINGTOOLS_API UScaleKelvinletBrushOpProps : public   UBaseKelvinle
 
 public:
 	/** Strength of the Brush */
-	UPROPERTY(EditAnywhere, Category = KelvinScaleBrush, meta = (DisplayName = "Strength", UIMin = "0.0", UIMax = "10.", ClampMin = "0.0", ClampMax = "10.", ModelingQuickEdit))
+	UPROPERTY(EditAnywhere, Category = KelvinScaleBrush, meta = (DisplayName = "Strength", UIMin = "0.0", UIMax = "10.", ClampMin = "0.0", ClampMax = "10.", ModelingQuickEdit, ModelingQuickSettings))
 	float Strength = 0.5;
 
 	/** Amount of falloff to apply */
@@ -400,7 +400,7 @@ class MESHMODELINGTOOLS_API UTwistKelvinletBrushOpProps : public  UBaseKelvinlet
 
 public:
 	/** Twisting strength of the Brush */
-	UPROPERTY(EditAnywhere, Category = KelvinTwistBrush, meta = (DisplayName = "Strength", UIMin = "0.0", UIMax = "10.", ClampMin = "0.0", ClampMax = "10."))
+	UPROPERTY(EditAnywhere, Category = KelvinTwistBrush, meta = (DisplayName = "Strength", UIMin = "0.0", UIMax = "10.", ClampMin = "0.0", ClampMax = "10.", ModelingQuickSettings))
 	float Strength = 0.5;
 
 	/** Amount of falloff to apply */
diff --git a/Engine/Plugins/Runtime/MeshModelingToolset/Source/MeshModelingTools/Public/MeshVertexSculptTool.h b/Engine/Plugins/Runtime/MeshModelingToolset/Source/MeshModelingTools/Public/MeshVertexSculptTool.h
--- a/Engine/Plugins/Runtime/MeshModelingToolset/Source/MeshModelingTools/Public/MeshVertexSculptTool.h
+++ b/Engine/Plugins/Runtime/MeshModelingToolset/Source/MeshModelingTools/Public/MeshVertexSculptTool.h
@@ -136,7 +136,7 @@ public:
 	int32 PrimaryBrushID = 0;
 
 	/** Primary Brush Falloff Type, multiplied by Alpha Mask where applicable */
-	UPROPERTY(EditAnywhere, Category = Sculpting, meta = (DisplayName = "Falloff"))
+	UPROPERTY(EditAnywhere, Category = Sculpting, meta = (DisplayName = "Falloff", ModelingQuickSettings))
 	EMeshSculptFalloffType PrimaryFalloffType = EMeshSculptFalloffType::Smooth;
 
 	/** Filter applied to Stamp Region Triangles, based on first Stroke Stamp */
@@ -167,7 +167,7 @@ class MESHMODELINGTOOLS_API UVertexBrushAlphaProperties : public UInteractiveToo
 
 public:
 	/** Alpha mask applied to brush stamp. Red channel is used. */
-	UPROPERTY(EditAnywhere, Category = Alpha, meta = (DisplayName = "Alpha Mask"))
+	UPROPERTY(EditAnywhere, Category = Alpha, meta = (DisplayName = "Alpha Mask", ModelingQuickSettings))
 	TObjectPtr<UTexture2D> Alpha = nullptr;
 
 	/** Alpha is rotated by this angle, inside the brush stamp frame (vertically aligned) */
@@ -196,7 +196,7 @@ class MESHMODELINGTOOLS_API UMeshSymmetryProperties : public UInteractiveToolPro
 	GENERATED_BODY()
 public:
 	/** Enable/Disable symmetric sculpting. This option will not be available if symmetry cannot be detected, or a non-symmetric edit has been made */
-	UPROPERTY(EditAnywhere, Category = Symmetry, meta = (HideEditConditionToggle, EditCondition = bSymmetryCanBeEnabled))
+	UPROPERTY(EditAnywhere, Category = Symmetry, meta = (HideEditConditionToggle, EditCondition = bSymmetryCanBeEnabled, ModelingQuickSettings))
 	bool bEnableSymmetry = false;
 
 	// this flag is set/updated by the Tool to enable/disable the bEnableSymmetry toggle
diff --git a/Engine/Plugins/Runtime/MeshModelingToolset/Source/MeshModelingTools/Public/Sculpting/MeshInflateBrushOps.h b/Engine/Plugins/Runtime/MeshModelingToolset/Source/MeshModelingTools/Public/Sculpting/MeshInflateBrushOps.h
--- a/Engine/Plugins/Runtime/MeshModelingToolset/Source/MeshModelingTools/Public/Sculpting/MeshInflateBrushOps.h
+++ b/Engine/Plugins/Runtime/MeshModelingToolset/Source/MeshModelingTools/Public/Sculpting/MeshInflateBrushOps.h
@@ -16,7 +16,7 @@ class MESHMODELINGTOOLS_API UInflateBrushOpProps : public UMeshSculptBrushOpProp
 	GENERATED_BODY()
 public:
 	/** Strength of the Brush */
-	UPROPERTY(EditAnywhere, Category = InflateBrush, meta = (DisplayName = "Strength", UIMin = "0.0", UIMax = "1.0", ClampMin = "0.0", ClampMax = "1.0", ModelingQuickEdit))
+	UPROPERTY(EditAnywhere, Category = InflateBrush, meta = (DisplayName = "Strength", UIMin = "0.0", UIMax = "1.0", ClampMin = "0.0", ClampMax = "1.0", ModelingQuickEdit, ModelingQuickSettings))
 	float Strength = 0.5;
 
 	/** Amount of falloff to apply */
diff --git a/Engine/Plugins/Runtime/MeshModelingToolset/Source/MeshModelingTools/Public/Sculpting/MeshMoveBrushOps.h b/Engine/Plugins/Runtime/MeshModelingToolset/Source/MeshModelingTools/Public/Sculpting/MeshMoveBrushOps.h
--- a/Engine/Plugins/Runtime/MeshModelingToolset/Source/MeshModelingTools/Public/Sculpting/MeshMoveBrushOps.h
+++ b/Engine/Plugins/Runtime/MeshModelingToolset/Source/MeshModelingTools/Public/Sculpting/MeshMoveBrushOps.h
@@ -18,7 +18,7 @@ class MESHMODELINGTOOLS_API UMoveBrushOpProps : public UMeshSculptBrushOpProps
 	GENERATED_BODY()
 public:
 	/** Strength of the Brush */
-	UPROPERTY(EditAnywhere, Category = MoveBrush, meta = (DisplayName = "Strength", UIMin = "0.0", UIMax = "1.0", ClampMin = "0.0", ClampMax = "1.0", ModelingQuickEdit))
+	UPROPERTY(EditAnywhere, Category = MoveBrush, meta = (DisplayName = "Strength", UIMin = "0.0", UIMax = "1.0", ClampMin = "0.0", ClampMax = "1.0", ModelingQuickEdit, ModelingQuickSettings))
 	float Strength = 1.0;
 
 	/** Amount of falloff to apply */
diff --git a/Engine/Plugins/Runtime/MeshModelingToolset/Source/MeshModelingTools/Public/Sculpting/MeshPinchBrushOps.h b/Engine/Plugins/Runtime/MeshModelingToolset/Source/MeshModelingTools/Public/Sculpting/MeshPinchBrushOps.h
--- a/Engine/Plugins/Runtime/MeshModelingToolset/Source/MeshModelingTools/Public/Sculpting/MeshPinchBrushOps.h
+++ b/Engine/Plugins/Runtime/MeshModelingToolset/Source/MeshModelingTools/Public/Sculpting/MeshPinchBrushOps.h
@@ -15,7 +15,7 @@ class MESHMODELINGTOOLS_API UPinchBrushOpProps : public UMeshSculptBrushOpProps
 	GENERATED_BODY()
 public:
 	/** Strength of the Brush */
-	UPROPERTY(EditAnywhere, Category = PinchBrush, meta = (DisplayName = "Strength", UIMin = "0.0", UIMax = "1.0", ClampMin = "0.0", ClampMax = "1.0", ModelingQuickEdit))
+	UPROPERTY(EditAnywhere, Category = PinchBrush, meta = (DisplayName = "Strength", UIMin = "0.0", UIMax = "1.0", ClampMin = "0.0", ClampMax = "1.0", ModelingQuickEdit, ModelingQuickSettings))
 	float Strength = 0.5;
 
 	/** Amount of falloff to apply */
diff --git a/Engine/Plugins/Runtime/MeshModelingToolset/Source/MeshModelingTools/Public/Sculpting/MeshPlaneBrushOps.h b/Engine/Plugins/Runtime/MeshModelingToolset/Source/MeshModelingTools/Public/Sculpting/MeshPlaneBrushOps.h
--- a/Engine/Plugins/Runtime/MeshModelingToolset/Source/MeshModelingTools/Public/Sculpting/MeshPlaneBrushOps.h
+++ b/Engine/Plugins/Runtime/MeshModelingToolset/Source/MeshModelingTools/Public/Sculpting/MeshPlaneBrushOps.h
@@ -25,7 +25,7 @@ class MESHMODELINGTOOLS_API UPlaneBrushOpProps : public UBasePlaneBrushOpProps
 	GENERATED_BODY()
 public:
 	/** Strength of the Brush */
-	UPROPERTY(EditAnywhere, Category = PlaneBrush, meta = (DisplayName = "Strength", UIMin = "0.0", UIMax = "1.0", ClampMin = "0.0", ClampMax = "1.0", ModelingQuickEdit))
+	UPROPERTY(EditAnywhere, Category = PlaneBrush, meta = (DisplayName = "Strength", UIMin = "0.0", UIMax = "1.0", ClampMin = "0.0", ClampMax = "1.0", ModelingQuickEdit, ModelingQuickSettings))
 	float Strength = 0.5;
 
 	/** Amount of falloff to apply */
@@ -54,7 +54,7 @@ class MESHMODELINGTOOLS_API UViewAlignedPlaneBrushOpProps : public UBasePlaneBru
 	GENERATED_BODY()
 public:
 	/** Strength of the Brush */
-	UPROPERTY(EditAnywhere, Category = ViewPlaneBrush, meta = (DisplayName = "Strength", UIMin = "0.0", UIMax = "1.0", ClampMin = "0.0", ClampMax = "1.0", ModelingQuickEdit))
+	UPROPERTY(EditAnywhere, Category = ViewPlaneBrush, meta = (DisplayName = "Strength", UIMin = "0.0", UIMax = "1.0", ClampMin = "0.0", ClampMax = "1.0", ModelingQuickEdit, ModelingQuickSettings))
 	float Strength = 0.5;
 
 	/** Amount of falloff to apply */
@@ -83,7 +83,7 @@ class MESHMODELINGTOOLS_API UFixedPlaneBrushOpProps : public UBasePlaneBrushOpPr
 	GENERATED_BODY()
 public:
 	/** Strength of the Brush */
-	UPROPERTY(EditAnywhere, Category = FixedPlaneBrush, meta = (DisplayName = "Strength", UIMin = "0.0", UIMax = "1.0", ClampMin = "0.0", ClampMax = "1.0", ModelingQuickEdit))
+	UPROPERTY(EditAnywhere, Category = FixedPlaneBrush, meta = (DisplayName = "Strength", UIMin = "0.0", UIMax = "1.0", ClampMin = "0.0", ClampMax = "1.0", ModelingQuickEdit, ModelingQuickSettings))
 	float Strength = 0.5;
 
 	/** Amount of falloff to apply */
diff --git a/Engine/Plugins/Runtime/MeshModelingToolset/Source/MeshModelingTools/Public/Sculpting/MeshSculptBrushOps.h b/Engine/Plugins/Runtime/MeshModelingToolset/Source/MeshModelingTools/Public/Sculpting/MeshSculptBrushOps.h
--- a/Engine/Plugins/Runtime/MeshModelingToolset/Source/MeshModelingTools/Public/Sculpting/MeshSculptBrushOps.h
+++ b/Engine/Plugins/Runtime/MeshModelingToolset/Source/MeshModelingTools/Public/Sculpting/MeshSculptBrushOps.h
@@ -16,7 +16,7 @@ class MESHMODELINGTOOLS_API UStandardSculptBrushOpProps : public UMeshSculptBrus
 	GENERATED_BODY()
 public:
 	/** Strength of the Brush */
-	UPROPERTY(EditAnywhere, Category = SculptBrush, meta = (DisplayName = "Strength", UIMin = "0.0", UIMax = "1.0", ClampMin = "0.0", ClampMax = "1.0", ModelingQuickEdit))
+	UPROPERTY(EditAnywhere, Category = SculptBrush, meta = (DisplayName = "Strength", UIMin = "0.0", UIMax = "1.0", ClampMin = "0.0", ClampMax = "1.0", ModelingQuickEdit, ModelingQuickSettings))
 	float Strength = 0.5;
 
 	/** Amount of falloff to apply */
@@ -97,7 +97,7 @@ class MESHMODELINGTOOLS_API UViewAlignedSculptBrushOpProps : public UMeshSculptB
 	GENERATED_BODY()
 public:
 	/** Strength of the Brush */
-	UPROPERTY(EditAnywhere, Category = SculptToViewBrush, meta = (DisplayName = "Strength", UIMin = "0.0", UIMax = "1.0", ClampMin = "0.0", ClampMax = "1.0", ModelingQuickEdit))
+	UPROPERTY(EditAnywhere, Category = SculptToViewBrush, meta = (DisplayName = "Strength", UIMin = "0.0", UIMax = "1.0", ClampMin = "0.0", ClampMax = "1.0", ModelingQuickEdit, ModelingQuickSettings))
 	float Strength = 0.5;
 
 	/** Amount of falloff to apply */
@@ -186,7 +186,7 @@ class MESHMODELINGTOOLS_API USculptMaxBrushOpProps : public UMeshSculptBrushOpPr
 	GENERATED_BODY()
 public:
 	/** Strength of the Brush */
-	UPROPERTY(EditAnywhere, Category = SculptMaxBrush, meta = (DisplayName = "Strength", UIMin = "0.0", UIMax = "1.0", ClampMin = "0.0", ClampMax = "1.0", ModelingQuickEdit))
+	UPROPERTY(EditAnywhere, Category = SculptMaxBrush, meta = (DisplayName = "Strength", UIMin = "0.0", UIMax = "1.0", ClampMin = "0.0", ClampMax = "1.0", ModelingQuickEdit, ModelingQuickSettings))
 	float Strength = 0.5;
 
 	/** Amount of falloff to apply */
diff --git a/Engine/Plugins/Runtime/MeshModelingToolset/Source/MeshModelingTools/Public/Sculpting/MeshSculptToolBase.h b/Engine/Plugins/Runtime/MeshModelingToolset/Source/MeshModelingTools/Public/Sculpting/MeshSculptToolBase.h
--- a/Engine/Plugins/Runtime/MeshModelingToolset/Source/MeshModelingTools/Public/Sculpting/MeshSculptToolBase.h
+++ b/Engine/Plugins/Runtime/MeshModelingToolset/Source/MeshModelingTools/Public/Sculpting/MeshSculptToolBase.h
@@ -101,7 +101,7 @@ class MESHMODELINGTOOLS_API USculptBrushProperties : public UInteractiveToolProp
 	GENERATED_BODY()
 public:
 
-	UPROPERTY(EditAnywhere, Category = Brush, meta = (DisplayPriority = 1, ModelingQuickEdit))
+	UPROPERTY(EditAnywhere, Category = Brush, meta = (DisplayPriority = 1, ModelingQuickEdit, ModelingQuickSettings))
 	FBrushToolRadius BrushSize;
 
 	/** Amount of falloff to apply (0.0 - 1.0) */
diff --git a/Engine/Plugins/Runtime/MeshModelingToolset/Source/MeshModelingTools/Public/Sculpting/MeshSmoothingBrushOps.h b/Engine/Plugins/Runtime/MeshModelingToolset/Source/MeshModelingTools/Public/Sculpting/MeshSmoothingBrushOps.h
--- a/Engine/Plugins/Runtime/MeshModelingToolset/Source/MeshModelingTools/Public/Sculpting/MeshSmoothingBrushOps.h
+++ b/Engine/Plugins/Runtime/MeshModelingToolset/Source/MeshModelingTools/Public/Sculpting/MeshSmoothingBrushOps.h
@@ -25,7 +25,7 @@ class MESHMODELINGTOOLS_API USmoothBrushOpProps : public UBaseSmoothBrushOpProps
 	GENERATED_BODY()
 public:
 	/** Strength of the Brush */
-	UPROPERTY(EditAnywhere, Category = SmoothBrush, meta = (DisplayName = "Strength", UIMin = "0.0", UIMax = "1.0", ClampMin = "0.0", ClampMax = "1.0", ModelingQuickEdit))
+	UPROPERTY(EditAnywhere, Category = SmoothBrush, meta = (DisplayName = "Strength", UIMin = "0.0", UIMax = "1.0", ClampMin = "0.0", ClampMax = "1.0", ModelingQuickEdit, ModelingQuickSettings))
 	float Strength = 0.5;
 
 	/** Amount of falloff to apply (0.0 - 1.0) */
@@ -114,7 +114,7 @@ class MESHMODELINGTOOLS_API USmoothFillBrushOpProps : public UBaseSmoothBrushOpP
 	GENERATED_BODY()
 public:
 	/** Strength of the Brush */
-	UPROPERTY(EditAnywhere, Category = SmoothFillBrush, meta = (DisplayName = "Strength", UIMin = "0.0", UIMax = "1.0", ClampMin = "0.0", ClampMax = "1.0", ModelingQuickEdit))
+	UPROPERTY(EditAnywhere, Category = SmoothFillBrush, meta = (DisplayName = "Strength", UIMin = "0.0", UIMax = "1.0", ClampMin = "0.0", ClampMax = "1.0", ModelingQuickEdit, ModelingQuickSettings))
 	float Strength = 0.5;
 
 	/** Amount of falloff to apply (0.0 - 1.0) */
@@ -187,7 +187,7 @@ class MESHMODELINGTOOLS_API UFlattenBrushOpProps : public UMeshSculptBrushOpProp
 	GENERATED_BODY()
 public:
 	/** Strength of the Brush */
-	UPROPERTY(EditAnywhere, Category = FlattenBrush, meta = (DisplayName = "Strength", UIMin = "0.0", UIMax = "1.0", ClampMin = "0.0", ClampMax = "1.0", ModelingQuickEdit))
+	UPROPERTY(EditAnywhere, Category = FlattenBrush, meta = (DisplayName = "Strength", UIMin = "0.0", UIMax = "1.0", ClampMin = "0.0", ClampMax = "1.0", ModelingQuickEdit, ModelingQuickSettings))
 	float Strength = 0.5;
 
 	/** Amount of falloff to apply (0.0 - 1.0) */
@@ -267,7 +267,7 @@ class MESHMODELINGTOOLS_API UEraseBrushOpProps : public UMeshSculptBrushOpProps
 	GENERATED_BODY()
 public:
 	/** Strength of the Brush */
-	UPROPERTY(EditAnywhere, Category = EraseBrush, meta = (DisplayName = "Strength", UIMin = "0.0", UIMax = "1.0", ClampMin = "0.0", ClampMax = "1.0", ModelingQuickEdit))
+	UPROPERTY(EditAnywhere, Category = EraseBrush, meta = (DisplayName = "Strength", UIMin = "0.0", UIMax = "1.0", ClampMin = "0.0", ClampMax = "1.0", ModelingQuickEdit, ModelingQuickSettings))
 	float Strength = 0.5;
 
 	/** Amount of falloff to apply (0.0 - 1.0) */
