Add an SceneComponentToolTarget concept to support tool targets for actors that don't derive from UPrimitiveComponent, since many tool features can operate on just a scene component

Move most of the IPrimitiveComponentBackedToolTarget interface to a corresponding ISceneComponentBackedToolTarget base class

diff --git a/Engine/Plugins/Runtime/MeshModelingToolset/Source/ModelingComponents/Private/ModelingObjectsCreationAPI.cpp b/Engine/Plugins/Runtime/MeshModelingToolset/Source/ModelingComponents/Private/ModelingObjectsCreationAPI.cpp
--- a/Engine/Plugins/Runtime/MeshModelingToolset/Source/ModelingComponents/Private/ModelingObjectsCreationAPI.cpp
+++ b/Engine/Plugins/Runtime/MeshModelingToolset/Source/ModelingComponents/Private/ModelingObjectsCreationAPI.cpp
@@ -5,6 +5,7 @@
 #include "ContextObjectStore.h"
 
 #include "Misc/Paths.h"
+#include "Components/PrimitiveComponent.h"
 #include "Components/StaticMeshComponent.h"
 #include "Engine/StaticMesh.h"
 #include "ModelingComponentsSettings.h"
@@ -176,7 +177,7 @@ FCreateComponentResult UE::Modeling::CreateNewComponentOnActor(UInteractiveToolM
 
 
 
-FString UE::Modeling::GetComponentAssetBaseName(UPrimitiveComponent* Component, bool bRemoveAutoGeneratedSuffixes)
+FString UE::Modeling::GetComponentAssetBaseName(UActorComponent* Component, bool bRemoveAutoGeneratedSuffixes)
 {
 	if (!ensure(Component != nullptr))
 	{
diff --git a/Engine/Plugins/Runtime/MeshModelingToolset/Source/ModelingComponents/Private/ModelingToolTargetUtil.cpp b/Engine/Plugins/Runtime/MeshModelingToolset/Source/ModelingComponents/Private/ModelingToolTargetUtil.cpp
--- a/Engine/Plugins/Runtime/MeshModelingToolset/Source/ModelingComponents/Private/ModelingToolTargetUtil.cpp
+++ b/Engine/Plugins/Runtime/MeshModelingToolset/Source/ModelingComponents/Private/ModelingToolTargetUtil.cpp
@@ -47,7 +47,7 @@ namespace UE::ToolTarget::Internal
 
 AActor* UE::ToolTarget::GetTargetActor(UToolTarget* Target)
 {
-	IPrimitiveComponentBackedTarget* TargetComponent = Cast<IPrimitiveComponentBackedTarget>(Target);
+	ISceneComponentBackedTarget* TargetComponent = Cast<ISceneComponentBackedTarget>(Target);
 	if (TargetComponent)
 	{
 		return TargetComponent->GetOwnerActor();
@@ -66,6 +66,16 @@ UPrimitiveComponent* UE::ToolTarget::GetTargetComponent(UToolTarget* Target)
 	return nullptr;
 }
 
+USceneComponent* UE::ToolTarget::GetTargetSceneComponent(UToolTarget* Target)
+{
+	ISceneComponentBackedTarget* TargetComponent = Cast<ISceneComponentBackedTarget>(Target);
+	if (TargetComponent)
+	{
+		return TargetComponent->GetOwnerSceneComponent();
+	}
+	return nullptr;
+}
+
 FString UE::ToolTarget::GetHumanReadableName(UToolTarget* Target)
 {
 	IAssetBackedTarget* TargetAsset = Cast<IAssetBackedTarget>(Target);
@@ -74,10 +84,10 @@ FString UE::ToolTarget::GetHumanReadableName(UToolTarget* Target)
 		return TargetAsset->GetSourceData()->GetName();
 	}
 
-	IPrimitiveComponentBackedTarget* TargetComponent = Cast<IPrimitiveComponentBackedTarget>(Target);
+	ISceneComponentBackedTarget* TargetComponent = Cast<ISceneComponentBackedTarget>(Target);
 	if (TargetComponent)
 	{
-		return TargetComponent->GetOwnerComponent()->GetFullGroupName(false);
+		return TargetComponent->GetOwnerSceneComponent()->GetFullGroupName(false);
 	}
 
 	return FString("");
@@ -85,7 +95,7 @@ FString UE::ToolTarget::GetHumanReadableName(UToolTarget* Target)
 
 bool UE::ToolTarget::HideSourceObject(UToolTarget* Target)
 {
-	IPrimitiveComponentBackedTarget* TargetComponent = Cast<IPrimitiveComponentBackedTarget>(Target);
+	ISceneComponentBackedTarget* TargetComponent = Cast<ISceneComponentBackedTarget>(Target);
 	if (TargetComponent)
 	{
 		TargetComponent->SetOwnerVisibility(false);
@@ -97,7 +107,7 @@ bool UE::ToolTarget::HideSourceObject(UToolTarget* Target)
 
 bool UE::ToolTarget::ShowSourceObject(UToolTarget* Target)
 {
-	IPrimitiveComponentBackedTarget* TargetComponent = Cast<IPrimitiveComponentBackedTarget>(Target);
+	ISceneComponentBackedTarget* TargetComponent = Cast<ISceneComponentBackedTarget>(Target);
 	if (TargetComponent)
 	{
 		TargetComponent->SetOwnerVisibility(true);
@@ -123,7 +133,7 @@ bool UE::ToolTarget::SetSourceObjectVisible(UToolTarget* Target, bool bVisible)
 
 FTransform3d UE::ToolTarget::GetLocalToWorldTransform(UToolTarget* Target)
 {
-	IPrimitiveComponentBackedTarget* TargetComponent = Cast<IPrimitiveComponentBackedTarget>(Target);
+	ISceneComponentBackedTarget* TargetComponent = Cast<ISceneComponentBackedTarget>(Target);
 	if (TargetComponent)
 	{
 		return (FTransform3d)TargetComponent->GetWorldTransform();
@@ -587,16 +597,16 @@ bool UE::ToolTarget::ApplyIncrementalMeshEditChange(
 
 bool UE::ToolTarget::ConfigureCreateMeshObjectParams(UToolTarget* SourceTarget, FCreateMeshObjectParams& DerivedParamsOut)
 {
-	IPrimitiveComponentBackedTarget* ComponentTarget = Cast<IPrimitiveComponentBackedTarget>(SourceTarget);
+	ISceneComponentBackedTarget* ComponentTarget = Cast<ISceneComponentBackedTarget>(SourceTarget);
 	if (ComponentTarget)
 	{
-		if (Cast<UStaticMeshComponent>(ComponentTarget->GetOwnerComponent()) != nullptr)
+		if (Cast<UStaticMeshComponent>(ComponentTarget->GetOwnerSceneComponent()) != nullptr)
 		{
 			DerivedParamsOut.TypeHint = ECreateObjectTypeHint::StaticMesh;
 			return true;
 		}
 
-		if (Cast<UDynamicMeshComponent>(ComponentTarget->GetOwnerComponent()) != nullptr)
+		if (Cast<UDynamicMeshComponent>(ComponentTarget->GetOwnerSceneComponent()) != nullptr)
 		{
 			DerivedParamsOut.TypeHint = ECreateObjectTypeHint::DynamicMeshActor;
 			return true;
diff --git a/Engine/Plugins/Runtime/MeshModelingToolset/Source/ModelingComponents/Public/ModelingObjectsCreationAPI.h b/Engine/Plugins/Runtime/MeshModelingToolset/Source/ModelingComponents/Public/ModelingObjectsCreationAPI.h
--- a/Engine/Plugins/Runtime/MeshModelingToolset/Source/ModelingComponents/Public/ModelingObjectsCreationAPI.h
+++ b/Engine/Plugins/Runtime/MeshModelingToolset/Source/ModelingComponents/Public/ModelingObjectsCreationAPI.h
@@ -17,6 +17,7 @@
 class UInteractiveToolManager;
 class AActor;
 class UPrimitiveComponent;
+class UActorComponent;
 class UMaterialInterface;
 class UTexture2D;
 
@@ -653,7 +654,7 @@ MODELINGCOMPONENTS_API FString GenerateRandomShortHexString(int32 NumChars = 8);
  * @return the Name for the Component
  */
 MODELINGCOMPONENTS_API FString GetComponentAssetBaseName(
-	UPrimitiveComponent* Component,
+	UActorComponent* Component,
 	bool bRemoveAutoGeneratedSuffixes = true);
 
 }
diff --git a/Engine/Plugins/Runtime/MeshModelingToolset/Source/ModelingComponents/Public/ModelingToolTargetUtil.h b/Engine/Plugins/Runtime/MeshModelingToolset/Source/ModelingComponents/Public/ModelingToolTargetUtil.h
--- a/Engine/Plugins/Runtime/MeshModelingToolset/Source/ModelingComponents/Public/ModelingToolTargetUtil.h
+++ b/Engine/Plugins/Runtime/MeshModelingToolset/Source/ModelingComponents/Public/ModelingToolTargetUtil.h
@@ -12,6 +12,7 @@ class UStaticMesh;
 class USkeletalMesh;
 class UToolTarget;
 class UPrimitiveComponent;
+class USceneComponent;
 class AActor;
 class UBodySetup;
 class IInterface_CollisionDataProvider;
@@ -44,6 +45,11 @@ MODELINGCOMPONENTS_API AActor* GetTargetActor(UToolTarget* Target);
  */
 MODELINGCOMPONENTS_API UPrimitiveComponent* GetTargetComponent(UToolTarget* Target);
 
+/**
+ * @return the UActorComponent backing a ToolTarget, or nullptr if there is no such Component
+ */
+MODELINGCOMPONENTS_API USceneComponent* GetTargetSceneComponent(UToolTarget* Target);
+
 /**
  * @return a human readable identifier for a ToolTarget, based on the underlying nature of it.
  */
diff --git a/Engine/Source/Runtime/InteractiveToolsFramework/Private/ToolTargets/PrimitiveComponentToolTarget.cpp b/Engine/Source/Runtime/InteractiveToolsFramework/Private/ToolTargets/PrimitiveComponentToolTarget.cpp
--- a/Engine/Source/Runtime/InteractiveToolsFramework/Private/ToolTargets/PrimitiveComponentToolTarget.cpp
+++ b/Engine/Source/Runtime/InteractiveToolsFramework/Private/ToolTargets/PrimitiveComponentToolTarget.cpp
@@ -18,6 +18,12 @@ UPrimitiveComponent* UPrimitiveComponentToolTarget::GetOwnerComponent() const
 	return IsValid() ? Component.Get() : nullptr;
 }
 
+USceneComponent* UPrimitiveComponentToolTarget::GetOwnerSceneComponent() const
+{
+	return StaticCast<USceneComponent*>(GetOwnerComponent());
+}
+
+
 AActor* UPrimitiveComponentToolTarget::GetOwnerActor() const
 {
 	return IsValid() ? Component->GetOwner() : nullptr;
diff --git a/Engine/Source/Runtime/InteractiveToolsFramework/Private/ToolTargets/SceneComponentToolTarget.cpp b/Engine/Source/Runtime/InteractiveToolsFramework/Private/ToolTargets/SceneComponentToolTarget.cpp
new file mode 100644
--- /dev/null
+++ b/Engine/Source/Runtime/InteractiveToolsFramework/Private/ToolTargets/SceneComponentToolTarget.cpp
@@ -0,0 +1,102 @@
+// Copyright Epic Games, Inc. All Rights Reserved.
+
+#include "ToolTargets/SceneComponentToolTarget.h"
+
+#include "Components/SceneComponent.h"
+#include "Components/PrimitiveComponent.h"
+
+#include UE_INLINE_GENERATED_CPP_BY_NAME(SceneComponentToolTarget)
+
+bool USceneComponentToolTarget::IsValid() const
+{
+	return Component.IsValid();
+}
+
+USceneComponent* USceneComponentToolTarget::GetOwnerSceneComponent() const
+{
+	// Note that we don't just return Component.Get() because we want to call the virtual IsValid
+	// for derived classes.
+	return IsValid() ? Component.Get() : nullptr;
+}
+
+AActor* USceneComponentToolTarget::GetOwnerActor() const
+{
+	return IsValid() ? Component->GetOwner() : nullptr;
+}
+
+void USceneComponentToolTarget::SetOwnerVisibility(bool bVisible) const
+{
+	if (IsValid())
+	{
+		Component->SetVisibility(bVisible);
+	}
+}
+
+FTransform USceneComponentToolTarget::GetWorldTransform() const
+{
+	return IsValid() ? Component->GetComponentTransform() : FTransform::Identity;
+}
+
+void USceneComponentToolTarget::InitializeComponent(USceneComponent* ComponentIn)
+{
+	Component = ComponentIn;
+
+	if (ensure(Component.IsValid()))
+	{
+#if WITH_EDITOR
+		// See comment in OnObjectsReplaced
+		FCoreUObjectDelegates::OnObjectsReplaced.AddUObject(this, &USceneComponentToolTarget::OnObjectsReplaced);
+#endif
+	}
+}
+
+void USceneComponentToolTarget::BeginDestroy()
+{
+#if WITH_EDITOR
+	FCoreUObjectDelegates::OnObjectsReplaced.RemoveAll(this);
+#endif
+
+	Super::BeginDestroy();
+}
+
+#if WITH_EDITOR
+void USceneComponentToolTarget::OnObjectsReplaced(const FCoreUObjectDelegates::FReplacementObjectMap& Map)
+{
+	// Components frequently get destroyed and recreated when they are part of blueprint actors that 
+	// get modified. For the most part, we don't need to worry about supporting these cases, but keeping
+	// a consistent reference here allows us to avoid getting into some bad states. For instance, we often
+	// hide the source component and unhide at tool end, and if we lose the reference to the component
+	// while the tool is running, we are unable to unhide it later. The user is unlikely to understand 
+	// why their object disappeared in that case or know to fix it via the component visibility property.
+
+	UObject* const * MappedObject = Map.Find(Component.Get());
+	if (MappedObject)
+	{
+		Component = Cast<USceneComponent>(*MappedObject);
+	}
+}
+#endif
+
+
+// Factory
+
+bool USceneComponentToolTargetFactory::CanBuildTarget(UObject* SourceObject, const FToolTargetTypeRequirements& Requirements) const
+{
+	USceneComponent* Component = Cast<USceneComponent>(SourceObject);
+	return Component 
+		&& ::IsValid(Component)
+		&& !Component->IsUnreachable() 
+		&& Component->IsValidLowLevel() 
+		&& Requirements.AreSatisfiedBy(USceneComponentToolTarget::StaticClass());
+}
+
+UToolTarget* USceneComponentToolTargetFactory::BuildTarget(UObject* SourceObject, const FToolTargetTypeRequirements& Requirements)
+{
+	USceneComponentToolTarget* Target = NewObject<USceneComponentToolTarget>();
+	Target->InitializeComponent(Cast<USceneComponent>(SourceObject));
+	checkSlow(Target->Component.IsValid() && Requirements.AreSatisfiedBy(Target));
+	return Target;
+}
+
+
+
diff --git a/Engine/Source/Runtime/InteractiveToolsFramework/Public/TargetInterfaces/PrimitiveComponentBackedTarget.h b/Engine/Source/Runtime/InteractiveToolsFramework/Public/TargetInterfaces/PrimitiveComponentBackedTarget.h
--- a/Engine/Source/Runtime/InteractiveToolsFramework/Public/TargetInterfaces/PrimitiveComponentBackedTarget.h
+++ b/Engine/Source/Runtime/InteractiveToolsFramework/Public/TargetInterfaces/PrimitiveComponentBackedTarget.h
@@ -4,21 +4,24 @@
 
 #include "CoreMinimal.h"
 #include "UObject/Interface.h"
+#include "TargetInterfaces/SceneComponentBackedTarget.h"
 
 #include "PrimitiveComponentBackedTarget.generated.h"
 
 class UPrimitiveComponent;
+class UActorComponent;
 class AActor;
 
 struct FHitResult;
 
+
 UINTERFACE(MinimalAPI)
-class UPrimitiveComponentBackedTarget : public UInterface
+class UPrimitiveComponentBackedTarget : public USceneComponentBackedTarget
 {
 	GENERATED_BODY()
 };
 
-class IPrimitiveComponentBackedTarget
+class IPrimitiveComponentBackedTarget : public ISceneComponentBackedTarget
 {
 	GENERATED_BODY()
 
@@ -27,23 +30,8 @@ public:
 	/** @return the Component this is a Source for */
 	virtual UPrimitiveComponent* GetOwnerComponent() const = 0;
 
-	/** @return the Actor that owns this Component */
-	virtual AActor* GetOwnerActor() const = 0;
-
-	/**
-	 * Set the visibility of the Component associated with this Source (ie to hide during Tool usage)
-	 * @param bVisible desired visibility
-	 */
-	virtual void SetOwnerVisibility(bool bVisible) const = 0;
-
-	/**
-	 * @return the transform on this component
-	 * @todo Do we need to return a list of transforms here?
-	 */
-	virtual FTransform GetWorldTransform() const = 0;
-
 	/**
-	 * Compute ray intersection with the MeshDescription this Source is providing
+	 * Compute ray intersection with geometry this Source is providing (if any)
 	 * @param WorldRay ray in world space
 	 * @param OutHit hit test data
 	 * @return true if ray intersected Component
diff --git a/Engine/Source/Runtime/InteractiveToolsFramework/Public/TargetInterfaces/SceneComponentBackedTarget.h b/Engine/Source/Runtime/InteractiveToolsFramework/Public/TargetInterfaces/SceneComponentBackedTarget.h
new file mode 100644
--- /dev/null
+++ b/Engine/Source/Runtime/InteractiveToolsFramework/Public/TargetInterfaces/SceneComponentBackedTarget.h
@@ -0,0 +1,45 @@
+// Copyright Epic Games, Inc. All Rights Reserved.
+
+#pragma once
+
+#include "CoreMinimal.h"
+#include "UObject/Interface.h"
+
+#include "SceneComponentBackedTarget.generated.h"
+
+class USceneComponent;
+class AActor;
+
+struct FHitResult;
+
+UINTERFACE(MinimalAPI)
+class USceneComponentBackedTarget : public UInterface
+{
+	GENERATED_BODY()
+};
+
+class ISceneComponentBackedTarget
+{
+	GENERATED_BODY()
+
+public:
+
+	/** @return the ActorComponent this is a Source for */
+	virtual USceneComponent* GetOwnerSceneComponent() const = 0;
+
+	/** @return the Actor that owns this Component */
+	virtual AActor* GetOwnerActor() const = 0;
+
+	/**
+	 * Set the visibility of the Component associated with this Source (ie to hide during Tool usage)
+	 * @param bVisible desired visibility
+	 */
+	virtual void SetOwnerVisibility(bool bVisible) const = 0;
+
+	/**
+	 * @return the transform on this component
+	 * @todo Do we need to return a list of transforms here?
+	 */
+	virtual FTransform GetWorldTransform() const = 0;
+};
+
diff --git a/Engine/Source/Runtime/InteractiveToolsFramework/Public/ToolTargets/PrimitiveComponentToolTarget.h b/Engine/Source/Runtime/InteractiveToolsFramework/Public/ToolTargets/PrimitiveComponentToolTarget.h
--- a/Engine/Source/Runtime/InteractiveToolsFramework/Public/ToolTargets/PrimitiveComponentToolTarget.h
+++ b/Engine/Source/Runtime/InteractiveToolsFramework/Public/ToolTargets/PrimitiveComponentToolTarget.h
@@ -24,6 +24,7 @@ public:
 
 	// IPrimitiveComponentBackedTarget implementation
 	INTERACTIVETOOLSFRAMEWORK_API virtual UPrimitiveComponent* GetOwnerComponent() const override;
+	INTERACTIVETOOLSFRAMEWORK_API virtual USceneComponent* GetOwnerSceneComponent() const override;
 	INTERACTIVETOOLSFRAMEWORK_API virtual AActor* GetOwnerActor() const override;
 	INTERACTIVETOOLSFRAMEWORK_API virtual void SetOwnerVisibility(bool bVisible) const override;
 	INTERACTIVETOOLSFRAMEWORK_API virtual FTransform GetWorldTransform() const override;
diff --git a/Engine/Source/Runtime/InteractiveToolsFramework/Public/ToolTargets/SceneComponentToolTarget.h b/Engine/Source/Runtime/InteractiveToolsFramework/Public/ToolTargets/SceneComponentToolTarget.h
new file mode 100644
--- /dev/null
+++ b/Engine/Source/Runtime/InteractiveToolsFramework/Public/ToolTargets/SceneComponentToolTarget.h
@@ -0,0 +1,58 @@
+// Copyright Epic Games, Inc. All Rights Reserved.
+
+#pragma once
+
+#include "CoreMinimal.h"
+
+#include "TargetInterfaces/PrimitiveComponentBackedTarget.h"
+#include "ToolTargets/ToolTarget.h"
+
+#include "SceneComponentToolTarget.generated.h"
+
+/** 
+ * A tool target to share some reusable code for tool targets that are
+ * backed by scene components
+ */
+UCLASS(Transient, MinimalAPI)
+class USceneComponentToolTarget : public UToolTarget, public ISceneComponentBackedTarget
+{
+	GENERATED_BODY()
+public:
+
+	// UToolTarget
+	INTERACTIVETOOLSFRAMEWORK_API virtual bool IsValid() const override;
+
+	// ISceneComponentBackedTarget implementation
+	INTERACTIVETOOLSFRAMEWORK_API virtual USceneComponent* GetOwnerSceneComponent() const;
+	INTERACTIVETOOLSFRAMEWORK_API virtual AActor* GetOwnerActor() const;
+	INTERACTIVETOOLSFRAMEWORK_API virtual void SetOwnerVisibility(bool bVisible) const;
+	INTERACTIVETOOLSFRAMEWORK_API virtual FTransform GetWorldTransform() const;
+
+	// UObject
+	INTERACTIVETOOLSFRAMEWORK_API void BeginDestroy() override;
+
+protected:
+	friend class USceneComponentToolTargetFactory;
+	
+	INTERACTIVETOOLSFRAMEWORK_API void InitializeComponent(USceneComponent* ComponentIn);
+
+	TWeakObjectPtr<USceneComponent> Component;
+
+private:
+#if WITH_EDITOR
+	INTERACTIVETOOLSFRAMEWORK_API void OnObjectsReplaced(const FCoreUObjectDelegates::FReplacementObjectMap& Map);
+#endif
+};
+
+/**
+ * Factory for USceneComponentToolTarget to be used by the target manager.
+ */
+UCLASS(Transient, MinimalAPI)
+class USceneComponentToolTargetFactory : public UToolTargetFactory
+{
+	GENERATED_BODY()
+public:
+	INTERACTIVETOOLSFRAMEWORK_API virtual bool CanBuildTarget(UObject* SourceObject, const FToolTargetTypeRequirements& TargetTypeInfo) const override;
+	INTERACTIVETOOLSFRAMEWORK_API virtual UToolTarget* BuildTarget(UObject* SourceObject, const FToolTargetTypeRequirements& TargetTypeInfo) override;
+};
+
