Fix crash in sculpt tool material handling when target is not primitive backed

diff --git a/Engine/Plugins/Runtime/MeshModelingToolset/Source/MeshModelingTools/Private/Sculpting/MeshSculptToolBase.cpp b/Engine/Plugins/Runtime/MeshModelingToolset/Source/MeshModelingTools/Private/Sculpting/MeshSculptToolBase.cpp
--- a/Engine/Plugins/Runtime/MeshModelingToolset/Source/MeshModelingTools/Private/Sculpting/MeshSculptToolBase.cpp
+++ b/Engine/Plugins/Runtime/MeshModelingToolset/Source/MeshModelingTools/Private/Sculpting/MeshSculptToolBase.cpp
@@ -1171,7 +1171,10 @@ void UMeshSculptToolBase::UpdateMaterialMode(EMeshEditingMaterialModes MaterialM
 	if (MaterialMode == EMeshEditingMaterialModes::ExistingMaterial)
 	{
 		GetSculptMeshComponent()->ClearOverrideRenderMaterial();
-		GetSculptMeshComponent()->SetShadowsEnabled(UE::ToolTarget::GetTargetComponent(Target)->bCastDynamicShadow);
+		if (UPrimitiveComponent* TargetComponent = UE::ToolTarget::GetTargetComponent(Target))
+		{
+			GetSculptMeshComponent()->SetShadowsEnabled(TargetComponent->bCastDynamicShadow);
+		}
 		ActiveOverrideMaterial = nullptr;
 	}
 	else
