make sculpt normals method robust to partially-set normals in overlay

diff --git a/Engine/Plugins/Runtime/MeshModelingToolset/Source/MeshModelingTools/Private/Sculpting/MeshSculptUtil.cpp b/Engine/Plugins/Runtime/MeshModelingToolset/Source/MeshModelingTools/Private/Sculpting/MeshSculptUtil.cpp
--- a/Engine/Plugins/Runtime/MeshModelingToolset/Source/MeshModelingTools/Private/Sculpting/MeshSculptUtil.cpp
+++ b/Engine/Plugins/Runtime/MeshModelingToolset/Source/MeshModelingTools/Private/Sculpting/MeshSculptUtil.cpp
@@ -20,13 +20,16 @@ void UE::SculptUtil::RecalculateNormals_Overlay(
 	for (int32 TriangleID : ModifiedTris)
 	{
 		FIndex3i TriElems = Normals->GetTriangle(TriangleID);
-		for (int32 j = 0; j < 3; ++j)
+		if (TriElems.A != INDEX_NONE)
 		{
-			int32 elemid = TriElems[j];
-			if (VertexSetBuffer.Contains(elemid) == false)
+			for (int32 j = 0; j < 3; ++j)
 			{
-				VertexSetBuffer.Add(elemid);
-				NormalsBuffer.Add(elemid);
+				int32 elemid = TriElems[j];
+				if (VertexSetBuffer.Contains(elemid) == false)
+				{
+					VertexSetBuffer.Add(elemid);
+					NormalsBuffer.Add(elemid);
+				}
 			}
 		}
 	}
@@ -109,9 +112,12 @@ void UE::SculptUtil::RecalculateNormals_Overlay(
 	for (int32 TriangleID : ModifiedTris)
 	{
 		FIndex3i TriElems = Normals->GetTriangle(TriangleID);
-		VertexSetTemp.Add(TriElems.A);
-		VertexSetTemp.Add(TriElems.B);
-		VertexSetTemp.Add(TriElems.C);
+		if (TriElems.A != INDEX_NONE)
+		{
+			VertexSetTemp.Add(TriElems.A);
+			VertexSetTemp.Add(TriElems.B);
+			VertexSetTemp.Add(TriElems.C);
+		}
 	}
 
 	const TArray<int32>& ElemIndices = VertexSetTemp.Indices();
