add (editor-only) ue standard simplifier option to geometry script

diff --git a/Engine/Plugins/Runtime/GeometryScripting/Source/GeometryScriptingCore/GeometryScriptingCore.Build.cs b/Engine/Plugins/Runtime/GeometryScripting/Source/GeometryScriptingCore/GeometryScriptingCore.Build.cs
--- a/Engine/Plugins/Runtime/GeometryScripting/Source/GeometryScriptingCore/GeometryScriptingCore.Build.cs
+++ b/Engine/Plugins/Runtime/GeometryScripting/Source/GeometryScriptingCore/GeometryScriptingCore.Build.cs
@@ -66,6 +66,8 @@ public class GeometryScriptingCore : ModuleRules
 		if (Target.bBuildEditor == true)
 		{
 			PrivateDependencyModuleNames.Add("EditorFramework");
+			PrivateDependencyModuleNames.Add("ModelingOperatorsEditorOnly");
+			PrivateDependencyModuleNames.Add("MeshReductionInterface"); // for UE standard simplification 
 			PrivateDependencyModuleNames.Add("UnrealEd");
 		}
 
diff --git a/Engine/Plugins/Runtime/GeometryScripting/Source/GeometryScriptingCore/Private/MeshSimplifyFunctions.cpp b/Engine/Plugins/Runtime/GeometryScripting/Source/GeometryScriptingCore/Private/MeshSimplifyFunctions.cpp
--- a/Engine/Plugins/Runtime/GeometryScripting/Source/GeometryScriptingCore/Private/MeshSimplifyFunctions.cpp
+++ b/Engine/Plugins/Runtime/GeometryScripting/Source/GeometryScriptingCore/Private/MeshSimplifyFunctions.cpp
@@ -5,6 +5,7 @@
 #include "DynamicMesh/DynamicMesh3.h"
 #include "DynamicMesh/MeshNormals.h"
 #include "DynamicMesh/DynamicMeshAABBTree3.h"
+#include "DynamicMeshToMeshDescription.h"
 #include "MeshSimplification.h"
 #include "MeshConstraintsUtil.h"
 #include "ProjectionTargets.h"
@@ -13,6 +14,15 @@
 #include "Operations/PolygroupRemesh.h"
 #include "ConstrainedDelaunay2.h"
 
+#if WITH_EDITOR
+// for UE Standard simplifier (editor only)
+#include "CleaningOps/SimplifyMeshOp.h"
+#include "StaticMeshAttributes.h"
+#include "Modules/ModuleManager.h"
+#include "IMeshReductionManagerModule.h"
+#include "IMeshReductionInterfaces.h"
+#endif
+
 #include "UDynamicMesh.h"
 
 #include UE_INLINE_GENERATED_CPP_BY_NAME(MeshSimplifyFunctions)
@@ -101,6 +111,38 @@ UDynamicMesh* UGeometryScriptLibrary_MeshSimplifyFunctions::ApplySimplifyToPolyg
 
 
 
+namespace UE::PrivateSimplifyHelper
+{
+	bool UEStandardEditorSimplify(FDynamicMesh3& Mesh, bool bTargetIsTriCount, int32 TargetCount = 0.f)
+	{
+#if WITH_EDITOR
+		IMeshReductionManagerModule* MeshReductionModule = FModuleManager::Get().LoadModulePtr<IMeshReductionManagerModule>("MeshReductionInterface");
+		if (!MeshReductionModule)
+		{
+			UE_LOG(LogGeometry, Warning, TEXT("Failed to load mesh reduction module; cannot simplify mesh"));
+			return false;
+		}
+		IMeshReduction* MeshReduction = MeshReductionModule->GetStaticMeshReductionInterface();
+		if (!MeshReduction)
+		{
+			UE_LOG(LogGeometry, Warning, TEXT("Failed to load mesh reduction interface; cannot simplify mesh"));
+			return false;
+		}
+
+		FMeshDescription SrcMeshDescription;
+		FStaticMeshAttributes Attributes(SrcMeshDescription);
+		Attributes.Register();
+
+		FDynamicMeshToMeshDescription Converter;
+		Converter.Convert(&Mesh, SrcMeshDescription, true);
+		float Percent = (float)TargetCount / (float)(bTargetIsTriCount ? Mesh.TriangleCount() : Mesh.VertexCount());
+		return FSimplifyMeshOp::ComputeStandardSimplifier(MeshReduction, SrcMeshDescription, Mesh, Percent, bTargetIsTriCount, false, nullptr);
+#else
+		return false;
+#endif
+	}
+}
+
 
 template<typename SimplificationType>
 void DoSimplifyMesh(
@@ -168,6 +210,48 @@ void DoSimplifyMesh(
 	}
 }
 
+UDynamicMesh* UGeometryScriptLibrary_MeshSimplifyFunctions::ApplyEditorSimplifyToTriangleCount(
+	UDynamicMesh* TargetMesh,
+	int32 TriangleCount,
+	UGeometryScriptDebug* Debug)
+{
+	if (TargetMesh == nullptr)
+	{
+		UE::Geometry::AppendError(Debug, EGeometryScriptErrorType::InvalidInputs, LOCTEXT("ApplyEditorSimplifyToTriangleCount_InvalidInput", "ApplyEditorSimplifyToTriangleCount: TargetMesh is Null"));
+		return TargetMesh;
+	}
+
+	TargetMesh->EditMesh([&](FDynamicMesh3& EditMesh) 
+	{
+		UE::PrivateSimplifyHelper::UEStandardEditorSimplify(EditMesh, true, TriangleCount);
+	}, EDynamicMeshChangeType::GeneralEdit, EDynamicMeshAttributeChangeFlags::Unknown, false);
+
+	return TargetMesh;
+}
+
+
+
+
+UDynamicMesh* UGeometryScriptLibrary_MeshSimplifyFunctions::ApplyEditorSimplifyToVertexCount(
+	UDynamicMesh* TargetMesh,
+	int32 VertexCount,
+	UGeometryScriptDebug* Debug)
+{
+	if (TargetMesh == nullptr)
+	{
+		UE::Geometry::AppendError(Debug, EGeometryScriptErrorType::InvalidInputs, LOCTEXT("ApplyEditorSimplifyToVertexCount_InvalidInput", "ApplySimplifyToVertexCount: TargetMesh is Null"));
+		return TargetMesh;
+	}
+
+	TargetMesh->EditMesh([&](FDynamicMesh3& EditMesh) 
+	{
+		UE::PrivateSimplifyHelper::UEStandardEditorSimplify(EditMesh, false, VertexCount);
+	}, EDynamicMeshChangeType::GeneralEdit, EDynamicMeshAttributeChangeFlags::Unknown, false);
+
+	return TargetMesh;
+}
+
+
 
 
 UDynamicMesh* UGeometryScriptLibrary_MeshSimplifyFunctions::ApplySimplifyToTriangleCount(
diff --git a/Engine/Plugins/Runtime/GeometryScripting/Source/GeometryScriptingCore/Public/GeometryScript/MeshSimplifyFunctions.h b/Engine/Plugins/Runtime/GeometryScripting/Source/GeometryScriptingCore/Public/GeometryScript/MeshSimplifyFunctions.h
--- a/Engine/Plugins/Runtime/GeometryScripting/Source/GeometryScriptingCore/Public/GeometryScript/MeshSimplifyFunctions.h
+++ b/Engine/Plugins/Runtime/GeometryScripting/Source/GeometryScriptingCore/Public/GeometryScript/MeshSimplifyFunctions.h
@@ -130,6 +130,26 @@ public:
 		FGeometryScriptSimplifyMeshOptions Options,
 		UGeometryScriptDebug* Debug = nullptr);
 
+	/**
+	* Simplifies the mesh until a target triangle count is reached, using the UE Editor's standard mesh simplifier. Editor only.
+	*/
+	UFUNCTION(BlueprintCallable, Category = "GeometryScript|Simplification", meta=(ScriptMethod))
+	static UPARAM(DisplayName = "Target Mesh") UDynamicMesh* 
+	ApplyEditorSimplifyToTriangleCount(  
+		UDynamicMesh* TargetMesh, 
+		int32 TriangleCount,
+		UGeometryScriptDebug* Debug = nullptr);
+
+	/**
+	* Simplifies the mesh until a target vertex count is reached, using the UE Editor's standard mesh simplifier. Editor only.
+	*/
+	UFUNCTION(BlueprintCallable, Category = "GeometryScript|Simplification", meta=(ScriptMethod))
+	static UPARAM(DisplayName = "Target Mesh") UDynamicMesh* 
+	ApplyEditorSimplifyToVertexCount(  
+		UDynamicMesh* TargetMesh, 
+		int32 VertexCount,
+		UGeometryScriptDebug* Debug = nullptr);
+
 	/**
 	* Simplifies the mesh to a target geometric tolerance. Stops when any further simplification would result in a deviation from the input mesh larger than the tolerance.
 	* Behavior can be additionally controlled with the Options. 
diff --git a/Engine/Plugins/Runtime/MeshModelingToolset/Source/ModelingOperatorsEditorOnly/Private/CleaningOps/SimplifyMeshOp.cpp b/Engine/Plugins/Runtime/MeshModelingToolset/Source/ModelingOperatorsEditorOnly/Private/CleaningOps/SimplifyMeshOp.cpp
--- a/Engine/Plugins/Runtime/MeshModelingToolset/Source/ModelingOperatorsEditorOnly/Private/CleaningOps/SimplifyMeshOp.cpp
+++ b/Engine/Plugins/Runtime/MeshModelingToolset/Source/ModelingOperatorsEditorOnly/Private/CleaningOps/SimplifyMeshOp.cpp
@@ -128,6 +128,99 @@ void ComputeSimplify(FDynamicMesh3* TargetMesh, const bool bReproject,
 	}
 }
 
+bool FSimplifyMeshOp::ComputeStandardSimplifier(IMeshReduction* MeshReduction, const FMeshDescription& SrcMeshDescription, FDynamicMesh3& OutResult, 
+							float PercentReduction, bool bTriBasedReduction, bool bDiscardAttributes, FProgressCancel* Progress)
+{
+	// Method must be called with a valid mesh reduction interface
+	if (!ensure(MeshReduction))
+	{
+		return false;
+	}
+
+	// Note: Simplifier cannot run in place, so we always need to make this copy (even if the SrcMeshDescription was temporary ...)
+	FMeshDescription DstMeshDescription(SrcMeshDescription);
+
+	if (Progress && Progress->Cancelled())
+	{
+		return false;
+	}
+
+	FOverlappingCorners OverlappingCorners;
+	FStaticMeshOperations::FindOverlappingCorners(OverlappingCorners, SrcMeshDescription, 1.e-5);
+
+	if (Progress && Progress->Cancelled())
+	{
+		return false;
+	}
+
+	FMeshReductionSettings ReductionSettings;
+	ReductionSettings.TerminationCriterion = bTriBasedReduction ? EStaticMeshReductionTerimationCriterion::Triangles : EStaticMeshReductionTerimationCriterion::Vertices;
+	ReductionSettings.PercentTriangles = PercentReduction;
+	ReductionSettings.PercentVertices = PercentReduction;
+
+	float FoundMaxDeviation_Unused = 0.f;
+	MeshReduction->ReduceMeshDescription(DstMeshDescription, FoundMaxDeviation_Unused, SrcMeshDescription, OverlappingCorners, ReductionSettings);
+
+	if (Progress && Progress->Cancelled())
+	{
+		return false;
+	}
+
+	// Put the reduced mesh into the target...
+	FMeshDescriptionToDynamicMesh Converter;
+	Converter.Convert(&DstMeshDescription, OutResult);
+	if (bDiscardAttributes)
+	{
+		OutResult.DiscardAttributes();
+	}
+
+
+	bool bFailedModifyNeedsRegen = false;
+	// The UEStandard simplifier will split the UV boundaries.  Need to weld this.
+	{
+		FDynamicMesh3* ComponentMesh = &OutResult;
+
+		FMergeCoincidentMeshEdges Merger(ComponentMesh);
+		Merger.MergeSearchTolerance = 10.0f * FMathf::ZeroTolerance;
+		Merger.OnlyUniquePairs = false;
+		if (Merger.Apply() == false)
+		{
+			bFailedModifyNeedsRegen = true;
+		}
+
+		if (Progress && Progress->Cancelled())
+		{
+			return false;
+		}
+
+		// TODO: This CheckValidity has been here a long time, but it's not clear why we need it ...
+		// If MergeConcidentMeshEdges is creating invalid meshes, we probably need to fix it rather than workaround it here.
+		// For now, have changed the fail mode to Ensure so we can see if/when this actually happens.
+		if (OutResult.CheckValidity(true, EValidityCheckFailMode::Ensure) == false)
+		{
+			bFailedModifyNeedsRegen = true;
+		}
+
+		if (Progress && Progress->Cancelled())
+		{
+			return false;
+		}
+
+		// in the fallback case where merge edges failed, give up and reset it to what it was before the attempted merger (w/ split UV boundaries everywhere, oh well)
+		if (bFailedModifyNeedsRegen)
+		{
+			OutResult.Clear();
+			Converter.Convert(&DstMeshDescription, OutResult);
+			if (bDiscardAttributes)
+			{
+				OutResult.DiscardAttributes();
+			}
+		}
+	}
+
+	return true;
+}
+
 
 
 void FSimplifyMeshOp::CalculateResult(FProgressCancel* Progress)
@@ -222,108 +315,42 @@ void FSimplifyMeshOp::CalculateResult(FProgressCancel* Progress)
 	}
 	else // SimplifierType == ESimplifyType::UEStandard
 	{
-		const FMeshDescription* SrcMeshDescription = OriginalMeshDescription.Get();
-		FMeshDescription DstMeshDescription(*SrcMeshDescription);
-
-		if (Progress && Progress->Cancelled())
+		if (!MeshReduction)
 		{
+			// no reduction possible, failed to load the required interface
+			ResultMesh->Copy(*OriginalMesh, true, true, true, !bDiscardAttributes);
 			return;
 		}
 
-		FOverlappingCorners OverlappingCorners;
-		FStaticMeshOperations::FindOverlappingCorners(OverlappingCorners, *SrcMeshDescription, 1.e-5);
+		const FMeshDescription* SrcMeshDescription = OriginalMeshDescription.Get();
 
 		if (Progress && Progress->Cancelled())
 		{
 			return;
 		}
 
-		FMeshReductionSettings ReductionSettings;
+		float PercentReduction = 1.f;
 		if (TargetMode == ESimplifyTargetType::Percentage)
 		{
-			ReductionSettings.PercentTriangles = FMath::Max(TargetPercentage / 100., .001);  // Only support triangle percentage and count, but not edge length
-			ReductionSettings.TerminationCriterion = EStaticMeshReductionTerimationCriterion::Triangles;
+			PercentReduction = FMath::Max(TargetPercentage / 100., .001);
 		}
 		else if (TargetMode == ESimplifyTargetType::TriangleCount)
 		{
 			int32 NumTris = SrcMeshDescription->Polygons().Num();
-			ReductionSettings.PercentTriangles = (float)TargetCount / (float)NumTris;
-			ReductionSettings.TerminationCriterion = EStaticMeshReductionTerimationCriterion::Triangles;
+			PercentReduction = (float)TargetCount / (float)NumTris;
 		}
 		else if (TargetMode == ESimplifyTargetType::VertexCount)
 		{
 			int32 NumVerts = SrcMeshDescription->Vertices().Num();
-			ReductionSettings.PercentVertices = (float)TargetCount / (float)NumVerts;
-			ReductionSettings.TerminationCriterion = EStaticMeshReductionTerimationCriterion::Vertices;
+			PercentReduction = (float)TargetCount / (float)NumVerts;
 		}
-
-		float Error;
-		{
-			if (!MeshReduction)
-			{
-				// no reduction possible, failed to load the required interface
-				Error = 0.f;
-				ResultMesh->Copy(*OriginalMesh, true, true, true, !bDiscardAttributes);
-				return;
-			}
-
-			Error = ReductionSettings.MaxDeviation;
-			MeshReduction->ReduceMeshDescription(DstMeshDescription, Error, *SrcMeshDescription, OverlappingCorners, ReductionSettings);
-		}
-
-		if (Progress && Progress->Cancelled())
+		bool bTargetIsTriangleCount = TargetMode != ESimplifyTargetType::VertexCount;
+		
+		if (!ComputeStandardSimplifier(MeshReduction, *SrcMeshDescription, *ResultMesh, PercentReduction, bTargetIsTriangleCount, 
+			bDiscardAttributes, Progress))
 		{
 			return;
 		}
-
-		// Put the reduced mesh into the target...
-		FMeshDescriptionToDynamicMesh Converter;
-		Converter.Convert(&DstMeshDescription, *ResultMesh);
-		if (bDiscardAttributes)
-		{
-			ResultMesh->DiscardAttributes();
-		}
-
-
-		bool bFailedModifyNeedsRegen = false;
-		// The UEStandard simplifier will split the UV boundaries.  Need to weld this.
-		{
-			FDynamicMesh3* ComponentMesh = ResultMesh.Get();
-
-			FMergeCoincidentMeshEdges Merger(ComponentMesh);
-			Merger.MergeSearchTolerance = 10.0f * FMathf::ZeroTolerance;
-			Merger.OnlyUniquePairs = false;
-			if (Merger.Apply() == false)
-			{
-				bFailedModifyNeedsRegen = true;
-			}
-
-			if (Progress && Progress->Cancelled())
-			{
-				return;
-			}
-
-			if (ResultMesh->CheckValidity(true, EValidityCheckFailMode::ReturnOnly) == false)
-			{
-				bFailedModifyNeedsRegen = true;
-			}
-
-			if (Progress && Progress->Cancelled())
-			{
-				return;
-			}
-
-			// in the fallback case where merge edges failed, give up and reset it to what it was before the attempted merger (w/ split UV boundaries everywhere, oh well)
-			if (bFailedModifyNeedsRegen)
-			{
-				ResultMesh->Clear();
-				Converter.Convert(&DstMeshDescription, *ResultMesh);
-				if (bDiscardAttributes)
-				{
-					ResultMesh->DiscardAttributes();
-				}
-			}
-		}
 	}
 
 
diff --git a/Engine/Plugins/Runtime/MeshModelingToolset/Source/ModelingOperatorsEditorOnly/Public/CleaningOps/SimplifyMeshOp.h b/Engine/Plugins/Runtime/MeshModelingToolset/Source/ModelingOperatorsEditorOnly/Public/CleaningOps/SimplifyMeshOp.h
--- a/Engine/Plugins/Runtime/MeshModelingToolset/Source/ModelingOperatorsEditorOnly/Public/CleaningOps/SimplifyMeshOp.h
+++ b/Engine/Plugins/Runtime/MeshModelingToolset/Source/ModelingOperatorsEditorOnly/Public/CleaningOps/SimplifyMeshOp.h
@@ -96,6 +96,21 @@ public:
 
 	IMeshReduction* MeshReduction;
 
+	/**
+	 * Simple helper to compute UE Standard simplified mesh
+	 * 
+	 * @param MeshReduction Interface for mesh reduction; must be non-null
+	 * @param SrcMeshDescription Input mesh
+	 * @param OutResult Simplified mesh will be written to this mesh
+	 * @param PercentReduction Percent to reduce
+	 * @param bTriBasedReduction Whether to reduce based on triangles; otherwise, will reduce based on vertices
+	 * @param bDiscardAttributes Whether to discard attributes when simplifying
+	 * @param Progress If non-null, can early-out if this is set to request cancellation
+	 * @return true if succeeded, false otherwise (e.g. if cancelled)
+	 */
+	static bool ComputeStandardSimplifier(IMeshReduction* MeshReduction, const FMeshDescription& SrcMeshDescription, FDynamicMesh3& OutResult,
+		float PercentReduction, bool bTriBasedReduction, bool bDiscardAttributes,
+		FProgressCancel* Progress = nullptr);
 
 	// set ability on protected transform.
 	void SetTransform(const FTransform& Transform)
