Modeling Mode: Add pressure sensitivity support for brush strength in the Vertex Sculpt Tool

diff --git a/Engine/Plugins/Editor/ModelingToolsEditorMode/Source/ModelingToolsEditorMode/Private/DetailsCustomizations/MeshVertexSculptToolCustomizations.cpp b/Engine/Plugins/Editor/ModelingToolsEditorMode/Source/ModelingToolsEditorMode/Private/DetailsCustomizations/MeshVertexSculptToolCustomizations.cpp
--- a/Engine/Plugins/Editor/ModelingToolsEditorMode/Source/ModelingToolsEditorMode/Private/DetailsCustomizations/MeshVertexSculptToolCustomizations.cpp
+++ b/Engine/Plugins/Editor/ModelingToolsEditorMode/Source/ModelingToolsEditorMode/Private/DetailsCustomizations/MeshVertexSculptToolCustomizations.cpp
@@ -20,6 +20,13 @@
 #include "MeshVertexSculptTool.h"
 #include "Widgets/Layout/SBox.h"
 #include "Widgets/SBoxPanel.h"
+#include "Sculpting/KelvinletBrushOp.h"
+#include "Sculpting/MeshInflateBrushOps.h"
+#include "Sculpting/MeshMoveBrushOps.h"
+#include "Sculpting/MeshPinchBrushOps.h"
+#include "Sculpting/MeshPlaneBrushOps.h"
+#include "Sculpting/MeshSculptBrushOps.h"
+#include "Sculpting/MeshSmoothingBrushOps.h"
 
 using namespace UE::ModelingUI;
 
@@ -106,6 +113,120 @@ void FSculptBrushPropertiesDetails::CustomizeDetails(IDetailLayoutBuilder& Detai
 
 }
 
+template<typename BrushClass>
+TSharedRef<IDetailCustomization> FMeshSculptBrushOpPropertiesDetails<BrushClass>::MakeInstance()
+{
+	return MakeShareable(new FMeshSculptBrushOpPropertiesDetails<BrushClass>);
+}
+
+template<typename BrushClass>
+void FMeshSculptBrushOpPropertiesDetails<BrushClass>::CustomizeDetails(IDetailLayoutBuilder& DetailBuilder)
+{
+	TArray<TWeakObjectPtr<UObject>> ObjectsBeingCustomized;
+	DetailBuilder.GetObjectsBeingCustomized(ObjectsBeingCustomized);
+	check(ObjectsBeingCustomized.Num() > 0);
+
+	// brush op properties
+	BrushClass* BrushOpProperties = Cast<BrushClass>(ObjectsBeingCustomized[0]);
+
+	// toggles pressure sensitivity for brush strength
+	const TSharedPtr<IPropertyHandle> EnablePressureSensitivity = DetailBuilder.GetProperty(GET_MEMBER_NAME_CHECKED(UMeshSculptBrushOpProps, bIsStrengthPressureEnabled), UMeshSculptBrushOpProps::StaticClass());
+	ensure(EnablePressureSensitivity->IsValidHandle());
+	EnablePressureSensitivity->MarkHiddenByCustomization(); // always hidden
+
+	// do the tool & brush support pressure sensitivity for brush strength?
+	const bool bSupportsPressureSensitivity = BrushOpProperties->SupportsStrengthPressure();
+
+	// controls brush strength
+	const TSharedPtr<IPropertyHandle> StrengthHandle = DetailBuilder.GetProperty("Strength", BrushClass::StaticClass());
+
+	if (bSupportsPressureSensitivity && ensure(StrengthHandle->IsValidHandle()))
+	{
+		// build customized widget representing brush strength & pressure sensitivity toggle (when applicable)
+		BuildStrengthPressureWidget(DetailBuilder, StrengthHandle, EnablePressureSensitivity);
+	}
+}
+
+template<typename BrushClass>
+void FMeshSculptBrushOpPropertiesDetails<BrushClass>::BuildStrengthPressureWidget(IDetailLayoutBuilder& DetailBuilder, const TSharedPtr<IPropertyHandle>& StrengthHandle, TSharedPtr<IPropertyHandle> EnablePressureSensitivityHandle)
+{
+	IDetailPropertyRow* DetailRow = DetailBuilder.EditDefaultProperty(StrengthHandle);
+	TSharedPtr<SWidget> NameWidget, ValueWidget;
+	DetailRow->GetDefaultWidgets(NameWidget, ValueWidget);
+
+	TSharedPtr<SHorizontalBox> ValueContent;
+
+	DetailRow->CustomWidget()
+	.NameContent() // strength text
+	[
+		NameWidget->AsShared()
+	]
+	.ValueContent() // strength sspinbox and pressure sensitivity toggle, when applicable
+	[
+		SAssignNew(ValueContent, SHorizontalBox)
+	];
+
+	// always add slider/SSpinbox to control strength, as usual
+	ValueContent->AddSlot()
+	[
+		ValueWidget->AsShared()
+	];
+
+	// add pressure sensitivity toggle
+	ValueContent->AddSlot()
+	.AutoWidth()
+	[
+		SNew(SCheckBox)
+		.Style(FAppStyle::Get(), "DetailsView.SectionButton")
+		.Padding(FMargin(4, 2))
+		.ToolTipText(EnablePressureSensitivityHandle->GetToolTipText())
+		.HAlign(HAlign_Center)
+		.OnCheckStateChanged_Lambda([EnablePressureSensitivityHandle](const ECheckBoxState NewState)
+		{
+			EnablePressureSensitivityHandle->SetValue(NewState == ECheckBoxState::Checked);
+		})
+		.IsChecked_Lambda([EnablePressureSensitivityHandle]() -> ECheckBoxState
+		{
+			bool bSet;
+			EnablePressureSensitivityHandle->GetValue(bSet);
+			return bSet ? ECheckBoxState::Checked : ECheckBoxState::Unchecked;
+		})
+		[
+			SNew(SHorizontalBox)
+			+ SHorizontalBox::Slot()
+			.VAlign(VAlign_Center)
+			.Padding(FMargin(0))
+			.AutoWidth()
+			[
+				SNew(SImage)
+				.Image(FModelingToolsEditorModeStyle::Get()->GetBrush("BrushIcons.PressureSensitivity"))
+			]
+		]
+	];
+}
+
+/* Explicitly instantiate the types which are supported */
+
+template class FMeshSculptBrushOpPropertiesDetails<UPinchBrushOpProps>;
+template class FMeshSculptBrushOpPropertiesDetails<UInflateBrushOpProps>;
+template class FMeshSculptBrushOpPropertiesDetails<USmoothBrushOpProps>;
+template class FMeshSculptBrushOpPropertiesDetails<USmoothFillBrushOpProps>;
+template class FMeshSculptBrushOpPropertiesDetails<UFlattenBrushOpProps>;
+template class FMeshSculptBrushOpPropertiesDetails<UEraseBrushOpProps>;
+template class FMeshSculptBrushOpPropertiesDetails<UStandardSculptBrushOpProps>;
+template class FMeshSculptBrushOpPropertiesDetails<UViewAlignedSculptBrushOpProps>;
+template class FMeshSculptBrushOpPropertiesDetails<USculptMaxBrushOpProps>;
+template class FMeshSculptBrushOpPropertiesDetails<UFixedPlaneBrushOpProps>;
+template class FMeshSculptBrushOpPropertiesDetails<UViewAlignedPlaneBrushOpProps>;
+template class FMeshSculptBrushOpPropertiesDetails<UPlaneBrushOpProps>;
+template class FMeshSculptBrushOpPropertiesDetails<UMoveBrushOpProps>;
+template class FMeshSculptBrushOpPropertiesDetails<USecondarySmoothBrushOpProps>;
+template class FMeshSculptBrushOpPropertiesDetails<UScaleKelvinletBrushOpProps>;
+template class FMeshSculptBrushOpPropertiesDetails<UPullKelvinletBrushOpProps>;
+template class FMeshSculptBrushOpPropertiesDetails<USharpPullKelvinletBrushOpProps>;
+template class FMeshSculptBrushOpPropertiesDetails<UTwistKelvinletBrushOpProps>;
+
+/***********************************************************/
 
 TSharedRef<IDetailCustomization> FVertexBrushSculptPropertiesDetails::MakeInstance()
 {
diff --git a/Engine/Plugins/Editor/ModelingToolsEditorMode/Source/ModelingToolsEditorMode/Private/ModelingToolsEditorModeModule.cpp b/Engine/Plugins/Editor/ModelingToolsEditorMode/Source/ModelingToolsEditorMode/Private/ModelingToolsEditorModeModule.cpp
--- a/Engine/Plugins/Editor/ModelingToolsEditorMode/Source/ModelingToolsEditorMode/Private/ModelingToolsEditorModeModule.cpp
+++ b/Engine/Plugins/Editor/ModelingToolsEditorMode/Source/ModelingToolsEditorMode/Private/ModelingToolsEditorModeModule.cpp
@@ -24,6 +24,13 @@
 #include "BakeMultiMeshAttributeMapsTool.h"
 #include "BakeMeshAttributeVertexTool.h"
 #include "BakeTransformTool.h"
+#include "Sculpting/KelvinletBrushOp.h"
+#include "Sculpting/MeshInflateBrushOps.h"
+#include "Sculpting/MeshMoveBrushOps.h"
+#include "Sculpting/MeshPinchBrushOps.h"
+#include "Sculpting/MeshPlaneBrushOps.h"
+#include "Sculpting/MeshSculptBrushOps.h"
+#include "Sculpting/MeshSmoothingBrushOps.h"
 
 
 #define LOCTEXT_NAMESPACE "FModelingToolsEditorModeModule"
@@ -89,6 +96,44 @@ void FModelingToolsEditorModeModule::OnPostEngineInit()
 	ClassesToUnregisterOnShutdown.Add(UVertexBrushSculptProperties::StaticClass()->GetFName());
 	PropertyModule.RegisterCustomClassLayout("VertexBrushAlphaProperties", FOnGetDetailCustomizationInstance::CreateStatic(&FVertexBrushAlphaPropertiesDetails::MakeInstance));
 	ClassesToUnregisterOnShutdown.Add(UVertexBrushAlphaProperties::StaticClass()->GetFName());
+	// Sculpt - BrushOpProps
+	PropertyModule.RegisterCustomClassLayout("PinchBrushOpProps", FOnGetDetailCustomizationInstance::CreateStatic(&FMeshSculptBrushOpPropertiesDetails<UPinchBrushOpProps>::MakeInstance));
+	ClassesToUnregisterOnShutdown.Add(UPinchBrushOpProps::StaticClass()->GetFName());
+	PropertyModule.RegisterCustomClassLayout("InflateBrushOpProps", FOnGetDetailCustomizationInstance::CreateStatic(&FMeshSculptBrushOpPropertiesDetails<UInflateBrushOpProps>::MakeInstance));
+	ClassesToUnregisterOnShutdown.Add(UInflateBrushOpProps::StaticClass()->GetFName());
+	PropertyModule.RegisterCustomClassLayout("SmoothBrushOpProps", FOnGetDetailCustomizationInstance::CreateStatic(&FMeshSculptBrushOpPropertiesDetails<USmoothBrushOpProps>::MakeInstance));
+	ClassesToUnregisterOnShutdown.Add(USmoothBrushOpProps::StaticClass()->GetFName());
+	PropertyModule.RegisterCustomClassLayout("SmoothFillBrushOpProps", FOnGetDetailCustomizationInstance::CreateStatic(&FMeshSculptBrushOpPropertiesDetails<USmoothFillBrushOpProps>::MakeInstance));
+	ClassesToUnregisterOnShutdown.Add(USmoothFillBrushOpProps::StaticClass()->GetFName());
+	PropertyModule.RegisterCustomClassLayout("FlattenBrushOpProps", FOnGetDetailCustomizationInstance::CreateStatic(&FMeshSculptBrushOpPropertiesDetails<UFlattenBrushOpProps>::MakeInstance));
+	ClassesToUnregisterOnShutdown.Add(UFlattenBrushOpProps::StaticClass()->GetFName());
+	PropertyModule.RegisterCustomClassLayout("EraseBrushOpProps", FOnGetDetailCustomizationInstance::CreateStatic(&FMeshSculptBrushOpPropertiesDetails<UEraseBrushOpProps>::MakeInstance));
+	ClassesToUnregisterOnShutdown.Add(UEraseBrushOpProps::StaticClass()->GetFName());
+	PropertyModule.RegisterCustomClassLayout("StandardSculptBrushOpProps", FOnGetDetailCustomizationInstance::CreateStatic(&FMeshSculptBrushOpPropertiesDetails<UStandardSculptBrushOpProps>::MakeInstance));
+	ClassesToUnregisterOnShutdown.Add(UStandardSculptBrushOpProps::StaticClass()->GetFName());
+	PropertyModule.RegisterCustomClassLayout("ViewAlignedSculptBrushOpProps", FOnGetDetailCustomizationInstance::CreateStatic(&FMeshSculptBrushOpPropertiesDetails<UViewAlignedSculptBrushOpProps>::MakeInstance));
+	ClassesToUnregisterOnShutdown.Add(UViewAlignedSculptBrushOpProps::StaticClass()->GetFName());
+	PropertyModule.RegisterCustomClassLayout("SculptMaxBrushOpProps", FOnGetDetailCustomizationInstance::CreateStatic(&FMeshSculptBrushOpPropertiesDetails<USculptMaxBrushOpProps>::MakeInstance));
+	ClassesToUnregisterOnShutdown.Add(USculptMaxBrushOpProps::StaticClass()->GetFName());
+	PropertyModule.RegisterCustomClassLayout("FixedPlaneBrushOpProps", FOnGetDetailCustomizationInstance::CreateStatic(&FMeshSculptBrushOpPropertiesDetails<UFixedPlaneBrushOpProps>::MakeInstance));
+	ClassesToUnregisterOnShutdown.Add(UFixedPlaneBrushOpProps::StaticClass()->GetFName());
+	PropertyModule.RegisterCustomClassLayout("ViewAlignedPlaneBrushOpProps", FOnGetDetailCustomizationInstance::CreateStatic(&FMeshSculptBrushOpPropertiesDetails<UViewAlignedPlaneBrushOpProps>::MakeInstance));
+	ClassesToUnregisterOnShutdown.Add(UViewAlignedPlaneBrushOpProps::StaticClass()->GetFName());
+	PropertyModule.RegisterCustomClassLayout("PlaneBrushOpProps", FOnGetDetailCustomizationInstance::CreateStatic(&FMeshSculptBrushOpPropertiesDetails<UPlaneBrushOpProps>::MakeInstance));
+	ClassesToUnregisterOnShutdown.Add(UPlaneBrushOpProps::StaticClass()->GetFName());
+	PropertyModule.RegisterCustomClassLayout("MoveBrushOpProps", FOnGetDetailCustomizationInstance::CreateStatic(&FMeshSculptBrushOpPropertiesDetails<UMoveBrushOpProps>::MakeInstance));
+	ClassesToUnregisterOnShutdown.Add(UMoveBrushOpProps::StaticClass()->GetFName());
+	PropertyModule.RegisterCustomClassLayout("SecondarySmoothBrushOpProps", FOnGetDetailCustomizationInstance::CreateStatic(&FMeshSculptBrushOpPropertiesDetails<USecondarySmoothBrushOpProps>::MakeInstance));
+	ClassesToUnregisterOnShutdown.Add(USecondarySmoothBrushOpProps::StaticClass()->GetFName());
+	PropertyModule.RegisterCustomClassLayout("ScaleKelvinletBrushOpProps", FOnGetDetailCustomizationInstance::CreateStatic(&FMeshSculptBrushOpPropertiesDetails<UScaleKelvinletBrushOpProps>::MakeInstance));
+	ClassesToUnregisterOnShutdown.Add(UScaleKelvinletBrushOpProps::StaticClass()->GetFName());
+	PropertyModule.RegisterCustomClassLayout("PullKelvinletBrushOpProps", FOnGetDetailCustomizationInstance::CreateStatic(&FMeshSculptBrushOpPropertiesDetails<UPullKelvinletBrushOpProps>::MakeInstance));
+	ClassesToUnregisterOnShutdown.Add(UPullKelvinletBrushOpProps::StaticClass()->GetFName());
+	PropertyModule.RegisterCustomClassLayout("SharpPullKelvinletBrushOpProps", FOnGetDetailCustomizationInstance::CreateStatic(&FMeshSculptBrushOpPropertiesDetails<USharpPullKelvinletBrushOpProps>::MakeInstance));
+	ClassesToUnregisterOnShutdown.Add(USharpPullKelvinletBrushOpProps::StaticClass()->GetFName());
+	PropertyModule.RegisterCustomClassLayout("TwistKelvinletBrushOpProps", FOnGetDetailCustomizationInstance::CreateStatic(&FMeshSculptBrushOpPropertiesDetails<UTwistKelvinletBrushOpProps>::MakeInstance));
+	ClassesToUnregisterOnShutdown.Add(UTwistKelvinletBrushOpProps::StaticClass()->GetFName());
+
 	/// Bake
 	PropertyModule.RegisterCustomClassLayout("BakeMeshAttributeMapsToolProperties", FOnGetDetailCustomizationInstance::CreateStatic(&FBakeMeshAttributeMapsToolDetails::MakeInstance));
 	ClassesToUnregisterOnShutdown.Add(UBakeMeshAttributeMapsToolProperties::StaticClass()->GetFName());
diff --git a/Engine/Plugins/Editor/ModelingToolsEditorMode/Source/ModelingToolsEditorMode/Public/DetailsCustomizations/MeshVertexSculptToolCustomizations.h b/Engine/Plugins/Editor/ModelingToolsEditorMode/Source/ModelingToolsEditorMode/Public/DetailsCustomizations/MeshVertexSculptToolCustomizations.h
--- a/Engine/Plugins/Editor/ModelingToolsEditorMode/Source/ModelingToolsEditorMode/Public/DetailsCustomizations/MeshVertexSculptToolCustomizations.h
+++ b/Engine/Plugins/Editor/ModelingToolsEditorMode/Source/ModelingToolsEditorMode/Public/DetailsCustomizations/MeshVertexSculptToolCustomizations.h
@@ -4,6 +4,7 @@
 
 #include "IDetailCustomization.h"
 #include "UObject/WeakObjectPtr.h"
+#include "DetailLayoutBuilder.h"
 
 class FReply;
 class SButton;
@@ -77,3 +78,20 @@ protected:
 	TSharedPtr<SToolInputAssetComboPanel> AlphaAssetPicker;
 	FDelegateHandle AlphaTextureUpdateHandle;
 };
+
+// customization for UMeshSculptBrushOpProps
+// - creates custom brush strength w/ pressure sensitivity toggle, when applicable
+template<typename BrushClass>
+class FMeshSculptBrushOpPropertiesDetails : public IDetailCustomization
+{
+public:
+	static TSharedRef<IDetailCustomization> MakeInstance();
+	virtual void CustomizeDetails(IDetailLayoutBuilder& DetailBuilder) override;
+private:
+	// helper to build strength widget, with pressure sensitivity toggle when applicable
+	static void BuildStrengthPressureWidget(
+		IDetailLayoutBuilder& DetailBuilder,
+		const TSharedPtr<IPropertyHandle>& StrengthHandle,
+		TSharedPtr<IPropertyHandle> EnablePressureSensitivityHandle
+	);
+};
\ No newline at end of file
diff --git a/Engine/Plugins/Runtime/MeshModelingToolset/Source/MeshModelingTools/Private/MeshVertexSculptTool.cpp b/Engine/Plugins/Runtime/MeshModelingToolset/Source/MeshModelingTools/Private/MeshVertexSculptTool.cpp
--- a/Engine/Plugins/Runtime/MeshModelingToolset/Source/MeshModelingTools/Private/MeshVertexSculptTool.cpp
+++ b/Engine/Plugins/Runtime/MeshModelingToolset/Source/MeshModelingTools/Private/MeshVertexSculptTool.cpp
@@ -497,7 +497,7 @@ void UMeshVertexSculptTool::OnBeginStroke(const FRay& WorldRay)
 	LastStamp.Falloff = GetCurrentBrushFalloff();
 	LastStamp.Direction = GetInInvertStroke() ? -1.0 : 1.0;
 	LastStamp.Depth = GetCurrentBrushDepth();
-	LastStamp.Power = GetActivePressure() * GetCurrentBrushStrength();
+	LastStamp.Power = GetActiveBrushStrength();
 	LastStamp.TimeStamp = FDateTime::Now();
 
 	PreviousRayDirection = FVector3d::ZeroVector;
@@ -917,7 +917,7 @@ bool UMeshVertexSculptTool::UpdateStampPosition(const FRay& WorldRay)
 	CurrentStamp.WorldFrame = GetBrushFrameWorld();
 	CurrentStamp.Radius = GetActiveBrushRadius();
 	CurrentStamp.LocalFrame = GetBrushFrameLocal();
-	CurrentStamp.Power = GetActivePressure() * GetCurrentBrushStrength();
+	CurrentStamp.Power = GetActiveBrushStrength();
 
 	if (bHaveBrushAlpha && (AlphaProperties->RotationAngle != 0 || AlphaProperties->bRandomize))
 	{
diff --git a/Engine/Plugins/Runtime/MeshModelingToolset/Source/MeshModelingTools/Private/Sculpting/KelvinletBrushOp.h b/Engine/Plugins/Runtime/MeshModelingToolset/Source/MeshModelingTools/Private/Sculpting/KelvinletBrushOp.h
deleted file mode 100644
diff --git a/Engine/Plugins/Runtime/MeshModelingToolset/Source/MeshModelingTools/Private/Sculpting/MeshSculptToolBase.cpp b/Engine/Plugins/Runtime/MeshModelingToolset/Source/MeshModelingTools/Private/Sculpting/MeshSculptToolBase.cpp
--- a/Engine/Plugins/Runtime/MeshModelingToolset/Source/MeshModelingTools/Private/Sculpting/MeshSculptToolBase.cpp
+++ b/Engine/Plugins/Runtime/MeshModelingToolset/Source/MeshModelingTools/Private/Sculpting/MeshSculptToolBase.cpp
@@ -668,6 +668,15 @@ bool UMeshSculptToolBase::GetBrushSizePressureSensitivityEnabled() const
 	return false;
 }
 
+bool UMeshSculptToolBase::GetBrushStrengthPressureEnabled()
+{
+	TUniquePtr<FMeshSculptBrushOp>& BrushOp = GetActiveBrushOp();
+	if (BrushOp->PropertySet.IsValid())
+	{
+		return BrushOp->PropertySet->GetStrengthPressureEnabled();
+	}
+	return false;
+}
 
 void UMeshSculptToolBase::UpdateBrushFrameWorld(const FVector3d& NewPosition, const FVector3d& NewNormal)
 {
@@ -1003,6 +1012,21 @@ double UMeshSculptToolBase::GetActiveBrushRadius()
 	return CurrentBrushRadius;
 }
 
+double UMeshSculptToolBase::GetActiveBrushStrength()
+{
+	TUniquePtr<FMeshSculptBrushOp>& BrushOp = GetActiveBrushOp();
+	if (BrushOp->PropertySet.IsValid())
+	{
+		double Strength = GetCurrentBrushStrength();
+		if (GetBrushStrengthPressureEnabled())
+		{
+			Strength *= GetActivePressure();
+		}
+
+		return FMathd::Clamp(Strength, 0.0, 1.0);
+	}
+	return 1.0;
+}
 
 
 double UMeshSculptToolBase::GetCurrentBrushStrength()
diff --git a/Engine/Plugins/Runtime/MeshModelingToolset/Source/MeshModelingTools/Public/Sculpting/KelvinletBrushOp.h b/Engine/Plugins/Runtime/MeshModelingToolset/Source/MeshModelingTools/Public/Sculpting/KelvinletBrushOp.h
new file mode 100644
--- /dev/null
+++ b/Engine/Plugins/Runtime/MeshModelingToolset/Source/MeshModelingTools/Public/Sculpting/KelvinletBrushOp.h
@@ -0,0 +1,756 @@
+// Copyright Epic Games, Inc. All Rights Reserved.
+
+#pragma once
+
+#include "MatrixTypes.h"
+#include "DynamicMesh/DynamicMesh3.h"
+#include "Deformers/Kelvinlets.h"
+#include "Sculpting/MeshBrushOpBase.h"
+#include "Sculpting/MeshSculptToolBase.h"
+#include "Async/ParallelFor.h"
+#include "KelvinletBrushOp.generated.h"
+
+// import kelvinlet types
+using UE::Geometry::FPullKelvinlet;
+using UE::Geometry::FLaplacianPullKelvinlet;
+using UE::Geometry::FBiLaplacianPullKelvinlet;
+using UE::Geometry::FSharpLaplacianPullKelvinlet;
+using UE::Geometry::FSharpBiLaplacianPullKelvinlet;
+using UE::Geometry::FAffineKelvinlet;
+using UE::Geometry::FTwistKelvinlet;
+using UE::Geometry::FPinchKelvinlet;
+using UE::Geometry::FScaleKelvinlet;
+using UE::Geometry::FBlendPullSharpKelvinlet;
+using UE::Geometry::FBlendPullKelvinlet;
+using UE::Geometry::FLaplacianTwistPullKelvinlet;
+using UE::Geometry::FBiLaplacianTwistPullKelvinlet;
+
+UCLASS()
+class MESHMODELINGTOOLS_API UBaseKelvinletBrushOpProps : public UMeshSculptBrushOpProps
+{
+	GENERATED_BODY()
+
+public:
+
+	/** How much the mesh resists shear */
+	UPROPERTY()
+	float Stiffness = 1.f;
+
+	/** How compressible the spatial region is: 1 - 2 x Poisson ratio */
+	UPROPERTY()
+	float Incompressiblity = 1.f;
+
+	/** Integration steps*/
+	UPROPERTY()
+	int32 BrushSteps = 3;
+
+	virtual float GetStiffness() { return Stiffness; }
+	virtual float GetIncompressiblity() { return Incompressiblity; }
+	virtual int32 GetNumSteps() { return BrushSteps; }
+
+};
+
+
+
+
+
+class FBaseKelvinletBrushOp : public FMeshSculptBrushOp
+{
+public:
+	
+	void SetBaseProperties(const FDynamicMesh3* SrcMesh, const FSculptBrushStamp& Stamp)
+	{
+		// pointer to mesh
+		Mesh = SrcMesh;
+		UBaseKelvinletBrushOpProps* BaseKelvinletBrushOpProps = GetPropertySetAs<UBaseKelvinletBrushOpProps>();
+
+		float Stiffness = BaseKelvinletBrushOpProps->GetStiffness();
+		float Incompressibility = BaseKelvinletBrushOpProps->GetIncompressiblity();
+		int NumIntSteps = BaseKelvinletBrushOpProps->GetNumSteps();
+
+		Mu = FMath::Max(Stiffness, 0.f);
+		Nu = FMath::Clamp(0.5f * (1.f - 2.f * Incompressibility), 0.f, 0.5f);
+		
+		Size = FMath::Max(Stamp.Radius, 0.01);
+
+		NumSteps = NumIntSteps;
+	}
+
+
+	template <typename KelvinletType>
+	void ApplyKelvinlet(const KelvinletType& Kelvinlet, const FFrame3d& LocalFrame, const TArray<int32>& Vertices, TArray<FVector3d>& NewPositionsOut) const 
+	{
+
+		if (NumSteps == 0)
+		{
+			DisplaceKelvinlet(Kelvinlet, LocalFrame, Vertices, NewPositionsOut);
+		}
+		else
+		{
+			IntegrateKelvinlet(Kelvinlet, LocalFrame, Vertices, NewPositionsOut, IntegrationTime, NumSteps);
+		}
+	}
+
+
+
+	virtual void ApplyStamp(const FDynamicMesh3*SrcMesh, const FSculptBrushStamp& Stamp, const TArray<int32>& Vertices, TArray<FVector3d>& NewPositionsOut) override {}
+
+protected:
+
+	void ApplyFalloff(const FSculptBrushStamp& Stamp, const TArray<int32>& Vertices, TArray<FVector3d>& NewPositionsOut) const
+	{
+		int32 NumV = Vertices.Num();
+		bool bParallel = true;
+		ParallelFor(NumV, [&Stamp, &Vertices, &NewPositionsOut, this](int32 k)
+			{
+				int32 VertIdx = Vertices[k];
+				FVector3d OrigPos = Mesh->GetVertex(VertIdx);
+
+				double Falloff = GetFalloff().Evaluate(Stamp, OrigPos);
+
+				double Alpha = FMath::Clamp(Falloff, 0., 1.);
+
+				NewPositionsOut[k] = Alpha * NewPositionsOut[k] + (1. - Alpha) * OrigPos;
+			}, bParallel);
+	}
+
+
+	// NB: this just moves the verts, but doesn't update the normal.  The kelvinlets will have to be extended if we want
+	// to do the Jacobian Transpose operation on the normals - but for now, we should just rebuild the normals after the brush
+	template <typename KelvinletType>
+	void DisplaceKelvinlet(const KelvinletType& Kelvinlet, const FFrame3d& LocalFrame, const TArray<int32>& Vertices, TArray<FVector3d>& NewPositionsOut) const
+	{
+		int32 NumV = Vertices.Num();
+		NewPositionsOut.SetNum(NumV, EAllowShrinking::No);
+
+		const bool bForceSingleThread = false;
+
+		ParallelFor(NumV, [&Kelvinlet, &Vertices, &NewPositionsOut, &LocalFrame, this](int32 k)
+			{
+				int VertIdx = Vertices[k];
+				FVector3d Pos = Mesh->GetVertex(VertIdx);
+				Pos = LocalFrame.ToFramePoint(Pos);
+
+				Pos = Kelvinlet.Evaluate(Pos) + Pos;
+
+				// Update the position in the ROI Array
+				NewPositionsOut[k] = LocalFrame.FromFramePoint(Pos);
+			}
+		, bForceSingleThread);
+
+	}
+
+	template <typename KelvinletType>
+	void IntegrateKelvinlet(const KelvinletType& Kelvinlet, const FFrame3d& LocalFrame, const TArray<int32>& Vertices, TArray<FVector3d>& NewPositionsOut, const double Dt, const int32 Steps) const
+	{
+		int NumV = Vertices.Num();
+		NewPositionsOut.SetNum(NumV, EAllowShrinking::No);
+
+		const bool bForceSingleThread = false;
+
+		ParallelFor(NumV, [&Kelvinlet, &Vertices, &NewPositionsOut, &LocalFrame, Dt, Steps, this](int32 k)
+			{
+				int VertIdx = Vertices[k];
+				FVector3d Pos = LocalFrame.ToFramePoint(Mesh->GetVertex(VertIdx));
+
+				double TimeScale = 1. / (Steps);
+				// move with several time steps
+				for (int i = 0; i < Steps; ++i)
+				{
+					// the position after deformation
+					Pos = Kelvinlet.IntegrateRK3(Pos, Dt * TimeScale);
+				}
+				// Update the position in the ROI Array
+				NewPositionsOut[k] = LocalFrame.FromFramePoint(Pos);
+			}, bForceSingleThread);
+
+
+	}
+
+protected:
+
+
+	// physical properties.
+	float Mu;
+	float Nu;
+	// model regularization parameter
+	float Size;
+	
+	// integration parameters
+	int32 NumSteps;
+	double IntegrationTime = 1.;
+	
+	const FDynamicMesh3* Mesh;
+};
+
+
+
+
+
+UCLASS()
+class MESHMODELINGTOOLS_API UScaleKelvinletBrushOpProps : public   UBaseKelvinletBrushOpProps
+{
+	GENERATED_BODY()
+
+public:
+	/** Strength of the Brush */
+	UPROPERTY(EditAnywhere, Category = KelvinScaleBrush, meta = (DisplayName = "Strength", UIMin = "0.0", UIMax = "10.", ClampMin = "0.0", ClampMax = "10.", ModelingQuickEdit, ModelingQuickSettings))
+	float Strength = 0.5;
+
+	/** Amount of falloff to apply */
+	UPROPERTY(EditAnywhere, Category = KelvinScaleBrush, meta = (DisplayName = "Falloff", UIMin = "0.0", UIMax = "1.0", ClampMin = "0.0", ClampMax = "1.0"))
+	float Falloff = 0.5;
+
+
+	virtual float GetStrength() override { return Strength; }
+	virtual float GetFalloff() override { return Falloff; }
+	virtual bool SupportsStrengthPressure() override { return true; }
+};
+
+class FScaleKelvinletBrushOp : public  FBaseKelvinletBrushOp
+{
+
+public:
+	
+	virtual void ApplyStamp(const FDynamicMesh3* SrcMesh, const FSculptBrushStamp& Stamp, const TArray<int32>& Vertices, TArray<FVector3d>& NewPositionsOut) override
+	{
+		SetBaseProperties(SrcMesh, Stamp);
+
+		float Strength = GetPropertySetAs<UScaleKelvinletBrushOpProps>()->GetStrength();
+
+		float Speed = Strength * 0.25 * FMath::Sqrt(Stamp.Radius) * Stamp.Direction;
+		FScaleKelvinlet ScaleKelvinlet(Speed, 0.35 * Size, Mu, Nu);
+
+		ApplyKelvinlet(ScaleKelvinlet, Stamp.LocalFrame, Vertices, NewPositionsOut);
+
+		ApplyFalloff(Stamp, Vertices, NewPositionsOut);
+	}
+
+};
+
+UCLASS()
+class MESHMODELINGTOOLS_API UPullKelvinletBrushOpProps : public UBaseKelvinletBrushOpProps
+{
+	GENERATED_BODY()
+public:
+	
+	/** Amount of falloff to apply */
+	UPROPERTY(EditAnywhere, Category = KelvinGrabBrush, meta = (DisplayName = "Falloff", UIMin = "0.0", UIMax = "1.0", ClampMin = "0.0", ClampMax = "1.0"))
+	float Falloff = 0.1;
+
+	/** Depth of Brush into surface along view ray */
+	UPROPERTY(EditAnywhere, Category = KelvinGrabBrush, meta = (UIMin = "-0.5", UIMax = "0.5", ClampMin = "-1.0", ClampMax = "1.0"))
+	float Depth = 0;
+
+	// these get routed into the Stamp
+	virtual float GetFalloff() override { return Falloff; }
+	virtual float GetDepth() override { return Depth; }
+};
+
+class FPullKelvinletBrushOp : public  FBaseKelvinletBrushOp
+{
+
+public:
+
+	virtual void ApplyStamp(const FDynamicMesh3* SrcMesh, const FSculptBrushStamp& Stamp, const TArray<int32>& Vertices, TArray<FVector3d>& NewPositionsOut) override
+	{
+		SetBaseProperties(SrcMesh, Stamp);
+
+		// compute the displacement vector in the local frame of the stamp
+		FVector3d Force = Stamp.LocalFrame.ToFrameVector(Stamp.LocalFrame.Origin - Stamp.PrevLocalFrame.Origin);
+		Size *= 0.6;
+
+		FLaplacianPullKelvinlet LaplacianPullKelvinlet(Force, Size, Mu, Nu);
+		FBiLaplacianPullKelvinlet BiLaplacianPullKelvinlet(Force, Size, Mu, Nu);
+
+		const double Alpha = Stamp.Falloff;
+		// Lerp between a broad and a narrow kelvinlet based on the fall-off 
+		FBlendPullKelvinlet BlendPullKelvinlet(BiLaplacianPullKelvinlet, LaplacianPullKelvinlet, Alpha);
+		ApplyKelvinlet(BlendPullKelvinlet, Stamp.LocalFrame, Vertices, NewPositionsOut);
+
+		ApplyFalloff(Stamp, Vertices, NewPositionsOut);
+	}
+
+	virtual ESculptBrushOpTargetType GetBrushTargetType() const override
+	{
+		return ESculptBrushOpTargetType::ActivePlane;
+	}
+
+	virtual bool IgnoreZeroMovements() const override
+	{
+		return true;
+	}
+};
+
+UCLASS()
+class MESHMODELINGTOOLS_API USharpPullKelvinletBrushOpProps : public UBaseKelvinletBrushOpProps
+{
+	GENERATED_BODY()
+public:
+
+	/** Amount of falloff to apply */
+	UPROPERTY(EditAnywhere, Category = KelvinSharpGrabBrush, meta = (DisplayName = "Falloff", UIMin = "0.0", UIMax = "1.0", ClampMin = "0.0", ClampMax = "1.0"))
+	float Falloff = 0.5;
+
+	/** Depth of Brush into surface along view ray */
+	UPROPERTY(EditAnywhere, Category = KelvinSharpGrabBrush, meta = (UIMin = "-0.5", UIMax = "0.5", ClampMin = "-1.0", ClampMax = "1.0"))
+	float Depth = 0;
+
+	// these get routed into the Stamp
+	virtual float GetFalloff() override { return Falloff; }
+	virtual float GetDepth() override { return Depth; }
+};
+
+class FSharpPullKelvinletBrushOp : public  FBaseKelvinletBrushOp
+{
+
+public:
+
+	virtual void ApplyStamp(const FDynamicMesh3* SrcMesh, const FSculptBrushStamp& Stamp, const TArray<int32>& Vertices, TArray<FVector3d>& NewPositionsOut) override
+	{
+		SetBaseProperties(SrcMesh, Stamp);
+
+		// compute the displacement vector in the local frame of the stamp
+		FVector3d Force = Stamp.LocalFrame.ToFrameVector(Stamp.LocalFrame.Origin - Stamp.PrevLocalFrame.Origin);
+
+		FSharpLaplacianPullKelvinlet SharpLaplacianPullKelvinlet(Force, Size, Mu, Nu);
+		FSharpBiLaplacianPullKelvinlet SharpBiLaplacianPullKelvinlet(Force, Size, Mu, Nu);
+
+		const double Alpha = Stamp.Falloff;
+		// Lerp between a broad and a narrow Kelvinlet based on the fall-off 
+		FBlendPullSharpKelvinlet BlendPullSharpKelvinlet(SharpBiLaplacianPullKelvinlet, SharpLaplacianPullKelvinlet, Alpha);
+
+		ApplyKelvinlet(BlendPullSharpKelvinlet, Stamp.LocalFrame, Vertices, NewPositionsOut);
+
+		ApplyFalloff(Stamp, Vertices, NewPositionsOut);
+	}
+
+	virtual ESculptBrushOpTargetType GetBrushTargetType() const override
+	{
+		return ESculptBrushOpTargetType::ActivePlane;
+	}
+
+	virtual bool IgnoreZeroMovements() const override
+	{
+		return true;
+	}
+
+};
+
+class FLaplacianPullKelvinletBrushOp : public FBaseKelvinletBrushOp
+{
+public:
+
+	virtual void ApplyStamp(const FDynamicMesh3* SrcMesh, const FSculptBrushStamp& Stamp, const TArray<int32>& Vertices, TArray<FVector3d>& NewPositionsOut) override
+	{
+		SetBaseProperties(SrcMesh, Stamp);
+
+		// compute the displacement vector in the local frame of the stamp
+		FVector3d Force = Stamp.LocalFrame.ToFrameVector(Stamp.LocalFrame.Origin - Stamp.PrevLocalFrame.Origin);
+
+		FLaplacianPullKelvinlet PullKelvinlet(Force, Size, Mu, Nu);
+		ApplyKelvinlet(PullKelvinlet, Stamp.LocalFrame, Vertices, NewPositionsOut);
+
+		ApplyFalloff(Stamp, Vertices, NewPositionsOut);
+	}
+
+	virtual ESculptBrushOpTargetType GetBrushTargetType() const override
+	{
+		return ESculptBrushOpTargetType::ActivePlane;
+	}
+
+	virtual bool IgnoreZeroMovements() const override
+	{
+		return true;
+	}
+};
+
+class FBiLaplacianPullKelvinletBrushOp : public FBaseKelvinletBrushOp
+{
+public:
+
+	virtual void ApplyStamp(const FDynamicMesh3* SrcMesh, const FSculptBrushStamp& Stamp, const TArray<int32>& Vertices, TArray<FVector3d>& NewPositionsOut) override
+	{
+		SetBaseProperties(SrcMesh, Stamp);
+
+		// compute the displacement vector in the local frame of the stamp
+		FVector3d Force = Stamp.LocalFrame.ToFrameVector(Stamp.LocalFrame.Origin - Stamp.PrevLocalFrame.Origin);
+
+		FBiLaplacianPullKelvinlet PullKelvinlet(Force, Size, Mu, Nu);
+		ApplyKelvinlet(PullKelvinlet, Stamp.LocalFrame, Vertices, NewPositionsOut);
+
+		ApplyFalloff(Stamp, Vertices, NewPositionsOut);
+	}
+
+	virtual ESculptBrushOpTargetType GetBrushTargetType() const override
+	{
+		return ESculptBrushOpTargetType::ActivePlane;
+	}
+
+	virtual bool IgnoreZeroMovements() const override
+	{
+		return true;
+	}
+
+};
+
+UCLASS()
+class MESHMODELINGTOOLS_API UTwistKelvinletBrushOpProps : public  UBaseKelvinletBrushOpProps
+{
+	GENERATED_BODY()
+
+public:
+	/** Twisting strength of the Brush */
+	UPROPERTY(EditAnywhere, Category = KelvinTwistBrush, meta = (DisplayName = "Strength", UIMin = "0.0", UIMax = "10.", ClampMin = "0.0", ClampMax = "10.", ModelingQuickSettings))
+	float Strength = 0.5;
+
+	/** Amount of falloff to apply */
+	UPROPERTY(EditAnywhere, Category = KelvinTwistBrush, meta = (DisplayName = "Falloff", UIMin = "0.0", UIMax = "1.0", ClampMin = "0.0", ClampMax = "1.0"))
+	float Falloff = 0.5;
+
+
+	virtual float GetStrength() override { return Strength; }
+	virtual float GetFalloff() override { return Falloff; }
+	virtual bool SupportsStrengthPressure() override { return true; }
+};
+
+class FTwistKelvinletBrushOp : public FBaseKelvinletBrushOp
+{
+public:
+
+
+	virtual void ApplyStamp(const FDynamicMesh3* SrcMesh, const FSculptBrushStamp& Stamp, const TArray<int32>& Vertices, TArray<FVector3d>& NewPositionsOut) override
+	{
+		SetBaseProperties(SrcMesh, Stamp);
+
+		float Strength = GetPropertySetAs<UTwistKelvinletBrushOpProps>()->GetStrength();
+
+		float Speed = Strength * Stamp.Direction * 10;
+		
+		// In the local frame, twist axis is in "z" direction 
+		FVector3d LocalTwistAxis = FVector3d(0., 0., Speed);
+
+		FTwistKelvinlet TwistKelvinlet(LocalTwistAxis, Size, Mu, Nu);
+		ApplyKelvinlet(TwistKelvinlet, Stamp.LocalFrame, Vertices, NewPositionsOut);
+
+		ApplyFalloff(Stamp, Vertices, NewPositionsOut);
+	}
+
+	virtual ESculptBrushOpTargetType GetBrushTargetType() const override
+	{
+		return ESculptBrushOpTargetType::ActivePlane;
+	}
+};
+
+class FPinchKelvinletBrushOp : public FBaseKelvinletBrushOp
+{
+public:
+	virtual void ApplyStamp(const FDynamicMesh3* SrcMesh, const FSculptBrushStamp& Stamp, const TArray<int32>& Vertices, TArray<FVector3d>& NewPositionsOut) override
+	{
+		// is this really the vector we want?
+		// compute the displacement vector in the local frame of the stamp
+		FVector3d Dir = Stamp.LocalFrame.ToFrameVector(Stamp.LocalFrame.Origin - Stamp.PrevLocalFrame.Origin);
+
+		FMatrix3d ForceMatrix = UE::Geometry::CrossProductMatrix(Dir);
+		// make symmetric
+		ForceMatrix.Row0[1] = -ForceMatrix.Row0[1];
+		ForceMatrix.Row0[2] = -ForceMatrix.Row0[2];
+		ForceMatrix.Row1[2] = -ForceMatrix.Row1[2];
+		FPinchKelvinlet PinchKelvinlet(ForceMatrix, Size, Mu, Nu);
+		ApplyKelvinlet(PinchKelvinlet, Stamp.LocalFrame, Vertices, NewPositionsOut);
+
+		ApplyFalloff(Stamp, Vertices, NewPositionsOut);
+	}
+};
+
+
+/**--------------------- Glue Layer Currently Used in the DynamicMeshSculptTool:   --------------------------
+* 
+* We should be able to remove these once the dynamic mesh sculpt tool starts using the new brush infrastructure ( used by the MeshVertexSculptTool)
+*/
+
+enum class EKelvinletBrushMode
+{
+	ScaleKelvinlet,
+	PinchKelvinlet,
+	TwistKelvinlet,
+	PullKelvinlet,
+	LaplacianPullKelvinlet,
+	BiLaplacianPullKelvinlet,
+	BiLaplacianTwistPullKelvinlet,
+	LaplacianTwistPullKelvinlet,
+	SharpPullKelvinlet
+};
+
+class FKelvinletBrushOp
+{
+public:
+
+
+	struct FKelvinletBrushOpProperties
+	{
+		FKelvinletBrushOpProperties(const EKelvinletBrushMode& BrushMode, const UKelvinBrushProperties& Properties, double BrushRadius, double BrushFalloff) 
+			: Mode(BrushMode)
+			, Direction(1.0, 0.0, 0.0)
+		{
+			Speed = 0.;
+			FallOff = BrushFalloff;
+			Mu = FMath::Max(Properties.Stiffness, 0.f);
+			Nu = FMath::Clamp(0.5f * (1.f - 2.f * Properties.Incompressiblity), 0.f, 0.5f);
+			Size = FMath::Max(BrushRadius * Properties.FallOffDistance, 0.f);
+			NumSteps = Properties.BrushSteps;
+		}
+
+
+		EKelvinletBrushMode Mode;
+		FVector Direction;
+
+		double Speed;   // Optionally used
+		double FallOff; // Optionally used
+		double Mu; // Shear Modulus
+		double Nu; // Poisson ratio
+
+		// regularization parameter
+		double Size;
+
+		int NumSteps;
+
+	};
+
+	FKelvinletBrushOp(const FDynamicMesh3& DynamicMesh) :
+		Mesh(&DynamicMesh)
+	{}
+
+	void ExtractTransform(const FMatrix& WorldToBrush)
+	{
+		// Extract the parts of the transform and account for the vector * matrix  format of FMatrix by transposing.
+
+		// Transpose of the 3x3 part
+		WorldToBrushMat.Row0[0] = WorldToBrush.M[0][0];
+		WorldToBrushMat.Row0[1] = WorldToBrush.M[0][1];
+		WorldToBrushMat.Row0[2] = WorldToBrush.M[0][2];
+
+		WorldToBrushMat.Row1[0] = WorldToBrush.M[1][0];
+		WorldToBrushMat.Row1[1] = WorldToBrush.M[1][1];
+		WorldToBrushMat.Row1[2] = WorldToBrush.M[1][2];
+
+		WorldToBrushMat.Row2[0] = WorldToBrush.M[2][0];
+		WorldToBrushMat.Row2[1] = WorldToBrush.M[2][1];
+		WorldToBrushMat.Row2[2] = WorldToBrush.M[2][2];
+
+		Translation[0] = WorldToBrush.M[3][0];
+		Translation[1] = WorldToBrush.M[3][1];
+		Translation[2] = WorldToBrush.M[3][2];
+
+		// The matrix should be unitary (det +/- 1) but want this to work with 
+		// more general input if needed so just make sure we can invert the matrix
+		check(FMath::Abs(WorldToBrush.Determinant()) > 1.e-4);
+
+		BrushToWorldMat = WorldToBrushMat.Inverse();
+
+	};
+
+	~FKelvinletBrushOp() {}
+
+public:
+	const FDynamicMesh3* Mesh;
+
+
+	double TimeStep = 1.0;
+	double NumSteps = 0.0;
+
+	// To be applied as WorlToBrushMat * v + Trans
+	// Note: could use FTransform3d in TransformTypes.h
+	FMatrix3d WorldToBrushMat;
+	FMatrix3d BrushToWorldMat;
+	FVector3d Translation;
+
+public:
+
+
+
+	void ApplyBrush(const FKelvinletBrushOpProperties& Properties, const FMatrix& WorldToBrush, const TArray<int>& VertRIO, TArray<FVector3d>& ROIPositionBuffer)
+	{
+
+		ExtractTransform(WorldToBrush);
+
+		NumSteps = Properties.NumSteps;
+
+		switch (Properties.Mode)
+		{
+
+		case EKelvinletBrushMode::ScaleKelvinlet:
+		{
+			double Scale = Properties.Direction.X;
+			FScaleKelvinlet ScaleKelvinlet(Scale, Properties.Size, Properties.Mu, Properties.Nu);
+			ApplyKelvinlet(ScaleKelvinlet, VertRIO, ROIPositionBuffer, TimeStep, NumSteps);
+			break;
+		}
+		case EKelvinletBrushMode::PullKelvinlet:
+		{
+			FVector Dir = Properties.Direction;
+			FVector3d Force(Dir.X, Dir.Y, Dir.Z);
+
+			FLaplacianPullKelvinlet LaplacianPullKelvinlet(Force, Properties.Size, Properties.Mu, Properties.Nu);
+			FBiLaplacianPullKelvinlet BiLaplacianPullKelvinlet(Force, Properties.Size, Properties.Mu, Properties.Nu);
+
+			const double Alpha = Properties.FallOff;
+			// Lerp between a broad and a narrow kelvinlet based on the fall-off 
+			FBlendPullKelvinlet BlendPullKelvinlet(BiLaplacianPullKelvinlet, LaplacianPullKelvinlet, Alpha);
+			ApplyKelvinlet(BlendPullKelvinlet, VertRIO, ROIPositionBuffer, TimeStep, NumSteps);
+			break;
+		}
+		case EKelvinletBrushMode::SharpPullKelvinlet:
+		{
+			FVector Dir = Properties.Direction;
+			FVector3d Force(Dir.X, Dir.Y, Dir.Z);
+
+			FSharpLaplacianPullKelvinlet SharpLaplacianPullKelvinlet(Force, Properties.Size, Properties.Mu, Properties.Nu);
+			FSharpBiLaplacianPullKelvinlet SharpBiLaplacianPullKelvinlet(Force, Properties.Size, Properties.Mu, Properties.Nu);
+
+			const double Alpha = Properties.FallOff;
+			// Lerp between a broad and a narrow Kelvinlet based on the fall-off 
+			FBlendPullSharpKelvinlet BlendPullSharpKelvinlet(SharpBiLaplacianPullKelvinlet, SharpLaplacianPullKelvinlet, Alpha);
+			ApplyKelvinlet(BlendPullSharpKelvinlet, VertRIO, ROIPositionBuffer, TimeStep, NumSteps);
+			break;
+		}
+		case EKelvinletBrushMode::LaplacianPullKelvinlet:
+		{
+			FVector Dir = Properties.Direction;
+			FVector3d Force(Dir.X, Dir.Y, Dir.Z);
+			FLaplacianPullKelvinlet PullKelvinlet(Force, Properties.Size, Properties.Mu, Properties.Nu);
+			ApplyKelvinlet(PullKelvinlet, VertRIO, ROIPositionBuffer, TimeStep, NumSteps);
+			break;
+		}
+		case EKelvinletBrushMode::BiLaplacianPullKelvinlet:
+		{
+			FVector Dir = Properties.Direction;
+			FVector3d Force(Dir.X, Dir.Y, Dir.Z);
+			FBiLaplacianPullKelvinlet PullKelvinlet(Force, Properties.Size, Properties.Mu, Properties.Nu);
+			ApplyKelvinlet(PullKelvinlet, VertRIO, ROIPositionBuffer, TimeStep, NumSteps);
+			break;
+		}
+		case EKelvinletBrushMode::TwistKelvinlet:
+		{
+			FVector TwistAxis = Properties.Direction;
+			FTwistKelvinlet TwistKelvinlet((FVector3d)TwistAxis, Properties.Size, Properties.Mu, Properties.Nu);
+			ApplyKelvinlet(TwistKelvinlet, VertRIO, ROIPositionBuffer, TimeStep, NumSteps);
+			break;
+		}
+		case EKelvinletBrushMode::PinchKelvinlet:
+		{
+			FVector Dir = Properties.Direction;
+
+			FMatrix3d ForceMatrix = UE::Geometry::CrossProductMatrix(FVector3d(Dir.X, Dir.Y, Dir.Z));
+			// make symmetric
+			ForceMatrix.Row0[1] = -ForceMatrix.Row0[1];
+			ForceMatrix.Row0[2] = -ForceMatrix.Row0[2];
+			ForceMatrix.Row1[2] = -ForceMatrix.Row1[2];
+			FPinchKelvinlet PinchKelvinlet(ForceMatrix, Properties.Size, Properties.Mu, Properties.Nu);
+			ApplyKelvinlet(PinchKelvinlet, VertRIO, ROIPositionBuffer, TimeStep, NumSteps);
+			break;
+		}
+		case EKelvinletBrushMode::LaplacianTwistPullKelvinlet:
+		{
+			FVector3d TwistAxis = (FVector3d)Properties.Direction;
+			UE::Geometry::Normalize(TwistAxis);
+			TwistAxis *= Properties.Speed;
+			FTwistKelvinlet TwistKelvinlet(TwistAxis, Properties.Size, Properties.Mu, Properties.Nu);
+
+			FVector Dir = Properties.Direction;
+			FVector3d Force(Dir.X, Dir.Y, Dir.Z);
+			FLaplacianPullKelvinlet PullKelvinlet(Force, Properties.Size, Properties.Mu, Properties.Nu);
+
+			FLaplacianTwistPullKelvinlet TwistPullKelvinlet(TwistKelvinlet, PullKelvinlet, 0.5);
+			ApplyKelvinlet(TwistPullKelvinlet, VertRIO, ROIPositionBuffer, TimeStep, NumSteps);
+			break;
+
+		}
+		default:
+			check(0);
+		}
+	}
+
+
+	// NB: this just moves the verts, but doesn't update the normal.  The kelvinlets will have to be extended if we want
+	// to do the Jacobian Transpose operation on the normals - but for now, we should just rebuild the normals after the brush
+	template <typename KelvinletType>
+	void DisplaceKelvinlet(const KelvinletType& Kelvinlet, const TArray<int>& VertexROI, TArray<FVector3d>& ROIPositionBuffer)
+	{
+		int NumV = VertexROI.Num();
+		ROIPositionBuffer.SetNum(NumV, EAllowShrinking::No);
+
+		const bool bForceSingleThread = false;
+
+		ParallelFor(NumV, [&Kelvinlet, &VertexROI, &ROIPositionBuffer, this](int k)
+
+		{
+			int VertIdx = VertexROI[k];
+			FVector3d Pos = Mesh->GetVertex(VertIdx);
+			Pos = XForm(Pos);
+
+			Pos = Kelvinlet.Evaluate(Pos) + Pos;
+
+			// Update the position in the ROI Array
+			ROIPositionBuffer[k] = InvXForm(Pos);
+		}
+		, bForceSingleThread);
+
+	}
+
+	template <typename KelvinletType>
+	void IntegrateKelvinlet(const KelvinletType& Kelvinlet, const TArray<int>& VertexROI, TArray<FVector3d>& ROIPositionBuffer, const double Dt, const int Steps)
+	{
+		int NumV = VertexROI.Num();
+		ROIPositionBuffer.SetNum(NumV, EAllowShrinking::No);
+
+		const bool bForceSingleThread = false;
+
+		ParallelFor(NumV, [&Kelvinlet, &VertexROI, &ROIPositionBuffer, Dt, Steps, this](int k)
+		{
+			int VertIdx = VertexROI[k];
+			FVector3d Pos = XForm(Mesh->GetVertex(VertIdx));
+
+			double TimeScale = 1. / (Steps);
+			// move with several time steps
+			for (int i = 0; i < Steps; ++i)
+			{
+				// the position after deformation
+				Pos = Kelvinlet.IntegrateRK3(Pos, Dt * TimeScale);
+			}
+			// Update the position in the ROI Array
+			ROIPositionBuffer[k] = InvXForm(Pos);
+		}, bForceSingleThread);
+
+
+	}
+
+	template <typename KelvinletType>
+	void ApplyKelvinlet(const KelvinletType& Kelvinlet, const TArray<int>& VertexROI, TArray<FVector3d>& ROIPositionBuffer, const double Dt, const int NumIntegrationSteps)
+	{
+
+		if (NumIntegrationSteps == 0)
+		{
+			DisplaceKelvinlet(Kelvinlet, VertexROI, ROIPositionBuffer);
+		}
+		else
+		{
+			IntegrateKelvinlet(Kelvinlet, VertexROI, ROIPositionBuffer, Dt, NumIntegrationSteps);
+		}
+	}
+
+private:
+
+	// apply the transform.
+	FVector3d XForm(const FVector3d& Pos) const
+	{
+		return WorldToBrushMat * Pos + Translation;
+	}
+
+	// apply the inverse transform.
+	FVector3d InvXForm(const FVector3d& Pos) const
+	{
+		return BrushToWorldMat * (Pos - Translation);
+	}
+
+};
\ No newline at end of file
diff --git a/Engine/Plugins/Runtime/MeshModelingToolset/Source/MeshModelingTools/Public/Sculpting/MeshBrushOpBase.h b/Engine/Plugins/Runtime/MeshModelingToolset/Source/MeshModelingTools/Public/Sculpting/MeshBrushOpBase.h
--- a/Engine/Plugins/Runtime/MeshModelingToolset/Source/MeshModelingTools/Public/Sculpting/MeshBrushOpBase.h
+++ b/Engine/Plugins/Runtime/MeshModelingToolset/Source/MeshModelingTools/Public/Sculpting/MeshBrushOpBase.h
@@ -93,6 +93,16 @@ public:
 	// support for this is optional, used by UI level to edit brush props via hotkeys/etc
 	virtual void SetStrength(float NewStrength) { }
 	virtual void SetFalloff(float NewFalloff) { }
+
+	// is Pressure Sensitivity for Brush Strength currently toggled on?
+	virtual bool GetStrengthPressureEnabled() { return  SupportsStrengthPressure() && bIsStrengthPressureEnabled; }
+
+	// is Pressure Sensitivity for Brush Strength supported?
+	virtual bool SupportsStrengthPressure() { return false; }
+
+	/** Should pressure affect brush strength? */
+	UPROPERTY(EditAnywhere, Category = BrushOp)
+	bool bIsStrengthPressureEnabled = true;
 };
 
 
diff --git a/Engine/Plugins/Runtime/MeshModelingToolset/Source/MeshModelingTools/Public/Sculpting/MeshInflateBrushOps.h b/Engine/Plugins/Runtime/MeshModelingToolset/Source/MeshModelingTools/Public/Sculpting/MeshInflateBrushOps.h
--- a/Engine/Plugins/Runtime/MeshModelingToolset/Source/MeshModelingTools/Public/Sculpting/MeshInflateBrushOps.h
+++ b/Engine/Plugins/Runtime/MeshModelingToolset/Source/MeshModelingTools/Public/Sculpting/MeshInflateBrushOps.h
@@ -26,6 +26,7 @@ public:
 	virtual float GetStrength() override { return Strength; }
 	virtual void SetStrength(float NewStrength) override { Strength = FMathf::Clamp(NewStrength, 0.0f, 1.0f); }
 	virtual float GetFalloff() override { return Falloff; }
+	virtual bool SupportsStrengthPressure() override { return true; }
 };
 
 class FInflateBrushOp : public FMeshSculptBrushOp
diff --git a/Engine/Plugins/Runtime/MeshModelingToolset/Source/MeshModelingTools/Public/Sculpting/MeshMoveBrushOps.h b/Engine/Plugins/Runtime/MeshModelingToolset/Source/MeshModelingTools/Public/Sculpting/MeshMoveBrushOps.h
--- a/Engine/Plugins/Runtime/MeshModelingToolset/Source/MeshModelingTools/Public/Sculpting/MeshMoveBrushOps.h
+++ b/Engine/Plugins/Runtime/MeshModelingToolset/Source/MeshModelingTools/Public/Sculpting/MeshMoveBrushOps.h
@@ -33,6 +33,7 @@ public:
 	virtual void SetStrength(float NewStrength) override { Strength = FMathf::Clamp(NewStrength, 0.0f, 1.0f); }
 	virtual float GetFalloff() override { return Falloff; }
 	virtual float GetDepth() override { return Depth; }
+	virtual bool SupportsStrengthPressure() override { return true; }
 
 	/** Axis filters restrict mesh movement to World X/Y/Z axes */
 	UPROPERTY(EditAnywhere, Category = MoveBrush)
diff --git a/Engine/Plugins/Runtime/MeshModelingToolset/Source/MeshModelingTools/Public/Sculpting/MeshPinchBrushOps.h b/Engine/Plugins/Runtime/MeshModelingToolset/Source/MeshModelingTools/Public/Sculpting/MeshPinchBrushOps.h
--- a/Engine/Plugins/Runtime/MeshModelingToolset/Source/MeshModelingTools/Public/Sculpting/MeshPinchBrushOps.h
+++ b/Engine/Plugins/Runtime/MeshModelingToolset/Source/MeshModelingTools/Public/Sculpting/MeshPinchBrushOps.h
@@ -34,6 +34,7 @@ public:
 	virtual void SetStrength(float NewStrength) override { Strength = FMathf::Clamp(NewStrength, 0.0f, 1.0f); }
 	virtual float GetFalloff() override { return Falloff; }
 	virtual float GetDepth() override { return Depth; }
+	virtual bool SupportsStrengthPressure() override { return true; }
 };
 
 
diff --git a/Engine/Plugins/Runtime/MeshModelingToolset/Source/MeshModelingTools/Public/Sculpting/MeshPlaneBrushOps.h b/Engine/Plugins/Runtime/MeshModelingToolset/Source/MeshModelingTools/Public/Sculpting/MeshPlaneBrushOps.h
--- a/Engine/Plugins/Runtime/MeshModelingToolset/Source/MeshModelingTools/Public/Sculpting/MeshPlaneBrushOps.h
+++ b/Engine/Plugins/Runtime/MeshModelingToolset/Source/MeshModelingTools/Public/Sculpting/MeshPlaneBrushOps.h
@@ -17,6 +17,7 @@ class MESHMODELINGTOOLS_API UBasePlaneBrushOpProps : public UMeshSculptBrushOpPr
 	GENERATED_BODY()
 public:
 	virtual EPlaneBrushSideMode GetWhichSide() { return EPlaneBrushSideMode::BothSides; }
+	virtual bool SupportsStrengthPressure() override { return true; }
 };
 
 UCLASS()
diff --git a/Engine/Plugins/Runtime/MeshModelingToolset/Source/MeshModelingTools/Public/Sculpting/MeshSculptBrushOps.h b/Engine/Plugins/Runtime/MeshModelingToolset/Source/MeshModelingTools/Public/Sculpting/MeshSculptBrushOps.h
--- a/Engine/Plugins/Runtime/MeshModelingToolset/Source/MeshModelingTools/Public/Sculpting/MeshSculptBrushOps.h
+++ b/Engine/Plugins/Runtime/MeshModelingToolset/Source/MeshModelingTools/Public/Sculpting/MeshSculptBrushOps.h
@@ -26,6 +26,7 @@ public:
 	virtual float GetStrength() override { return Strength; }
 	virtual void SetStrength(float NewStrength) override { Strength = FMathf::Clamp(NewStrength, 0.0f, 1.0f); }
 	virtual float GetFalloff() override { return Falloff; }
+	virtual bool SupportsStrengthPressure() override { return true; }
 };
 
 class FSurfaceSculptBrushOp : public FMeshSculptBrushOp
@@ -107,6 +108,7 @@ public:
 	virtual float GetStrength() override { return Strength; }
 	virtual void SetStrength(float NewStrength) override { Strength = FMathf::Clamp(NewStrength, 0.0f, 1.0f); }
 	virtual float GetFalloff() override { return Falloff; }
+	virtual bool SupportsStrengthPressure() override { return true; }
 };
 
 
@@ -209,6 +211,7 @@ public:
 	virtual float GetStrength() override { return Strength; }
 	virtual void SetStrength(float NewStrength) override { Strength = FMathf::Clamp(NewStrength, 0.0f, 1.0f); }
 	virtual float GetFalloff() override { return Falloff; }
+	virtual bool SupportsStrengthPressure() override { return true; }
 };
 
 
diff --git a/Engine/Plugins/Runtime/MeshModelingToolset/Source/MeshModelingTools/Public/Sculpting/MeshSculptToolBase.h b/Engine/Plugins/Runtime/MeshModelingToolset/Source/MeshModelingTools/Public/Sculpting/MeshSculptToolBase.h
--- a/Engine/Plugins/Runtime/MeshModelingToolset/Source/MeshModelingTools/Public/Sculpting/MeshSculptToolBase.h
+++ b/Engine/Plugins/Runtime/MeshModelingToolset/Source/MeshModelingTools/Public/Sculpting/MeshSculptToolBase.h
@@ -294,6 +294,9 @@ protected:
 	/** Returns true if BrushSize Pressure Sensitivity is both supported by the tool and currently enabled. */
 	bool GetBrushSizePressureSensitivityEnabled() const;
 
+	/** Returns true if Pressure Sensitivity for Brush Strength is both supported by the tool and currently enabled. */
+	bool GetBrushStrengthPressureEnabled();
+
 	/**
 	 * Subclass must implement this and return relevant rendering component
 	 */
@@ -407,6 +410,8 @@ protected:
 	double ActivePressure = 1.0;
 	virtual double GetActivePressure() const { return ActivePressure; }
 
+	// returns brush strength, factoring in pressure sensitivity if applicable
+	virtual double GetActiveBrushStrength();
 	virtual double GetCurrentBrushStrength();
 	virtual double GetCurrentBrushDepth();
 
diff --git a/Engine/Plugins/Runtime/MeshModelingToolset/Source/MeshModelingTools/Public/Sculpting/MeshSmoothingBrushOps.h b/Engine/Plugins/Runtime/MeshModelingToolset/Source/MeshModelingTools/Public/Sculpting/MeshSmoothingBrushOps.h
--- a/Engine/Plugins/Runtime/MeshModelingToolset/Source/MeshModelingTools/Public/Sculpting/MeshSmoothingBrushOps.h
+++ b/Engine/Plugins/Runtime/MeshModelingToolset/Source/MeshModelingTools/Public/Sculpting/MeshSmoothingBrushOps.h
@@ -16,6 +16,7 @@ class MESHMODELINGTOOLS_API UBaseSmoothBrushOpProps : public UMeshSculptBrushOpP
 	GENERATED_BODY()
 public:
 	virtual bool GetPreserveUVFlow() { return false; }
+	virtual bool SupportsStrengthPressure() override { return true; }
 };
 
 
@@ -206,6 +207,7 @@ public:
 	virtual void SetStrength(float NewStrength) override { Strength = FMathf::Clamp(NewStrength, 0.0f, 1.0f); }
 	virtual float GetFalloff() override { return Falloff; }
 	virtual float GetDepth() override { return Depth; }
+	virtual bool SupportsStrengthPressure() override { return true; }
 };
 
 
@@ -277,6 +279,7 @@ public:
 	virtual float GetStrength() override { return Strength; }
 	virtual void SetStrength(float NewStrength) override { Strength = FMathf::Clamp(NewStrength, 0.0f, 1.0f); }
 	virtual float GetFalloff() override { return Falloff; }
+	virtual bool SupportsStrengthPressure() override { return true; }
 };
 
 
