ModelingTools: Fixed PaintVertexColors applying double sRGB transformation on shutdown.

diff --git a/Engine/Plugins/Runtime/MeshModelingToolset/Source/MeshModelingTools/Private/MeshVertexPaintTool.cpp b/Engine/Plugins/Runtime/MeshModelingToolset/Source/MeshModelingTools/Private/MeshVertexPaintTool.cpp
--- a/Engine/Plugins/Runtime/MeshModelingToolset/Source/MeshModelingTools/Private/MeshVertexPaintTool.cpp
+++ b/Engine/Plugins/Runtime/MeshModelingToolset/Source/MeshModelingTools/Private/MeshVertexPaintTool.cpp
@@ -197,16 +197,6 @@ void UMeshVertexPaintTool::Setup()
 	StrokeInitialColorBuffer.SetNum(ActiveColorOverlay->MaxElementID());
 	StrokeAccumColorBuffer.SetNum(ActiveColorOverlay->MaxElementID());
 
-	// UMeshSculptToolBase::InitializeSculptMeshComponent() should provide some way for this level to control how 
-	// colors are interpreted. Actually calls GetDynamicMeshCopy() which may or may not convert to SRGB...
-	// 
-	// code above auto-converted to SRGB so we will convert back to linear now...
-	for (int32 ElementID : ActiveColorOverlay->ElementIndicesItr())
-	{
-		FVector4f Color = ActiveColorOverlay->GetElement(ElementID);
-		LinearColors::SRGBToLinear(Color);
-		ActiveColorOverlay->SetElement(ElementID, Color);
-	}
 	DynamicMeshComponent->ColorSpaceMode = EDynamicMeshVertexColorTransformMode::LinearToSRGB;
 
 	//Mesh->EnableVertexColors(FVector3f::One());
@@ -458,16 +448,6 @@ void UMeshVertexPaintTool::Shutdown(EToolShutdownType ShutdownType)
 			FSplitAttributeWelder::WeldSplitColors(vid, *ColorOverlay, FMathd::ZeroTolerance);
 		}
 		Mesh->CompactInPlace();		// to clean up colors we just welded
-
-		// TODO can override UMeshSculptToolBase::CommitResult() that ::Shutdown calls to avoid this conversion...
-
-		// shutdown call below wants colors to be in SRGB so it can convert back to linear...
-		for (int32 ElementID : ActiveColorOverlay->ElementIndicesItr())
-		{
-			FVector4f Color = ColorOverlay->GetElement(ElementID);
-			LinearColors::LinearToSRGB(Color);
-			ColorOverlay->SetElement(ElementID, Color);
-		}
 	}
 
 	UMeshSculptToolBase::Shutdown(ShutdownType);
@@ -1838,13 +1818,13 @@ static void GetMeshDescriptionLODWithCopiedVertexColors(
 
 	FDynamicMesh3 LODMesh;
 	FMeshDescriptionToDynamicMesh Converter1;
-	Converter1.bTransformVertexColorsLinearToSRGB = false;
+	Converter1.bTransformVertexColorsLinearToSRGB = true;
 	Converter1.Convert(&UpdatedMeshDescriptionOut, LODMesh);
 
 	UELocal::CopyVertexColorsToLOD(SourceMesh, ColorOverlay, SourceMeshSpatial, LODMesh, DefaultColor);
 
 	FDynamicMeshToMeshDescription Converter2;
-	Converter2.ConversionOptions.bTransformVtxColorsSRGBToLinear = false;
+	Converter2.ConversionOptions.bTransformVtxColorsSRGBToLinear = true;
 	Converter2.UpdateVertexColors(&LODMesh, UpdatedMeshDescriptionOut);
 }
 
