ModelingTools: Fixed PaintVertexColors displaying the wrong values in Unlit Vertex Color mode.
Change default material mode to Unlit Vertex Color mode.

diff --git a/Engine/Plugins/Runtime/MeshModelingToolset/Source/MeshModelingTools/Private/MeshVertexPaintTool.cpp b/Engine/Plugins/Runtime/MeshModelingToolset/Source/MeshModelingTools/Private/MeshVertexPaintTool.cpp
--- a/Engine/Plugins/Runtime/MeshModelingToolset/Source/MeshModelingTools/Private/MeshVertexPaintTool.cpp
+++ b/Engine/Plugins/Runtime/MeshModelingToolset/Source/MeshModelingTools/Private/MeshVertexPaintTool.cpp
@@ -197,8 +197,7 @@ void UMeshVertexPaintTool::Setup()
 	StrokeInitialColorBuffer.SetNum(ActiveColorOverlay->MaxElementID());
 	StrokeAccumColorBuffer.SetNum(ActiveColorOverlay->MaxElementID());
 
-	DynamicMeshComponent->ColorSpaceMode = EDynamicMeshVertexColorTransformMode::LinearToSRGB;
-
+	
 	//Mesh->EnableVertexColors(FVector3f::One());
 	FAxisAlignedBox3d Bounds = Mesh->GetBounds(true);
 
@@ -1482,15 +1481,16 @@ void UMeshVertexPaintTool::UpdateVertexPaintMaterialMode()
 	{
 		UpdateMaterialMode(EMeshEditingMaterialModes::ExistingMaterial);
 	}
-	else
+	else if (FilterProperties->MaterialMode == EMeshVertexPaintMaterialMode::LitVertexColor)
 	{
 		// TODO: add two-sided vertex color support when it becomes available
 		UpdateMaterialMode(EMeshEditingMaterialModes::VertexColor);
-		if (FilterProperties->MaterialMode == EMeshVertexPaintMaterialMode::UnlitVertexColor
-			&& GEngine && GEngine->VertexColorMaterial)
-		{
-			GetSculptMeshComponent()->SetOverrideRenderMaterial(GEngine->VertexColorMaterial);
-		}
+	}
+	else if (FilterProperties->MaterialMode == EMeshVertexPaintMaterialMode::UnlitVertexColor
+		&& GEngine && GEngine->VertexColorViewModeMaterial_ColorOnly)
+	{
+		UpdateMaterialMode(EMeshEditingMaterialModes::ExistingMaterial);
+		GetSculptMeshComponent()->SetOverrideRenderMaterial(GEngine->VertexColorViewModeMaterial_ColorOnly);
 	}
 
 	OnChannelFilterModified();		// channel visibility filtering may be disabled for existing material
diff --git a/Engine/Plugins/Runtime/MeshModelingToolset/Source/MeshModelingTools/Public/MeshVertexPaintTool.h b/Engine/Plugins/Runtime/MeshModelingToolset/Source/MeshModelingTools/Public/MeshVertexPaintTool.h
--- a/Engine/Plugins/Runtime/MeshModelingToolset/Source/MeshModelingTools/Public/MeshVertexPaintTool.h
+++ b/Engine/Plugins/Runtime/MeshModelingToolset/Source/MeshModelingTools/Public/MeshVertexPaintTool.h
@@ -258,7 +258,7 @@ public:
 
 	/** Specify which Materials should be used to render the Mesh */
 	UPROPERTY(EditAnywhere, Category = Visualization)
-	EMeshVertexPaintMaterialMode MaterialMode = EMeshVertexPaintMaterialMode::LitVertexColor;
+	EMeshVertexPaintMaterialMode MaterialMode = EMeshVertexPaintMaterialMode::UnlitVertexColor;
 
 	/** Display the Color under the cursor */
 	UPROPERTY(EditAnywhere, Category = Visualization, meta=(DisplayName="Show Color Under Cursor"))
