ModelingTools: Disable MeshElementSelection volume locking by default to make volume editing more accessible.

diff --git a/Engine/Plugins/Runtime/MeshModelingToolset/Source/ModelingComponentsEditorOnly/Private/Selection/VolumeSelector.cpp b/Engine/Plugins/Runtime/MeshModelingToolset/Source/ModelingComponentsEditorOnly/Private/Selection/VolumeSelector.cpp
--- a/Engine/Plugins/Runtime/MeshModelingToolset/Source/ModelingComponentsEditorOnly/Private/Selection/VolumeSelector.cpp
+++ b/Engine/Plugins/Runtime/MeshModelingToolset/Source/ModelingComponentsEditorOnly/Private/Selection/VolumeSelector.cpp
@@ -26,7 +26,7 @@ using namespace UE::Geometry;
 
 static TAutoConsoleVariable<bool> CVarEnableModelingSelectionVolumeLocking(
 	TEXT("modeling.Selection.EnableVolumeLocking"),
-	true,
+	false,
 	TEXT("Control whether Selection Locking is enabled by default for Volumes"));
 
 
