ModelingTools: Fixed mesh element selection not updating correctly for volumes after using a selection tool.

diff --git a/Engine/Plugins/Runtime/MeshModelingToolset/Source/ModelingComponentsEditorOnly/Private/ConversionUtils/DynamicMeshToVolume.cpp b/Engine/Plugins/Runtime/MeshModelingToolset/Source/ModelingComponentsEditorOnly/Private/ConversionUtils/DynamicMeshToVolume.cpp
--- a/Engine/Plugins/Runtime/MeshModelingToolset/Source/ModelingComponentsEditorOnly/Private/ConversionUtils/DynamicMeshToVolume.cpp
+++ b/Engine/Plugins/Runtime/MeshModelingToolset/Source/ModelingComponentsEditorOnly/Private/ConversionUtils/DynamicMeshToVolume.cpp
@@ -845,7 +845,7 @@ void DynamicMeshToVolume(const FDynamicMesh3&, TArray<FDynamicMeshFace>& Faces,
 	FBSPOps::bspUnlinkPolys(Model);
 
 	// requires editor
-	FBSPOps::csgPrepMovingBrush(TargetVolume);
+	TargetVolume->GetOnVolumeShapeChangedDelegate().Broadcast(*TargetVolume);
 
 	// do we need to do this?
 	[[maybe_unused]] bool MarkedDirty = TargetVolume->MarkPackageDirty();
diff --git a/Engine/Plugins/Runtime/MeshModelingToolset/Source/ModelingComponentsEditorOnly/Private/Selection/VolumeSelector.cpp b/Engine/Plugins/Runtime/MeshModelingToolset/Source/ModelingComponentsEditorOnly/Private/Selection/VolumeSelector.cpp
--- a/Engine/Plugins/Runtime/MeshModelingToolset/Source/ModelingComponentsEditorOnly/Private/Selection/VolumeSelector.cpp
+++ b/Engine/Plugins/Runtime/MeshModelingToolset/Source/ModelingComponentsEditorOnly/Private/Selection/VolumeSelector.cpp
@@ -119,6 +119,8 @@ bool FVolumeSelector::Initialize(
 	FBaseDynamicMeshSelector::Initialize(SourceGeometryIdentifierIn, LocalTargetMesh.Get(), 
 		[this]() { return IsValid(this->BrushComponent) ? (UE::Geometry::FTransformSRT3d)this->BrushComponent->GetComponentTransform() : FTransformSRT3d::Identity(); });
 
+	ParentVolumeChangedHandle = ParentVolume->GetOnVolumeShapeChangedDelegate().AddLambda([this](AVolume&){ UpdateDynamicMeshFromVolume(); });
+
 	// AVolume and BrushComponent seem to have no way of signalling when they are modified,
 	// so on undo and redo, we have no idea that the mesh has changed. The brute-force way
 	// to work around this is simply to rebuild the mesh on every undo/redo event. 
@@ -134,6 +136,12 @@ bool FVolumeSelector::Initialize(
 
 void FVolumeSelector::Shutdown()
 {
+	if (ParentVolume)
+	{
+		ParentVolume->GetOnVolumeShapeChangedDelegate().Remove(ParentVolumeChangedHandle);
+		ParentVolumeChangedHandle.Reset();
+	}
+	
 	LocalTargetMesh.Reset();
 	FBaseDynamicMeshSelector::Shutdown();
 
diff --git a/Engine/Plugins/Runtime/MeshModelingToolset/Source/ModelingComponentsEditorOnly/Public/Selection/VolumeSelector.h b/Engine/Plugins/Runtime/MeshModelingToolset/Source/ModelingComponentsEditorOnly/Public/Selection/VolumeSelector.h
--- a/Engine/Plugins/Runtime/MeshModelingToolset/Source/ModelingComponentsEditorOnly/Public/Selection/VolumeSelector.h
+++ b/Engine/Plugins/Runtime/MeshModelingToolset/Source/ModelingComponentsEditorOnly/Public/Selection/VolumeSelector.h
@@ -65,6 +65,9 @@ protected:
 
 	TSharedPtr<FBasicDynamicMeshSelectionTransformer> ActiveTransformer;
 	void CommitMeshTransform();
+
+private:
+	FDelegateHandle ParentVolumeChangedHandle;
 };
 
 
