ModelingTools: Fixed crash due to StaticMesh->IsCompiling checks returning false for IsValid on StaticMesh tool targets.

The StaticMesh->IsCompiling checks were recently introduced to prevent blocking the game thread on StaticMeshBuilds. This was occurring because both the IsValid and CanBuildTarget interfaces were querying the source models of the StaticMesh which in turn would block on any in progress SM compiles. The unfortunate side effect introduced here is that any BaseCreateFromSelectedTool would write to an existing StaticMesh target that would cause a compilation. Since IsValid is pervasive and now returning false if the SM was compiling, simple tool target queries such as GetOwnerActor would fail and return nullptr causing later crashes.

This change removes HasNonGeneratedLOD from IsValid checks and instead enforces a valid generated LOD on GetValidEditingLOD (falling back on LOD0 for invalid requested EditingLODs). CanBuildTarget continues the explicit HasNonGeneratedLOD to ensure that the tool can still only start on a valid editing LOD.

diff --git a/Engine/Plugins/Runtime/MeshModelingToolset/Source/ModelingComponentsEditorOnly/Private/ToolTargets/StaticMeshComponentToolTarget.cpp b/Engine/Plugins/Runtime/MeshModelingToolset/Source/ModelingComponentsEditorOnly/Private/ToolTargets/StaticMeshComponentToolTarget.cpp
--- a/Engine/Plugins/Runtime/MeshModelingToolset/Source/ModelingComponentsEditorOnly/Private/ToolTargets/StaticMeshComponentToolTarget.cpp
+++ b/Engine/Plugins/Runtime/MeshModelingToolset/Source/ModelingComponentsEditorOnly/Private/ToolTargets/StaticMeshComponentToolTarget.cpp
@@ -46,8 +46,7 @@ bool UStaticMeshComponentReadOnlyToolTarget::IsValid() const
 		return false;
 	}
 	UStaticMesh* StaticMesh = StaticMeshComponent->GetStaticMesh();
-	
-	return UStaticMeshReadOnlyToolTarget::HasNonGeneratedLOD(StaticMesh, EditingLOD);
+	return ::IsValid(StaticMesh) && !StaticMesh->IsUnreachable() && StaticMesh->IsValidLowLevel();
 }
 
 
diff --git a/Engine/Plugins/Runtime/MeshModelingToolset/Source/ModelingComponentsEditorOnly/Private/ToolTargets/StaticMeshToolTarget.cpp b/Engine/Plugins/Runtime/MeshModelingToolset/Source/ModelingComponentsEditorOnly/Private/ToolTargets/StaticMeshToolTarget.cpp
--- a/Engine/Plugins/Runtime/MeshModelingToolset/Source/ModelingComponentsEditorOnly/Private/ToolTargets/StaticMeshToolTarget.cpp
+++ b/Engine/Plugins/Runtime/MeshModelingToolset/Source/ModelingComponentsEditorOnly/Private/ToolTargets/StaticMeshToolTarget.cpp
@@ -67,8 +67,19 @@ EMeshLODIdentifier UStaticMeshReadOnlyToolTarget::GetValidEditingLOD(const UStat
 			int32 MaxExistingLOD = StaticMeshIn->GetNumSourceModels() - 1;
 			if ((int32)ValidEditingLOD > MaxExistingLOD)
 			{
-				DisplayCriticalWarningMessage(FString::Printf(TEXT("LOD%d Requested but not available - Falling Back to LOD%d"), (int32)ValidEditingLOD, MaxExistingLOD));
-				ValidEditingLOD = (EMeshLODIdentifier)MaxExistingLOD;
+				DisplayCriticalWarningMessage(FString::Printf(TEXT("LOD%d Requested but not available - Falling Back to LOD%d"), (int32)ValidEditingLOD, 0));
+				ValidEditingLOD = EMeshLODIdentifier::LOD0;
+			}
+		}
+
+		// Ensure the selected Editing LOD is a non-generated LOD
+		if (ValidEditingLOD != EMeshLODIdentifier::HiResSource)
+		{
+			const int32 LODint = static_cast<int32>(ValidEditingLOD);
+			if (!StaticMeshIn->GetSourceModel(LODint).IsSourceModelInitialized())
+			{
+				DisplayCriticalWarningMessage(FString::Printf(TEXT("LOD%d Requested but is a generated LOD - Falling Back to LOD%d"), LODint, 0));
+				ValidEditingLOD = EMeshLODIdentifier::LOD0;
 			}
 		}
 	}
@@ -114,39 +125,11 @@ TArray<EMeshLODIdentifier> UStaticMeshReadOnlyToolTarget::GetAvailableLODs(const
 
 bool UStaticMeshReadOnlyToolTarget::IsValid() const
 {
-	// The StaticMesh.Get() call here will give a nullptr if object is not alive, which 
-	// HasNonGeneratedLOD checks for.
-	return HasNonGeneratedLOD(StaticMesh.Get(), EditingLOD);
-}
-
-bool UStaticMeshReadOnlyToolTarget::IsValid(const UStaticMesh* StaticMeshIn, EMeshLODIdentifier EditingLODIn)
-{
-	if (!StaticMeshIn || !IsValidChecked(StaticMeshIn) || StaticMeshIn->IsUnreachable() || !StaticMeshIn->IsValidLowLevel())
-	{
-		return false;
-	}
-
-	// Early out if the static mesh is still compiling. A significant number of functions checking the state of the static mesh are
-	// blocking on static mesh builds such as `GetNumSourceModels` used below. It is best to avoid long blocking calls in IsValid to
-	// ensure editor interactions are smooth. A tool target can't be properly initialized until the static mesh is finished building anyways.
-	if (StaticMeshIn->IsCompiling())
-	{
-		return false;
-	}
-
-	if (EditingLODIn == EMeshLODIdentifier::HiResSource)
-	{
-		if (StaticMeshIn->IsHiResMeshDescriptionValid() == false)
-		{
-			return false;
-		}
-	}
-	else if ((int32)EditingLODIn >= StaticMeshIn->GetNumSourceModels())
+	if (UStaticMesh* Mesh = StaticMesh.Get())
 	{
-		return false;
+		return ::IsValid(Mesh) && !Mesh->IsUnreachable() && Mesh->IsValidLowLevel();
 	}
-
-	return true;
+	return false;
 }
 
 bool UStaticMeshReadOnlyToolTarget::HasNonGeneratedLOD(const UStaticMesh* StaticMeshIn, EMeshLODIdentifier MeshLOD)
@@ -157,7 +140,7 @@ bool UStaticMeshReadOnlyToolTarget::HasNonGeneratedLOD(const UStaticMesh* Static
 	}
 
 	// Early out if the static mesh is still compiling. A significant number of functions checking the state of the static mesh are
-	// blocking on static mesh builds such as `GetNumSourceModels` used below. It is best to avoid long blocking calls in IsValid to
+	// blocking on static mesh builds such as `GetNumSourceModels` used below. It is best to avoid long blocking calls in CanBuildTarget to
 	// ensure editor interactions are smooth. A tool target can't be properly initialized until the static mesh is finished building anyways.
 	if (StaticMeshIn->IsCompiling())
 	{
@@ -283,7 +266,7 @@ const FMeshDescription* UStaticMeshReadOnlyToolTarget::GetMeshDescription(const
 		EMeshLODIdentifier UseLOD = EditingLOD;
 		if (GetMeshParams.bHaveRequestLOD)
 		{
-			UseLOD = UStaticMeshReadOnlyToolTarget::GetValidEditingLOD(StaticMesh.Get(), GetMeshParams.RequestLOD);
+			UseLOD = GetValidEditingLOD(StaticMesh.Get(), GetMeshParams.RequestLOD);
 			ensure(UseLOD == GetMeshParams.RequestLOD);		// probably a bug somewhere if this is not true
 		}
 
diff --git a/Engine/Plugins/Runtime/MeshModelingToolset/Source/ModelingComponentsEditorOnly/Public/ToolTargets/StaticMeshToolTarget.h b/Engine/Plugins/Runtime/MeshModelingToolset/Source/ModelingComponentsEditorOnly/Public/ToolTargets/StaticMeshToolTarget.h
--- a/Engine/Plugins/Runtime/MeshModelingToolset/Source/ModelingComponentsEditorOnly/Public/ToolTargets/StaticMeshToolTarget.h
+++ b/Engine/Plugins/Runtime/MeshModelingToolset/Source/ModelingComponentsEditorOnly/Public/ToolTargets/StaticMeshToolTarget.h
@@ -97,11 +97,6 @@ protected:
 
 	friend class UStaticMeshComponentToolTargetFactory;
 
-	UE_DEPRECATED(5.3, "Compared to using a TWeakObjectPtr, this function performs both redundant and insufficient "
-		"checks of whether the UStaticMesh is alive, and the LOD check does not guard against a generated LOD (which "
-		"tools should not edit) as HasNonGeneratedLOD does.")
-	static bool IsValid(const UStaticMesh* StaticMesh, EMeshLODIdentifier EditingLOD);
-
 	// returns true if the mesh has a non-generated LOD of the requested level.  Interprets EMeshLODIdentifier::Default as LOD0
 	static bool HasNonGeneratedLOD(const UStaticMesh* StaticMesh, EMeshLODIdentifier MeshLOD);
 	static EMeshLODIdentifier GetValidEditingLOD(const UStaticMesh* StaticMesh, 
