Modeling Tools: update how Vertex Sculpt tool registers different brushes to support re-registration

diff --git a/Engine/Plugins/Editor/ModelingToolsEditorMode/Source/ModelingToolsEditorMode/Private/DetailsCustomizations/MeshVertexSculptToolCustomizations.cpp b/Engine/Plugins/Editor/ModelingToolsEditorMode/Source/ModelingToolsEditorMode/Private/DetailsCustomizations/MeshVertexSculptToolCustomizations.cpp
--- a/Engine/Plugins/Editor/ModelingToolsEditorMode/Source/ModelingToolsEditorMode/Private/DetailsCustomizations/MeshVertexSculptToolCustomizations.cpp
+++ b/Engine/Plugins/Editor/ModelingToolsEditorMode/Source/ModelingToolsEditorMode/Private/DetailsCustomizations/MeshVertexSculptToolCustomizations.cpp
@@ -252,7 +252,7 @@ void FVertexBrushSculptPropertiesDetails::CustomizeDetails(IDetailLayoutBuilder&
 	UMeshVertexSculptTool* Tool = BrushProperties->Tool.Get();
 	TargetTool = Tool;
 
-	const TArray<UMeshSculptToolBase::FBrushTypeInfo>& BrushTypeInfos = Tool->GetRegisteredPrimaryBrushTypes();
+	const TSet<UMeshSculptToolBase::FBrushTypeInfo> BrushTypeInfos = Tool->GetRegisteredPrimaryBrushTypes();
 	int32 CurrentBrushType = (int32)BrushProperties->PrimaryBrushID;
 	int32 CurrentBrushTypeIndex = 0;
 
diff --git a/Engine/Plugins/Runtime/MeshModelingToolset/Source/MeshModelingTools/Private/Sculpting/MeshSculptToolBase.cpp b/Engine/Plugins/Runtime/MeshModelingToolset/Source/MeshModelingTools/Private/Sculpting/MeshSculptToolBase.cpp
--- a/Engine/Plugins/Runtime/MeshModelingToolset/Source/MeshModelingTools/Private/Sculpting/MeshSculptToolBase.cpp
+++ b/Engine/Plugins/Runtime/MeshModelingToolset/Source/MeshModelingTools/Private/Sculpting/MeshSculptToolBase.cpp
@@ -386,7 +386,12 @@ void UMeshSculptToolBase::RegisterBrushType(int32 Identifier, FText Name, TUniqu
 	TypeInfo.Identifier = Identifier;
 	RegisteredPrimaryBrushTypes.Add(TypeInfo);
 
-	check(BrushOpPropSets.Contains(Identifier) == false && BrushOpFactories.Contains(Identifier) == false);
+	// ensure we are not adding duplicates of PropSets when brush registration occurs
+	if (TObjectPtr<UMeshSculptBrushOpProps>* OldPropSet = BrushOpPropSets.Find(Identifier))
+	{
+		RemoveToolPropertySource(*OldPropSet);
+	}
+
 	BrushOpPropSets.Add(Identifier, PropSet);
 	BrushOpFactories.Add(Identifier, MoveTemp(Factory));
 
@@ -404,7 +409,12 @@ void UMeshSculptToolBase::RegisterSecondaryBrushType(int32 Identifier, FText Nam
 	TypeInfo.Identifier = Identifier;
 	RegisteredSecondaryBrushTypes.Add(TypeInfo);
 
-	check(SecondaryBrushOpPropSets.Contains(Identifier) == false && SecondaryBrushOpFactories.Contains(Identifier) == false);
+	// ensure we are not adding duplicates of PropSets when brush registration occurs
+	if (TObjectPtr<UMeshSculptBrushOpProps>* OldPropSet = SecondaryBrushOpPropSets.Find(Identifier))
+	{
+		RemoveToolPropertySource(*OldPropSet);
+	}
+	
 	SecondaryBrushOpPropSets.Add(Identifier, PropSet);
 	SecondaryBrushOpFactories.Add(Identifier, MoveTemp(Factory));
 
diff --git a/Engine/Plugins/Runtime/MeshModelingToolset/Source/MeshModelingTools/Public/Sculpting/MeshSculptToolBase.h b/Engine/Plugins/Runtime/MeshModelingToolset/Source/MeshModelingTools/Public/Sculpting/MeshSculptToolBase.h
--- a/Engine/Plugins/Runtime/MeshModelingToolset/Source/MeshModelingTools/Public/Sculpting/MeshSculptToolBase.h
+++ b/Engine/Plugins/Runtime/MeshModelingToolset/Source/MeshModelingTools/Public/Sculpting/MeshSculptToolBase.h
@@ -333,13 +333,24 @@ public:
 	{
 		FText Name;
 		int32 Identifier;
+
+		bool operator==( const FBrushTypeInfo& Other ) const
+		{
+			return Other.Identifier == Identifier &&
+					Other.Name.EqualTo(Name);
+		}
+
+		friend uint32 GetTypeHash(const FBrushTypeInfo& Other)
+		{
+			return HashCombine(GetTypeHash(Other.Identifier), GetTypeHash(Other.Name.ToString()));
+		}
 	};
-	const TArray<FBrushTypeInfo>& GetRegisteredPrimaryBrushTypes() const { return RegisteredPrimaryBrushTypes; }
-	const TArray<FBrushTypeInfo>& GetRegisteredSecondaryBrushTypes() const { return RegisteredSecondaryBrushTypes; }
+	const TSet<FBrushTypeInfo>& GetRegisteredPrimaryBrushTypes() const { return RegisteredPrimaryBrushTypes; }
+	const TSet<FBrushTypeInfo>& GetRegisteredSecondaryBrushTypes() const { return RegisteredSecondaryBrushTypes; }
 
 protected:
-	TArray<FBrushTypeInfo> RegisteredPrimaryBrushTypes;
-	TArray<FBrushTypeInfo> RegisteredSecondaryBrushTypes;
+	TSet<FBrushTypeInfo> RegisteredPrimaryBrushTypes;
+	TSet<FBrushTypeInfo> RegisteredSecondaryBrushTypes;
 
 	UPROPERTY()
 	TMap<int32, TObjectPtr<UMeshSculptBrushOpProps>> BrushOpPropSets;
