Add a new FMeshRegionChangeBase class for FMeshChange/FMeshVertexChange/FMeshReplacementChange to inherit, and add a delegate on dynamic mesh component that fires for these change types, so that the vertex sculpt tool (and other tools in future) can do local updates in response to more general changes

diff --git a/Engine/Plugins/Runtime/MeshModelingToolset/Source/MeshModelingTools/Private/MeshVertexSculptTool.cpp b/Engine/Plugins/Runtime/MeshModelingToolset/Source/MeshModelingTools/Private/MeshVertexSculptTool.cpp
--- a/Engine/Plugins/Runtime/MeshModelingToolset/Source/MeshModelingTools/Private/MeshVertexSculptTool.cpp
+++ b/Engine/Plugins/Runtime/MeshModelingToolset/Source/MeshModelingTools/Private/MeshVertexSculptTool.cpp
@@ -25,6 +25,7 @@
 #include "Generators/RectangleMeshGenerator.h"
 
 #include "Changes/MeshVertexChange.h"
+#include "Changes/MeshRegionChange.h"
 
 #include "Sculpting/KelvinletBrushOp.h"
 #include "Sculpting/MeshSmoothingBrushOps.h"
@@ -121,7 +122,7 @@ void UMeshVertexSculptTool::Setup()
 	}
 
 	DynamicMeshComponent->SetInvalidateProxyOnChangeEnabled(false);
-	OnDynamicMeshComponentChangedHandle = DynamicMeshComponent->OnMeshVerticesChanged.AddUObject(this, &UMeshVertexSculptTool::OnDynamicMeshComponentChanged);
+	OnDynamicMeshComponentChangedHandle = DynamicMeshComponent->OnMeshRegionChanged.AddUObject(this, &UMeshVertexSculptTool::OnDynamicMeshComponentChanged);
 
 	FDynamicMesh3* SculptMesh = GetSculptMesh();
 	FAxisAlignedBox3d Bounds = SculptMesh->GetBounds(true);
@@ -1653,7 +1654,7 @@ void UMeshVertexSculptTool::WaitForPendingUndoRedo()
 	}
 }
 
-void UMeshVertexSculptTool::OnDynamicMeshComponentChanged(UDynamicMeshComponent* Component, const FMeshVertexChange* Change, bool bRevert)
+void UMeshVertexSculptTool::OnDynamicMeshComponentChanged(UDynamicMeshComponent* Component, const FMeshRegionChangeBase* Change, bool bRevert)
 {
 	// have to wait for any outstanding stamp update to finish...
 	WaitForPendingStampUpdate();
@@ -1664,7 +1665,10 @@ void UMeshVertexSculptTool::OnDynamicMeshComponentChanged(UDynamicMeshComponent*
 
 	// figure out the set of modified triangles
 	AccumulatedTriangleROI.Reset();
-	UE::Geometry::VertexToTriangleOneRing(Mesh, Change->Vertices, AccumulatedTriangleROI);
+	Change->ProcessChangeVertices(Mesh, [this, &Mesh](TConstArrayView<int32> Vertices)
+	{
+		UE::Geometry::VertexToTriangleOneRing(Mesh, Vertices, AccumulatedTriangleROI);
+	}, bRevert);
 
 	// start the normal recomputation
 	UndoNormalsFuture = Async(VertexSculptToolAsyncExecTarget, [this, Mesh]()
diff --git a/Engine/Plugins/Runtime/MeshModelingToolset/Source/MeshModelingTools/Public/MeshVertexSculptTool.h b/Engine/Plugins/Runtime/MeshModelingToolset/Source/MeshModelingTools/Public/MeshVertexSculptTool.h
--- a/Engine/Plugins/Runtime/MeshModelingToolset/Source/MeshModelingTools/Public/MeshVertexSculptTool.h
+++ b/Engine/Plugins/Runtime/MeshModelingToolset/Source/MeshModelingTools/Public/MeshVertexSculptTool.h
@@ -309,7 +309,7 @@ protected:
 	TObjectPtr<UDynamicMeshComponent> DynamicMeshComponent = nullptr;
 
 	// realtime visualization
-	void OnDynamicMeshComponentChanged(UDynamicMeshComponent* Component, const FMeshVertexChange* Change, bool bRevert);
+	void OnDynamicMeshComponentChanged(UDynamicMeshComponent* Component, const FMeshRegionChangeBase* Change, bool bRevert);
 	FDelegateHandle OnDynamicMeshComponentChangedHandle;
 
 	UE_DEPRECATED(5.6, "Vertex sculpt now uses an integer for a brush ID to allow custom brush registration")
diff --git a/Engine/Source/Runtime/GeometryCore/Private/DynamicMesh/DynamicMeshChangeTracker.cpp b/Engine/Source/Runtime/GeometryCore/Private/DynamicMesh/DynamicMeshChangeTracker.cpp
--- a/Engine/Source/Runtime/GeometryCore/Private/DynamicMesh/DynamicMeshChangeTracker.cpp
+++ b/Engine/Source/Runtime/GeometryCore/Private/DynamicMesh/DynamicMeshChangeTracker.cpp
@@ -878,6 +878,24 @@ void FDynamicMeshChange::GetSavedTriangleList(TArray<int>& TrianglesOut, bool bI
 	}
 }
 
+void FDynamicMeshChange::GetAffectedVertices(TArray<int>& VerticesOut, bool bRevert) const
+{
+	const TArray<FChangeTriangle>& UseTris = (bRevert) ? OldTriangles : NewTriangles;
+	const TArray<FChangeVertex>& UseVerts = (bRevert) ? OldVertices : NewVertices;
+	TSet<int32> UniqueVIDs;
+	for (const FChangeTriangle& Tri : UseTris)
+	{
+		UniqueVIDs.Add(Tri.Vertices.A);
+		UniqueVIDs.Add(Tri.Vertices.B);
+		UniqueVIDs.Add(Tri.Vertices.C);
+	}
+	for (const FChangeVertex& Vert : UseVerts)
+	{
+		UniqueVIDs.Add(Vert.VertexID);
+	}
+	VerticesOut = UniqueVIDs.Array();
+}
+
 
 
 void FDynamicMeshChange::CheckValidity(EValidityCheckFailMode FailMode) const
diff --git a/Engine/Source/Runtime/GeometryCore/Private/DynamicMesh/MeshIndexUtil.cpp b/Engine/Source/Runtime/GeometryCore/Private/DynamicMesh/MeshIndexUtil.cpp
--- a/Engine/Source/Runtime/GeometryCore/Private/DynamicMesh/MeshIndexUtil.cpp
+++ b/Engine/Source/Runtime/GeometryCore/Private/DynamicMesh/MeshIndexUtil.cpp
@@ -5,7 +5,7 @@
 
 using namespace UE::Geometry;
 
-void UE::Geometry::TriangleToVertexIDs(const FDynamicMesh3* Mesh, const TArray<int>& TriangleIDs, TArray<int>& VertexIDsOut)
+void UE::Geometry::TriangleToVertexIDs(const FDynamicMesh3* Mesh, TConstArrayView<int> TriangleIDs, TArray<int>& VertexIDsOut)
 {
 	// if we are getting close to full mesh it is probably more efficient to use a bitmap...
 
@@ -50,7 +50,7 @@ void UE::Geometry::TriangleToVertexIDs(const FDynamicMesh3* Mesh, const TArray<i
 
 
 
-void UE::Geometry::VertexToTriangleOneRing(const FDynamicMesh3* Mesh, const TArray<int>& VertexIDs, TSet<int>& TriangleIDsOut)
+void UE::Geometry::VertexToTriangleOneRing(const FDynamicMesh3* Mesh, TConstArrayView<int> VertexIDs, TSet<int>& TriangleIDsOut)
 {
 	// for a TSet it is more efficient to just try to add each triangle twice, than it is to
 	// try to avoid duplicate adds with more complex mesh queries
diff --git a/Engine/Source/Runtime/GeometryCore/Public/DynamicMesh/DynamicMeshChangeTracker.h b/Engine/Source/Runtime/GeometryCore/Public/DynamicMesh/DynamicMeshChangeTracker.h
--- a/Engine/Source/Runtime/GeometryCore/Public/DynamicMesh/DynamicMeshChangeTracker.h
+++ b/Engine/Source/Runtime/GeometryCore/Public/DynamicMesh/DynamicMeshChangeTracker.h
@@ -169,6 +169,9 @@ public:
 	/** store IDs of saved triangles in TrianglesOut. if bInitial=true, old triangles are stored, otherwise new triangles */
 	GEOMETRYCORE_API void GetSavedTriangleList(TArray<int>& TrianglesOut, bool bInitial) const;
 
+	/** store IDs of affected vertices in VerticesOut -- both vertices on saved triangles, and directly saved vertices. if bRevert=true, returns old vertex IDs, otherwise new IDs are returned. */
+	GEOMETRYCORE_API void GetAffectedVertices(TArray<int>& VerticesOut, bool bRevert) const;
+
 	/** run self-validity checks on internal data structures to test if change is well-formed */
 	GEOMETRYCORE_API void CheckValidity(EValidityCheckFailMode FailMode = EValidityCheckFailMode::Check) const;
 
diff --git a/Engine/Source/Runtime/GeometryCore/Public/DynamicMesh/MeshIndexUtil.h b/Engine/Source/Runtime/GeometryCore/Public/DynamicMesh/MeshIndexUtil.h
--- a/Engine/Source/Runtime/GeometryCore/Public/DynamicMesh/MeshIndexUtil.h
+++ b/Engine/Source/Runtime/GeometryCore/Public/DynamicMesh/MeshIndexUtil.h
@@ -20,7 +20,7 @@ namespace Geometry
 	 * @param TriangleIDs list of triangle IDs of Mesh
 	 * @param VertexIDsOut list of vertices contained by triangles
 	 */
-	GEOMETRYCORE_API void TriangleToVertexIDs(const FDynamicMesh3* Mesh, const TArray<int>& TriangleIDs, TArray<int>& VertexIDsOut);
+	GEOMETRYCORE_API void TriangleToVertexIDs(const FDynamicMesh3* Mesh, TConstArrayView<int> TriangleIDs, TArray<int>& VertexIDsOut);
 
 
 	/**
@@ -29,7 +29,7 @@ namespace Geometry
 	 * @param VertexIDs list of Vertex IDs of Mesh
 	 * @param TriangleIDsOut list of triangle IDs where any vertex of triangle touches an element of VertexIDs
 	 */
-	GEOMETRYCORE_API void VertexToTriangleOneRing(const FDynamicMesh3* Mesh, const TArray<int>& VertexIDs, TSet<int>& TriangleIDsOut);
+	GEOMETRYCORE_API void VertexToTriangleOneRing(const FDynamicMesh3* Mesh, TConstArrayView<int> VertexIDs, TSet<int>& TriangleIDsOut);
 
 
 	/**
diff --git a/Engine/Source/Runtime/GeometryFramework/Private/Changes/MeshChange.cpp b/Engine/Source/Runtime/GeometryFramework/Private/Changes/MeshChange.cpp
--- a/Engine/Source/Runtime/GeometryFramework/Private/Changes/MeshChange.cpp
+++ b/Engine/Source/Runtime/GeometryFramework/Private/Changes/MeshChange.cpp
@@ -38,6 +38,15 @@ void FMeshChange::Revert(UObject* Object)
 	}
 }
 
+void FMeshChange::ProcessChangeVertices(const FDynamicMesh3* ChangedMesh, TFunctionRef<void(TConstArrayView<int32>)> ProcessFn, bool bRevert) const
+{
+	TArray<int32> VerticesOut;
+	if (ensure(DynamicMeshChange))
+	{
+		DynamicMeshChange->GetAffectedVertices(VerticesOut, bRevert);
+	}
+	ProcessFn(VerticesOut);
+}
 
 FString FMeshChange::ToString() const
 {
diff --git a/Engine/Source/Runtime/GeometryFramework/Private/Changes/MeshReplacementChange.cpp b/Engine/Source/Runtime/GeometryFramework/Private/Changes/MeshReplacementChange.cpp
--- a/Engine/Source/Runtime/GeometryFramework/Private/Changes/MeshReplacementChange.cpp
+++ b/Engine/Source/Runtime/GeometryFramework/Private/Changes/MeshReplacementChange.cpp
@@ -39,6 +39,17 @@ void FMeshReplacementChange::Revert(UObject* Object)
 	}
 }
 
+void FMeshReplacementChange::ProcessChangeVertices(const FDynamicMesh3* ChangedMesh, TFunctionRef<void(TConstArrayView<int32>)> ProcessFn, bool bRevert) const
+{
+	check(ChangedMesh);
+	TArray<int32> Vertices;
+	Vertices.Reserve(ChangedMesh->VertexCount());
+	for (int32 VID : ChangedMesh->VertexIndicesItr())
+	{
+		Vertices.Add(VID);
+	}
+	ProcessFn(Vertices);
+}
 
 FString FMeshReplacementChange::ToString() const
 {
diff --git a/Engine/Source/Runtime/GeometryFramework/Private/Changes/MeshVertexChange.cpp b/Engine/Source/Runtime/GeometryFramework/Private/Changes/MeshVertexChange.cpp
--- a/Engine/Source/Runtime/GeometryFramework/Private/Changes/MeshVertexChange.cpp
+++ b/Engine/Source/Runtime/GeometryFramework/Private/Changes/MeshVertexChange.cpp
@@ -21,6 +21,11 @@ void FMeshVertexChange::Revert(UObject* Object)
 	ChangeTarget->ApplyChange(this, true);
 }
 
+void FMeshVertexChange::ProcessChangeVertices(const FDynamicMesh3* ChangedMesh, TFunctionRef<void(TConstArrayView<int32>)> ProcessFn, bool bRevert) const
+{
+	// Note: Since the FMeshVertexChange doesn't change vertex IDs, we don't need to look at ChangedMesh or bRevert and can just pass the Vertices array
+	ProcessFn(Vertices);
+}
 
 FString FMeshVertexChange::ToString() const
 {
diff --git a/Engine/Source/Runtime/GeometryFramework/Private/Components/DynamicMeshComponent.cpp b/Engine/Source/Runtime/GeometryFramework/Private/Components/DynamicMeshComponent.cpp
--- a/Engine/Source/Runtime/GeometryFramework/Private/Components/DynamicMeshComponent.cpp
+++ b/Engine/Source/Runtime/GeometryFramework/Private/Components/DynamicMeshComponent.cpp
@@ -1378,6 +1378,7 @@ void UDynamicMeshComponent::OnMeshObjectChanged(UDynamicMesh* ChangedMeshObject,
 		{
 			OnMeshVerticesChanged.Broadcast(this, ChangeInfo.VertexChange, ChangeInfo.bIsRevertChange);
 		}
+		OnMeshRegionChanged.Broadcast(this, ChangeInfo.GetChange(), ChangeInfo.bIsRevertChange);
 	}
 	else
 	{
diff --git a/Engine/Source/Runtime/GeometryFramework/Public/Changes/MeshChange.h b/Engine/Source/Runtime/GeometryFramework/Public/Changes/MeshChange.h
--- a/Engine/Source/Runtime/GeometryFramework/Public/Changes/MeshChange.h
+++ b/Engine/Source/Runtime/GeometryFramework/Public/Changes/MeshChange.h
@@ -2,6 +2,7 @@
 
 #pragma once
 
+#include "Changes/MeshRegionChange.h"
 #include "Containers/UnrealString.h"
 #include "CoreMinimal.h"
 #include "DynamicMesh/DynamicMeshChangeTracker.h"
@@ -27,7 +28,7 @@ namespace UE { namespace Geometry { class FDynamicMesh3; } }
  * FMeshChange represents an undoable change to a FDynamicMesh3.
  * Currently only valid to call Apply/Revert when the Object is a one of several components backed by FDynamicMesh: UDynamicMeshComponent, UOctreeDynamicMeshComponent, UPreviewMesh
  */
-class FMeshChange : public FToolCommandChange
+class FMeshChange : public FMeshRegionChangeBase
 {
 public:
 	GEOMETRYFRAMEWORK_API FMeshChange();
@@ -45,6 +46,9 @@ public:
 	/** Reverts change to the object */
 	GEOMETRYFRAMEWORK_API virtual void Revert(UObject* Object) override;
 
+	/** FRegionMeshChangeBase API -- builds an array of all vertices changed and passes to the ProcessFn lambda */
+	GEOMETRYFRAMEWORK_API virtual void ProcessChangeVertices(const UE::Geometry::FDynamicMesh3* ChangedMesh, TFunctionRef<void(TConstArrayView<int32>)> ProcessFn, bool bRevert) const override;
+
 	/** Describes this change (for debugging) */
 	GEOMETRYFRAMEWORK_API virtual FString ToString() const override;
 
diff --git a/Engine/Source/Runtime/GeometryFramework/Public/Changes/MeshRegionChange.h b/Engine/Source/Runtime/GeometryFramework/Public/Changes/MeshRegionChange.h
new file mode 100644
--- /dev/null
+++ b/Engine/Source/Runtime/GeometryFramework/Public/Changes/MeshRegionChange.h
@@ -0,0 +1,17 @@
+// Copyright Epic Games, Inc. All Rights Reserved.
+
+#pragma once
+
+#include "InteractiveToolChange.h"
+
+namespace UE::Geometry { class FDynamicMesh3; }
+
+// Base class for mesh changes that apply to a selection of the mesh
+class FMeshRegionChangeBase : public FToolCommandChange
+{
+public:
+	// @param ChangedMesh The mesh after the change has been applied (or undone, if bRevert is true)
+	// @param ProcessFn The callback that will be called with the vertices affected by this mesh change
+	// @param bRevert Whether the changes is an undo
+	virtual void ProcessChangeVertices(const UE::Geometry::FDynamicMesh3* ChangedMesh, TFunctionRef<void(TConstArrayView<int32>)> ProcessFn, bool bRevert) const = 0;
+};
diff --git a/Engine/Source/Runtime/GeometryFramework/Public/Changes/MeshReplacementChange.h b/Engine/Source/Runtime/GeometryFramework/Public/Changes/MeshReplacementChange.h
--- a/Engine/Source/Runtime/GeometryFramework/Public/Changes/MeshReplacementChange.h
+++ b/Engine/Source/Runtime/GeometryFramework/Public/Changes/MeshReplacementChange.h
@@ -2,6 +2,7 @@
 
 #pragma once
 
+#include "Changes/MeshRegionChange.h"
 #include "Containers/UnrealString.h"
 #include "CoreMinimal.h"
 #include "GeometryBase.h"
@@ -22,7 +23,7 @@ class UObject;
  * FMeshReplacementChange represents an undoable *complete* change to a FDynamicMesh3.
  * Currently only valid to call Apply/Revert when the Object is a UDynamicMeshComponent
  */
-class FMeshReplacementChange : public FToolCommandChange
+class FMeshReplacementChange : public FMeshRegionChangeBase
 {
 	TSharedPtr<const FDynamicMesh3, ESPMode::ThreadSafe> Before, After;
 
@@ -44,6 +45,9 @@ public:
 	/** Reverts change to the object */
 	GEOMETRYFRAMEWORK_API virtual void Revert(UObject* Object) override;
 
+	/** FRegionMeshChangeBase API -- passes all mesh vertex IDs as an array to the ProcessFn lambda */
+	GEOMETRYFRAMEWORK_API virtual void ProcessChangeVertices(const UE::Geometry::FDynamicMesh3* ChangedMesh, TFunctionRef<void(TConstArrayView<int32>)> ProcessFn, bool bRevert) const override;
+
 	/** Describes this change (for debugging) */
 	GEOMETRYFRAMEWORK_API virtual FString ToString() const override;
 };
diff --git a/Engine/Source/Runtime/GeometryFramework/Public/Changes/MeshVertexChange.h b/Engine/Source/Runtime/GeometryFramework/Public/Changes/MeshVertexChange.h
--- a/Engine/Source/Runtime/GeometryFramework/Public/Changes/MeshVertexChange.h
+++ b/Engine/Source/Runtime/GeometryFramework/Public/Changes/MeshVertexChange.h
@@ -2,6 +2,7 @@
 
 #pragma once
 
+#include "Changes/MeshRegionChange.h"
 #include "CoreMinimal.h"
 #include "UObject/Interface.h"
 #include "InteractiveToolChange.h"
@@ -17,7 +18,7 @@ PREDECLARE_USE_GEOMETRY_CLASS(FDynamicMesh3);
  * @todo support optionally storing old/new normals and tangents
  * @todo support applying to a StaticMeshComponent/MeshDescription ?
  */
-class FMeshVertexChange : public FToolCommandChange
+class FMeshVertexChange : public FMeshRegionChangeBase
 {
 public:
 	bool bHaveVertexPositions = true;
@@ -44,6 +45,9 @@ public:
 	/** Reverts change to the object */
 	GEOMETRYFRAMEWORK_API virtual void Revert(UObject* Object) override;
 
+	/** FRegionMeshChangeBase API -- passes the Vertices array to the ProcessFn lambda */
+	GEOMETRYFRAMEWORK_API virtual void ProcessChangeVertices(const UE::Geometry::FDynamicMesh3* ChangedMesh, TFunctionRef<void(TConstArrayView<int32>)> ProcessFn, bool bRevert) const override;
+
 	/** Describes this change (for debugging) */
 	GEOMETRYFRAMEWORK_API virtual FString ToString() const override;
 };
diff --git a/Engine/Source/Runtime/GeometryFramework/Public/Components/DynamicMeshComponent.h b/Engine/Source/Runtime/GeometryFramework/Public/Components/DynamicMeshComponent.h
--- a/Engine/Source/Runtime/GeometryFramework/Public/Components/DynamicMeshComponent.h
+++ b/Engine/Source/Runtime/GeometryFramework/Public/Components/DynamicMeshComponent.h
@@ -2,6 +2,7 @@
 
 #pragma once
 
+#include "Changes/MeshRegionChange.h"
 #include "CoreMinimal.h"
 #include "BaseDynamicMeshComponent.h"
 #include "MeshConversionOptions.h"
@@ -402,6 +403,13 @@ public:
 	DECLARE_MULTICAST_DELEGATE_ThreeParams(FMeshVerticesModified, UDynamicMeshComponent*, const FMeshVertexChange*, bool);
 	FMeshVerticesModified OnMeshVerticesChanged;
 
+	/**
+	 * This delegate fires when the mesh has been changed via an FMeshRegionChange-derived change -- FMeshChange, FMeshVertexChange, or FMeshReplacementChange
+	 * Note that in the case of an FMeshVertexChange, the OnMeshVerticesChanged delegate will also fire
+	 */
+	DECLARE_MULTICAST_DELEGATE_ThreeParams(FMeshRegionModified, UDynamicMeshComponent*, const FMeshRegionChangeBase*, bool);
+	FMeshRegionModified OnMeshRegionChanged;
+
 	/**
 	 * When a FMeshChange or FMeshVertexChange is applied, by default we currently fully invalidate the render proxy. However in certain
 	 * realtime situations (eg like Sculpting tools) it can be critical to undo/redo performance to do more optimized render data updates
diff --git a/Engine/Source/Runtime/GeometryFramework/Public/UDynamicMesh.h b/Engine/Source/Runtime/GeometryFramework/Public/UDynamicMesh.h
--- a/Engine/Source/Runtime/GeometryFramework/Public/UDynamicMesh.h
+++ b/Engine/Source/Runtime/GeometryFramework/Public/UDynamicMesh.h
@@ -79,6 +79,25 @@ struct FDynamicMeshChangeInfo
 	UPROPERTY(BlueprintReadOnly, EditAnywhere, Category = "DynamicMeshChangeInfo")
 	bool bIsRevertChange = false;
 
+	// return the single valid FChange, or null
+	const FMeshRegionChangeBase* GetChange() const
+	{
+		if (MeshChange)
+		{
+			ensure(!ReplaceChange && !VertexChange); // not valid to have multiple non-null FChanges
+			return MeshChange;
+		}
+		else if (ReplaceChange)
+		{
+			ensure(!VertexChange); // not valid to have multiple non-null FChanges
+			return ReplaceChange;
+		}
+		else
+		{
+			return VertexChange;
+		}
+	}
+
 	//
 	// internals
 	//
