add initial sculpt layer controls to vertex sculpt tool

diff --git a/Engine/Plugins/Runtime/MeshModelingToolset/Source/MeshModelingTools/Private/MeshVertexSculptTool.cpp b/Engine/Plugins/Runtime/MeshModelingToolset/Source/MeshModelingTools/Private/MeshVertexSculptTool.cpp
--- a/Engine/Plugins/Runtime/MeshModelingToolset/Source/MeshModelingTools/Private/MeshVertexSculptTool.cpp
+++ b/Engine/Plugins/Runtime/MeshModelingToolset/Source/MeshModelingTools/Private/MeshVertexSculptTool.cpp
@@ -27,6 +27,7 @@
 #include "Changes/MeshVertexChange.h"
 #include "Changes/MeshRegionChange.h"
 
+#include "MeshSculptLayersManagerAPI.h"
 #include "Sculpting/KelvinletBrushOp.h"
 #include "Sculpting/MeshSmoothingBrushOps.h"
 #include "Sculpting/MeshInflateBrushOps.h"
@@ -93,7 +94,101 @@ public:
 	virtual void Revert(UObject* Object) override;
 };
 
+void UMeshSculptLayerProperties::Init(IModelingToolExternalDynamicMeshUpdateAPI* InTool, int32 InNumLockedBaseLayers)
+{
+	bHasSculptLayers = true;
+	LayerWeights.Reset();
+
+	Tool = InTool;
+	if (!Tool)
+	{
+		return;
+	}
+	NumLockedBaseLayers = InNumLockedBaseLayers;
+
+	Tool->ProcessToolMeshes([this](const FDynamicMesh3& Mesh, int32 MeshIdx)
+		{
+			// Sculpt layer UI only supports a single mesh for now
+			if (MeshIdx > 0)
+			{
+				return;
+			}
+			bHasSculptLayers = Mesh.HasAttributes() && Mesh.Attributes()->NumSculptLayers() > NumLockedBaseLayers;
+			if (bHasSculptLayers)
+			{
+				ActiveLayer = Mesh.Attributes()->GetSculptLayers()->GetActiveLayer();
+				TConstArrayView<double> SourceWeights = Mesh.Attributes()->GetSculptLayers()->GetLayerWeights();
+				int32 NumLayers = FMath::Max(0, SourceWeights.Num() - NumLockedBaseLayers);
+				LayerWeights.SetNum(NumLayers);
+				for (int32 Idx = 0; Idx < LayerWeights.Num(); ++Idx)
+				{
+					LayerWeights[Idx] = SourceWeights[Idx + NumLockedBaseLayers];
+				}
+			}
+		}
+	);
+}
+
+void UMeshSculptLayerProperties::SetLayerWeights(FDynamicMeshSculptLayers* SculptLayers) const
+{
+	TArray<double> FullLayerWeights(SculptLayers->GetLayerWeights());
+	FullLayerWeights.SetNum(LayerWeights.Num() + NumLockedBaseLayers);
+	for (int32 Idx = 0; Idx < LayerWeights.Num(); ++Idx)
+	{
+		FullLayerWeights[Idx + NumLockedBaseLayers] = LayerWeights[Idx];
+	}
+	SculptLayers->UpdateLayerWeights(FullLayerWeights);
+}
 
+#if WITH_EDITOR
+void UMeshSculptLayerProperties::PostEditChangeProperty(FPropertyChangedEvent& PropertyChangedEvent)
+{
+	Super::PostEditChangeProperty(PropertyChangedEvent);
+
+	if (Tool && Tool->AllowToolMeshUpdates())
+	{
+		Tool->UpdateToolMeshes([this, &PropertyChangedEvent](FDynamicMesh3& Mesh, int32 Idx) -> TUniquePtr<FMeshRegionChangeBase>
+			{
+				// Sculpt layer UI only supports a single mesh for now
+				if (Idx > 0)
+				{
+					return nullptr;
+				}
+
+				if (!InitialMesh)
+				{
+					InitialMesh = MakeShared<FDynamicMesh3>(Mesh);
+				}
+				if (FDynamicMeshAttributeSet* Attributes = Mesh.Attributes())
+				{
+					if (FDynamicMeshSculptLayers* SculptLayers = Attributes->GetSculptLayers())
+					{
+						if (PropertyChangedEvent.GetMemberPropertyName() == GET_MEMBER_NAME_CHECKED(UMeshSculptLayerProperties, LayerWeights))
+						{
+							SetLayerWeights(SculptLayers);
+						}
+						if (PropertyChangedEvent.GetMemberPropertyName() == GET_MEMBER_NAME_CHECKED(UMeshSculptLayerProperties, ActiveLayer))
+						{
+							ActiveLayer = FMath::Clamp(ActiveLayer, NumLockedBaseLayers, SculptLayers->NumLayers() - 1);
+							SculptLayers->SetActiveLayer(ActiveLayer);
+						}
+					}
+				}
+
+				
+
+				if (PropertyChangedEvent.ChangeType == EPropertyChangeType::Interactive)
+				{
+					return nullptr;
+				}
+				TUniquePtr<FMeshRegionChangeBase> Result(new FMeshReplacementChange(InitialMesh, MakeShared<FDynamicMesh3>(Mesh)));
+				InitialMesh = nullptr;
+				return MoveTemp(Result);
+			}
+		);
+	}
+}
+#endif
 
 /*
  * Tool
@@ -233,6 +328,21 @@ void UMeshVertexSculptTool::Setup()
 	SymmetryProperties->bSymmetryCanBeEnabled = false;
 	AddToolPropertySource(SymmetryProperties);
 
+	if (ISceneComponentBackedTarget* SceneComponentTarget = Cast<ISceneComponentBackedTarget>(Target))
+	{
+		if (IMeshSculptLayersManager* SculptLayersManager = Cast<IMeshSculptLayersManager>(SceneComponentTarget->GetOwnerSceneComponent()))
+		{
+			if (SculptLayersManager->HasSculptLayers())
+			{
+				SculptLayerProperties = NewObject<UMeshSculptLayerProperties>(this);
+				SculptLayerProperties->Init(this, SculptLayersManager->NumLockedBaseSculptLayers());
+				AddToolPropertySource(SculptLayerProperties);
+			}
+		}
+	}
+
+
+
 	this->BaseMeshQueryFunc = [&](int32 VertexID, const FVector3d& Position, double MaxDist, FVector3d& PosOut, FVector3d& NormalOut)
 	{
 		return GetBaseMeshNearest(VertexID, Position, MaxDist, PosOut, NormalOut);
@@ -396,7 +506,39 @@ void UMeshVertexSculptTool::OnPropertyModified(UObject* PropertySet, FProperty*
 	CalculateBrushRadius();
 }
 
+void UMeshVertexSculptTool::UpdateToolMeshes(TFunctionRef<TUniquePtr<FMeshRegionChangeBase>(FDynamicMesh3&, int32 MeshIdx)> UpdateMesh)
+{
+	if (AllowToolMeshUpdates())
+	{
+		// have to wait for any outstanding stamp/undo update to finish...
+		WaitForPendingStampUpdate();
+		WaitForPendingUndoRedo();
+
+		TUniquePtr<FMeshRegionChangeBase> Change = UpdateMesh(*GetSculptMesh(), 0);
+		// A change was created -- emit it to the tool manager and update associated data structures, etc
+		if (Change)
+		{
+			// pass through the change to trigger standard mesh updates / octree recomputation
+			OnDynamicMeshComponentChanged(DynamicMeshComponent, Change.Get(), false);
 
+			TUniquePtr<TWrappedToolCommandChange<FMeshRegionChangeBase>> NewChange = MakeUnique<TWrappedToolCommandChange<FMeshRegionChangeBase>>();
+			NewChange->WrappedChange = MoveTemp(Change);
+			NewChange->BeforeModify = [this](bool bRevert)
+			{
+				this->WaitForPendingUndoRedo();
+			};
+			// Note this change should be in the context of a larger transaction, so the text isn't that important
+			GetToolManager()->EmitObjectChange(DynamicMeshComponent, MoveTemp(NewChange), LOCTEXT("UpdateVertexSculptMesh", "Updated Mesh"));
+			
+			// TODO: If bMeshSymmetryIsValid, re-check that symmetry is still valid after sculpt layer change
+		}
+		// No change is ready to emit, just update component rendering
+		else
+		{
+			DynamicMeshComponent->FastNotifyPositionsUpdated();
+		}
+	}
+}
 
 UPreviewMesh* UMeshVertexSculptTool::MakeBrushIndicatorMesh(UObject* Parent, UWorld* World)
 {
diff --git a/Engine/Plugins/Runtime/MeshModelingToolset/Source/MeshModelingTools/Public/MeshVertexSculptTool.h b/Engine/Plugins/Runtime/MeshModelingToolset/Source/MeshModelingTools/Public/MeshVertexSculptTool.h
--- a/Engine/Plugins/Runtime/MeshModelingToolset/Source/MeshModelingTools/Public/MeshVertexSculptTool.h
+++ b/Engine/Plugins/Runtime/MeshModelingToolset/Source/MeshModelingTools/Public/MeshVertexSculptTool.h
@@ -11,6 +11,7 @@
 #include "DynamicMesh/MeshNormals.h"
 #include "BaseTools/BaseBrushTool.h"
 #include "BaseTools/MeshSurfacePointMeshEditingTool.h"
+#include "ModelingToolExternalMeshUpdateAPI.h"
 #include "TransformTypes.h"
 #include "Sculpting/MeshSculptToolBase.h"
 #include "Sculpting/MeshBrushOpBase.h"
@@ -23,6 +24,7 @@
 class UMaterialInstanceDynamic;
 class FMeshVertexChangeBuilder;
 class UPreviewMesh;
+namespace UE::Geometry { class FDynamicMeshSculptLayers; }
 PREDECLARE_GEOMETRY(class FMeshPlanarSymmetry);
 
 
@@ -206,11 +208,50 @@ public:
 
 
 
+
+
+//~ TODO: These properties should be pulled out to a separate UI, and not tied to the vertex sculpt tool
+UCLASS(MinimalAPI)
+class UMeshSculptLayerProperties : public UObject
+{
+	GENERATED_BODY()
+public:
+	/** Set the active sculpt layer */
+	UPROPERTY(EditAnywhere, Category = SculptLayers, meta = (ClampMin = 0, HideEditConditionToggle, EditCondition = bHasSculptLayers, ModelingQuickSettings))
+	int32 ActiveLayer = 0;
+
+	/** Set the sculpt layer weights */
+	UPROPERTY(EditAnywhere, Category = SculptLayers, meta = (HideEditConditionToggle, EditCondition = bHasSculptLayers, ModelingQuickSettings))
+	TArray<double> LayerWeights;
+
+	UPROPERTY(meta = (TransientToolProperty))
+	bool bHasSculptLayers = false;
+
+	MESHMODELINGTOOLS_API void Init(IModelingToolExternalDynamicMeshUpdateAPI* InTool, int32 InNumLockedBaseLayers);
+
+#if WITH_EDITOR
+	MESHMODELINGTOOLS_API virtual void PostEditChangeProperty(FPropertyChangedEvent& PropertyChangedEvent) override;
+#endif
+
+private:
+	IModelingToolExternalDynamicMeshUpdateAPI* Tool;
+
+	int32 NumLockedBaseLayers = 0;
+
+	// Mesh before layer change has been applied. Used for tracking mesh changes that occur over multiple frames (i.e., from an interactive drag)
+	// TODO: Add/update a mesh FChange to track sculpt layer changes, and use that instead of saving an entire mesh here
+	TSharedPtr<FDynamicMesh3> InitialMesh;
+
+	// Helper to set sculpt layers from the current LayerWeights property (accounting for the NumLockedBaseLayers)
+	void SetLayerWeights(UE::Geometry::FDynamicMeshSculptLayers* SculptLayers) const;
+};
+
+
 /**
  * Mesh Vertex Sculpt Tool Class
  */
 UCLASS()
-class MESHMODELINGTOOLS_API UMeshVertexSculptTool : public UMeshSculptToolBase, public IInteractiveToolManageGeometrySelectionAPI
+class MESHMODELINGTOOLS_API UMeshVertexSculptTool : public UMeshSculptToolBase, public IInteractiveToolManageGeometrySelectionAPI, public IModelingToolExternalDynamicMeshUpdateAPI
 {
 	GENERATED_BODY()
 public:
@@ -234,6 +275,18 @@ public:
 		return true;
 	}
 
+	// IInteractiveToolExternalDynamicMeshUpdateAPI methods
+	virtual bool AllowToolMeshUpdates() const override { return !ActiveVertexChange; }
+	virtual void UpdateToolMeshes(TFunctionRef<TUniquePtr<FMeshRegionChangeBase>(FDynamicMesh3&, int32 MeshIdx)> UpdateMesh) override;
+	virtual void ProcessToolMeshes(TFunctionRef<void(const FDynamicMesh3&, int32 MeshIdx)> ProcessMesh) const override
+	{
+		ProcessMesh(*GetSculptMesh(), 0);
+	}
+	virtual int32 NumToolMeshes() const override
+	{
+		return 1;
+	}
+
 public:
 
 	/** Properties that control sculpting*/
@@ -249,6 +302,11 @@ public:
 	UPROPERTY()
 	TObjectPtr<UMeshSymmetryProperties> SymmetryProperties;
 
+private:
+	// TODO: These sculpt layer properties should be moved off this tool, to a separate UI
+	UPROPERTY()
+	TObjectPtr<UMeshSculptLayerProperties> SculptLayerProperties;
+
 public:
 	virtual void IncreaseBrushSpeedAction() override;
 	virtual void DecreaseBrushSpeedAction() override;
diff --git a/Engine/Plugins/Runtime/MeshModelingToolset/Source/ModelingComponents/Public/MeshSculptLayersManagerAPI.h b/Engine/Plugins/Runtime/MeshModelingToolset/Source/ModelingComponents/Public/MeshSculptLayersManagerAPI.h
new file mode 100644
--- /dev/null
+++ b/Engine/Plugins/Runtime/MeshModelingToolset/Source/ModelingComponents/Public/MeshSculptLayersManagerAPI.h
@@ -0,0 +1,37 @@
+// Copyright Epic Games, Inc. All Rights Reserved.
+
+#pragma once
+
+#include "CoreMinimal.h"
+#include "UObject/Interface.h"
+
+#include "GeometryBase.h"
+
+#include "MeshSculptLayersManagerAPI.generated.h"
+
+
+UINTERFACE()
+class MODELINGCOMPONENTS_API UMeshSculptLayersManager : public UInterface
+{
+	GENERATED_BODY()
+};
+
+// API to provide control over mesh sculpt layer support
+class MODELINGCOMPONENTS_API IMeshSculptLayersManager
+{
+	GENERATED_BODY()
+
+public:
+
+	virtual bool HasSculptLayers() const
+	{
+		return false;
+	}
+
+	// @return A number of base sculpt layers which should not be edited, or 0 if all layers can be edited
+	virtual int32 NumLockedBaseSculptLayers() const
+	{
+		return 0;
+	}
+	
+};
diff --git a/Engine/Plugins/Runtime/MeshModelingToolset/Source/ModelingComponents/Public/ModelingToolExternalMeshUpdateAPI.h b/Engine/Plugins/Runtime/MeshModelingToolset/Source/ModelingComponents/Public/ModelingToolExternalMeshUpdateAPI.h
new file mode 100644
--- /dev/null
+++ b/Engine/Plugins/Runtime/MeshModelingToolset/Source/ModelingComponents/Public/ModelingToolExternalMeshUpdateAPI.h
@@ -0,0 +1,43 @@
+// Copyright Epic Games, Inc. All Rights Reserved.
+
+#pragma once
+
+#include "Changes/MeshRegionChange.h"
+#include "DynamicMesh/DynamicMesh3.h"
+#include "UObject/Interface.h"
+
+#include "ModelingToolExternalMeshUpdateAPI.generated.h"
+
+
+class UInteractiveToolManager;
+
+// UInterface for IModelingToolExternalDynamicMeshUpdateAPI
+UINTERFACE(MinimalAPI)
+class UModelingToolExternalDynamicMeshUpdateAPI : public UInterface
+{
+	GENERATED_BODY()
+};
+
+/**
+ * Provides an API to allow external methods to update the mesh or meshes managed by a tool
+ */
+class IModelingToolExternalDynamicMeshUpdateAPI
+{
+	GENERATED_BODY()
+public:
+
+	// @return true if the tool will currently allow an external mesh update to run
+	virtual bool AllowToolMeshUpdates() const { return false; }
+
+	// Update the tool meshes with the provided method, then emit any resulting change transaction(s) and do associated book-keeping required by the tool (updating any spatial data-structures, rendering, etc)
+	virtual void UpdateToolMeshes(TFunctionRef<TUniquePtr<FMeshRegionChangeBase>(UE::Geometry::FDynamicMesh3&, int32 MeshIdx)> UpdateMesh) {}
+
+	// Read the current tool meshes
+	virtual void ProcessToolMeshes(TFunctionRef<void(const UE::Geometry::FDynamicMesh3&, int32 MeshIdx)> UpdateMesh) const {}
+
+	// @return the number of meshes managed by the tool
+	virtual int32 NumToolMeshes() const
+	{
+		return 0;
+	}
+};
\ No newline at end of file
