Add buttons to add and remove sculpt layers

diff --git a/Engine/Plugins/Runtime/MeshModelingToolset/Source/MeshModelingTools/Private/MeshVertexSculptTool.cpp b/Engine/Plugins/Runtime/MeshModelingToolset/Source/MeshModelingTools/Private/MeshVertexSculptTool.cpp
--- a/Engine/Plugins/Runtime/MeshModelingToolset/Source/MeshModelingTools/Private/MeshVertexSculptTool.cpp
+++ b/Engine/Plugins/Runtime/MeshModelingToolset/Source/MeshModelingTools/Private/MeshVertexSculptTool.cpp
@@ -116,19 +116,24 @@ void UMeshSculptLayerProperties::Init(IModelingToolExternalDynamicMeshUpdateAPI*
 			bHasSculptLayers = Mesh.HasAttributes() && Mesh.Attributes()->NumSculptLayers() > NumLockedBaseLayers;
 			if (bHasSculptLayers)
 			{
-				ActiveLayer = Mesh.Attributes()->GetSculptLayers()->GetActiveLayer();
-				TConstArrayView<double> SourceWeights = Mesh.Attributes()->GetSculptLayers()->GetLayerWeights();
-				int32 NumLayers = FMath::Max(0, SourceWeights.Num() - NumLockedBaseLayers);
-				LayerWeights.SetNum(NumLayers);
-				for (int32 Idx = 0; Idx < LayerWeights.Num(); ++Idx)
-				{
-					LayerWeights[Idx] = SourceWeights[Idx + NumLockedBaseLayers];
-				}
+				UpdateSettingsFromMesh(Mesh.Attributes()->GetSculptLayers());
 			}
 		}
 	);
 }
 
+void UMeshSculptLayerProperties::UpdateSettingsFromMesh(const UE::Geometry::FDynamicMeshSculptLayers* SculptLayers)
+{
+	TConstArrayView<double> SourceWeights = SculptLayers->GetLayerWeights();
+	int32 NumLayers = FMath::Max(0, SourceWeights.Num() - NumLockedBaseLayers);
+	LayerWeights.SetNum(NumLayers);
+	for (int32 Idx = 0; Idx < LayerWeights.Num(); ++Idx)
+	{
+		LayerWeights[Idx] = SourceWeights[Idx + NumLockedBaseLayers];
+	}
+	ActiveLayer = SculptLayers->GetActiveLayer();
+}
+
 void UMeshSculptLayerProperties::SetLayerWeights(FDynamicMeshSculptLayers* SculptLayers) const
 {
 	TArray<double> FullLayerWeights(SculptLayers->GetLayerWeights());
@@ -140,14 +145,11 @@ void UMeshSculptLayerProperties::SetLayerWeights(FDynamicMeshSculptLayers* Sculp
 	SculptLayers->UpdateLayerWeights(FullLayerWeights);
 }
 
-#if WITH_EDITOR
-void UMeshSculptLayerProperties::PostEditChangeProperty(FPropertyChangedEvent& PropertyChangedEvent)
+void UMeshSculptLayerProperties::EditSculptLayers(TFunctionRef<void(FDynamicMesh3& Mesh, FDynamicMeshSculptLayers* SculptLayers)> EditFn, bool bEmitChange)
 {
-	Super::PostEditChangeProperty(PropertyChangedEvent);
-
 	if (Tool && Tool->AllowToolMeshUpdates())
 	{
-		Tool->UpdateToolMeshes([this, &PropertyChangedEvent](FDynamicMesh3& Mesh, int32 Idx) -> TUniquePtr<FMeshRegionChangeBase>
+		Tool->UpdateToolMeshes([this, &EditFn, bEmitChange](FDynamicMesh3& Mesh, int32 Idx) -> TUniquePtr<FMeshRegionChangeBase>
 			{
 				// Sculpt layer UI only supports a single mesh for now
 				if (Idx > 0)
@@ -163,21 +165,11 @@ void UMeshSculptLayerProperties::PostEditChangeProperty(FPropertyChangedEvent& P
 				{
 					if (FDynamicMeshSculptLayers* SculptLayers = Attributes->GetSculptLayers())
 					{
-						if (PropertyChangedEvent.GetMemberPropertyName() == GET_MEMBER_NAME_CHECKED(UMeshSculptLayerProperties, LayerWeights))
-						{
-							SetLayerWeights(SculptLayers);
-						}
-						if (PropertyChangedEvent.GetMemberPropertyName() == GET_MEMBER_NAME_CHECKED(UMeshSculptLayerProperties, ActiveLayer))
-						{
-							ActiveLayer = FMath::Clamp(ActiveLayer, NumLockedBaseLayers, SculptLayers->NumLayers() - 1);
-							SculptLayers->SetActiveLayer(ActiveLayer);
-						}
+						EditFn(Mesh, SculptLayers);
 					}
 				}
 
-				
-
-				if (PropertyChangedEvent.ChangeType == EPropertyChangeType::Interactive)
+				if (!bEmitChange)
 				{
 					return nullptr;
 				}
@@ -188,8 +180,75 @@ void UMeshSculptLayerProperties::PostEditChangeProperty(FPropertyChangedEvent& P
 		);
 	}
 }
+
+#if WITH_EDITOR
+void UMeshSculptLayerProperties::PostEditChangeProperty(FPropertyChangedEvent& PropertyChangedEvent)
+{
+	Super::PostEditChangeProperty(PropertyChangedEvent);
+
+	EditSculptLayers([this, &PropertyChangedEvent](FDynamicMesh3& Mesh, FDynamicMeshSculptLayers* SculptLayers)
+		{
+			if (PropertyChangedEvent.GetMemberPropertyName() == GET_MEMBER_NAME_CHECKED(UMeshSculptLayerProperties, LayerWeights))
+			{
+				SetLayerWeights(SculptLayers);
+			}
+			if (PropertyChangedEvent.GetMemberPropertyName() == GET_MEMBER_NAME_CHECKED(UMeshSculptLayerProperties, ActiveLayer))
+			{
+				ActiveLayer = FMath::Clamp(ActiveLayer, NumLockedBaseLayers, SculptLayers->NumLayers() - 1);
+				SculptLayers->SetActiveLayer(ActiveLayer);
+			}
+		}, 
+		PropertyChangedEvent.ChangeType != EPropertyChangeType::Interactive
+	);
+}
 #endif
 
+void UMeshSculptLayerProperties::AddLayer()
+{
+	EditSculptLayers([this](FDynamicMesh3& Mesh, FDynamicMeshSculptLayers* SculptLayers)
+		{
+			Mesh.Attributes()->EnableSculptLayers(SculptLayers->NumLayers() + 1);
+			LayerWeights.Add(1.0);
+		}, true);
+}
+void UMeshSculptLayerProperties::RemoveLayer()
+{
+	if (LayerWeights.Num() <= 1)
+	{
+		return;
+	}
+	EditSculptLayers([this](FDynamicMesh3& Mesh, FDynamicMeshSculptLayers* SculptLayers)
+		{
+			if (!SculptLayers->DiscardSculptLayer(ActiveLayer))
+			{
+				return;
+			}
+			UpdateSettingsFromMesh(SculptLayers);
+			// If the system picked a locked active layer, try to pick a different layer instead
+			if (ActiveLayer < NumLockedBaseLayers)
+			{
+				// We shouldn't set a zero-weight layer as the active layer, so look for a non-zero weight layer to set
+				int32 NonZeroLayerIdx = INDEX_NONE;
+				for (int32 Idx = 0; Idx < LayerWeights.Num(); ++Idx)
+				{
+					if (LayerWeights[Idx] != 0.0)
+					{
+						NonZeroLayerIdx = Idx;
+						break;
+					}
+				}
+				// if all layers had zero weight, just pick the first layer and set its weight to 1.0 so it is ready to sculpt on
+				if (NonZeroLayerIdx == INDEX_NONE)
+				{
+					LayerWeights[0] = 1.0;
+					NonZeroLayerIdx = 0;
+					SetLayerWeights(SculptLayers);
+				}
+				ActiveLayer = SculptLayers->SetActiveLayer(NumLockedBaseLayers + NonZeroLayerIdx);
+			}
+		}, true);
+}
+
 /*
  * Tool
  */
diff --git a/Engine/Plugins/Runtime/MeshModelingToolset/Source/MeshModelingTools/Public/MeshVertexSculptTool.h b/Engine/Plugins/Runtime/MeshModelingToolset/Source/MeshModelingTools/Public/MeshVertexSculptTool.h
--- a/Engine/Plugins/Runtime/MeshModelingToolset/Source/MeshModelingTools/Public/MeshVertexSculptTool.h
+++ b/Engine/Plugins/Runtime/MeshModelingToolset/Source/MeshModelingTools/Public/MeshVertexSculptTool.h
@@ -221,12 +221,18 @@ public:
 	int32 ActiveLayer = 0;
 
 	/** Set the sculpt layer weights */
-	UPROPERTY(EditAnywhere, Category = SculptLayers, meta = (HideEditConditionToggle, EditCondition = bHasSculptLayers, ModelingQuickSettings))
+	UPROPERTY(EditAnywhere, EditFixedSize, NoClear, Category = SculptLayers, meta = (NoResetToDefault, HideEditConditionToggle, EditCondition = bHasSculptLayers, ModelingQuickSettings))
 	TArray<double> LayerWeights;
 
 	UPROPERTY(meta = (TransientToolProperty))
 	bool bHasSculptLayers = false;
 
+	UFUNCTION(CallInEditor, Category = SculptLayers)
+	void AddLayer();
+
+	UFUNCTION(CallInEditor, Category = SculptLayers)
+	void RemoveLayer();
+
 	MESHMODELINGTOOLS_API void Init(IModelingToolExternalDynamicMeshUpdateAPI* InTool, int32 InNumLockedBaseLayers);
 
 #if WITH_EDITOR
@@ -244,6 +250,14 @@ private:
 
 	// Helper to set sculpt layers from the current LayerWeights property (accounting for the NumLockedBaseLayers)
 	void SetLayerWeights(UE::Geometry::FDynamicMeshSculptLayers* SculptLayers) const;
+
+	// Update the ActiveLayer and LayerWeights settings from the current sculpt layers
+	void UpdateSettingsFromMesh(const UE::Geometry::FDynamicMeshSculptLayers* SculptLayers);
+
+	// Helper to apply edits to the current sculpt layers if possible, with associated book keeping
+	// @param EditFn The edit to apply if possible
+	// @param bEmitChange Whether to emit a change object along with the edit
+	void EditSculptLayers(TFunctionRef<void(UE::Geometry::FDynamicMesh3& Mesh, UE::Geometry::FDynamicMeshSculptLayers* SculptLayers)> EditFn, bool bEmitChange);
 };
 
 
diff --git a/Engine/Source/Runtime/GeometryCore/Private/DynamicMesh/DynamicMeshSculptLayers.cpp b/Engine/Source/Runtime/GeometryCore/Private/DynamicMesh/DynamicMeshSculptLayers.cpp
--- a/Engine/Source/Runtime/GeometryCore/Private/DynamicMesh/DynamicMeshSculptLayers.cpp
+++ b/Engine/Source/Runtime/GeometryCore/Private/DynamicMesh/DynamicMeshSculptLayers.cpp
@@ -254,7 +254,7 @@ bool FDynamicMeshSculptLayers::UpdateLayersFromMesh()
 
 bool FDynamicMeshSculptLayers::UpdateMeshFromLayers()
 {
-	if (Layers.IsEmpty() || !ensure(HasValidLayers()))
+	if (Layers.IsEmpty() || !ensure(Layers.Num() == LayerWeights.Num()))
 	{
 		return false;
 	}
@@ -299,7 +299,7 @@ bool FDynamicMeshSculptLayers::ValidateActiveLayer()
 		}
 		for (int32 Idx = ActiveLayer - 1; Idx >= 0; --Idx)
 		{
-			if (LayerWeights[ActiveLayer] != 0)
+			if (LayerWeights[Idx] != 0)
 			{
 				ActiveLayer = Idx;
 				return true;
