Add erase sculpt layer brush

diff --git a/Engine/Plugins/Editor/ModelingToolsEditorMode/Source/ModelingToolsEditorMode/Private/DetailsCustomizations/MeshVertexSculptToolCustomizations.cpp b/Engine/Plugins/Editor/ModelingToolsEditorMode/Source/ModelingToolsEditorMode/Private/DetailsCustomizations/MeshVertexSculptToolCustomizations.cpp
--- a/Engine/Plugins/Editor/ModelingToolsEditorMode/Source/ModelingToolsEditorMode/Private/DetailsCustomizations/MeshVertexSculptToolCustomizations.cpp
+++ b/Engine/Plugins/Editor/ModelingToolsEditorMode/Source/ModelingToolsEditorMode/Private/DetailsCustomizations/MeshVertexSculptToolCustomizations.cpp
@@ -27,6 +27,7 @@
 #include "Sculpting/MeshPlaneBrushOps.h"
 #include "Sculpting/MeshSculptBrushOps.h"
 #include "Sculpting/MeshSmoothingBrushOps.h"
+#include "Sculpting/MeshEraseSculptLayerBrushOps.h"
 
 using namespace UE::ModelingUI;
 
@@ -225,6 +226,7 @@ template class FMeshSculptBrushOpPropertiesDetails<UScaleKelvinletBrushOpProps>;
 template class FMeshSculptBrushOpPropertiesDetails<UPullKelvinletBrushOpProps>;
 template class FMeshSculptBrushOpPropertiesDetails<USharpPullKelvinletBrushOpProps>;
 template class FMeshSculptBrushOpPropertiesDetails<UTwistKelvinletBrushOpProps>;
+template class FMeshSculptBrushOpPropertiesDetails<UEraseSculptLayerBrushOpProps>;
 
 /***********************************************************/
 
diff --git a/Engine/Plugins/Editor/ModelingToolsEditorMode/Source/ModelingToolsEditorMode/Private/ModelingToolsEditorModeStyle.cpp b/Engine/Plugins/Editor/ModelingToolsEditorMode/Source/ModelingToolsEditorMode/Private/ModelingToolsEditorModeStyle.cpp
--- a/Engine/Plugins/Editor/ModelingToolsEditorMode/Source/ModelingToolsEditorMode/Private/ModelingToolsEditorModeStyle.cpp
+++ b/Engine/Plugins/Editor/ModelingToolsEditorMode/Source/ModelingToolsEditorMode/Private/ModelingToolsEditorModeStyle.cpp
@@ -465,6 +465,7 @@ void FModelingToolsEditorModeStyle::Initialize()
 		StyleSet->Set("BrushTypeIcons.Grab", new IMAGE_BRUSH_SVG("Icons/BrushIcons/Brush_Grab", Icon120));
 		StyleSet->Set("BrushTypeIcons.GrabSharp", new IMAGE_BRUSH_SVG("Icons/BrushIcons/Brush_GrabSharp", Icon120));
 		StyleSet->Set("BrushTypeIcons.Twist", new IMAGE_BRUSH_SVG("Icons/BrushIcons/Brush_Twist", Icon120));
+		StyleSet->Set("BrushTypeIcons.EraseSculptLayer", new FSlateImageBrush(StyleSet->RootToCoreContentDir(TEXT("../Editor/Slate/Icons/Landscape/Landscape_Erase_x40.png")), Icon20x20));
 		// Height brushes (currently using same icons as corresponding landscape brushes)
 		StyleSet->Set("BrushTypeIcons.HeightSculpt", new FSlateImageBrush(StyleSet->RootToCoreContentDir(TEXT("../Editor/Slate/Icons/Landscape/Landscape_Sculpt_x40.png")), Icon20x20));
 		StyleSet->Set("BrushTypeIcons.HeightSmooth", new FSlateImageBrush(StyleSet->RootToCoreContentDir(TEXT("../Editor/Slate/Icons/Landscape/Landscape_Smooth_x40.png")), Icon20x20));
diff --git a/Engine/Plugins/Runtime/MeshModelingToolset/Source/MeshModelingTools/Private/MeshVertexSculptTool.cpp b/Engine/Plugins/Runtime/MeshModelingToolset/Source/MeshModelingTools/Private/MeshVertexSculptTool.cpp
--- a/Engine/Plugins/Runtime/MeshModelingToolset/Source/MeshModelingTools/Private/MeshVertexSculptTool.cpp
+++ b/Engine/Plugins/Runtime/MeshModelingToolset/Source/MeshModelingTools/Private/MeshVertexSculptTool.cpp
@@ -35,6 +35,7 @@
 #include "Sculpting/MeshPlaneBrushOps.h"
 #include "Sculpting/MeshPinchBrushOps.h"
 #include "Sculpting/MeshSculptBrushOps.h"
+#include "Sculpting/MeshEraseSculptLayerBrushOps.h"
 #include "Sculpting/StampFalloffs.h"
 #include "Sculpting/MeshSculptUtil.h"
 #include "TargetInterfaces/DynamicMeshCommitter.h"
@@ -531,6 +532,19 @@ void UMeshVertexSculptTool::RegisterBrushes()
 		MakeUnique<TBasicMeshSculptBrushOpFactory<FTwistKelvinletBrushOp>>(),
 		NewObject<UTwistKelvinletBrushOpProps>(this));
 
+	if (ISceneComponentBackedTarget* SceneComponentTarget = Cast<ISceneComponentBackedTarget>(Target))
+	{
+		if (IMeshSculptLayersManager* SculptLayersManager = Cast<IMeshSculptLayersManager>(SceneComponentTarget->GetOwnerSceneComponent()))
+		{
+			if (SculptLayersManager->HasSculptLayers())
+			{
+				RegisterBrushType((int32)EMeshVertexSculptBrushType::EraseSculptLayer, LOCTEXT("EraseSculptLayer", "EraseSculptLayer"),
+					MakeUnique<TBasicMeshSculptBrushOpFactory<FEraseSculptLayerBrushOp>>(),
+					NewObject<UEraseSculptLayerBrushOpProps>(this));
+			}
+		}
+	}
+
 	// secondary brushes
 	// We activate ID 0 as our default secondary brush, so use that as the registration ID
 	RegisterSecondaryBrushType(0, LOCTEXT("Smooth", "Smooth"),
diff --git a/Engine/Plugins/Runtime/MeshModelingToolset/Source/MeshModelingTools/Public/MeshVertexSculptTool.h b/Engine/Plugins/Runtime/MeshModelingToolset/Source/MeshModelingTools/Public/MeshVertexSculptTool.h
--- a/Engine/Plugins/Runtime/MeshModelingToolset/Source/MeshModelingTools/Public/MeshVertexSculptTool.h
+++ b/Engine/Plugins/Runtime/MeshModelingToolset/Source/MeshModelingTools/Public/MeshVertexSculptTool.h
@@ -103,6 +103,9 @@ enum class EMeshVertexSculptBrushType : uint8
 	/** Move vertices towards a fixed plane in world space, positioned with a 3D gizmo */
 	FixedPlane UMETA(DisplayName = "FixedPlane"),
 
+	/** Erase sculpt layers */
+	EraseSculptLayer UMETA(DisplayName = "Erase Sculpt Layer"),
+
 	LastValue UMETA(Hidden)
 
 };
diff --git a/Engine/Plugins/Runtime/MeshModelingToolset/Source/MeshModelingTools/Public/Sculpting/MeshEraseSculptLayerBrushOps.h b/Engine/Plugins/Runtime/MeshModelingToolset/Source/MeshModelingTools/Public/Sculpting/MeshEraseSculptLayerBrushOps.h
new file mode 100644
--- /dev/null
+++ b/Engine/Plugins/Runtime/MeshModelingToolset/Source/MeshModelingTools/Public/Sculpting/MeshEraseSculptLayerBrushOps.h
@@ -0,0 +1,99 @@
+// Copyright Epic Games, Inc. All Rights Reserved.
+
+#pragma once
+
+#include "Sculpting/MeshBrushOpBase.h"
+#include "DynamicMesh/DynamicMeshAttributeSet.h"
+#include "DynamicMesh/DynamicMeshSculptLayers.h"
+#include "Async/ParallelFor.h"
+
+#include "MeshEraseSculptLayerBrushOps.generated.h"
+
+UCLASS()
+class MESHMODELINGTOOLS_API UEraseSculptLayerBrushOpProps : public UMeshSculptBrushOpProps
+{
+	GENERATED_BODY()
+
+public:
+	/** Strength of the Brush */
+	UPROPERTY(EditAnywhere, Category = EraseSculptLayerBrush, meta = (DisplayName = "Strength", UIMin = "0.0", UIMax = "10.", ClampMin = "0.0", ClampMax = "10.", ModelingQuickSettings))
+	float Strength = 0.5;
+
+	/** Amount of falloff to apply */
+	UPROPERTY(EditAnywhere, Category = EraseSculptLayerBrush, meta = (DisplayName = "Falloff", UIMin = "0.0", UIMax = "1.0", ClampMin = "0.0", ClampMax = "1.0"))
+	float Falloff = 0.5;
+
+
+	virtual float GetStrength() override { return Strength; }
+	virtual float GetFalloff() override { return Falloff; }
+	virtual bool SupportsStrengthPressure() override { return true; }
+};
+
+
+class FEraseSculptLayerBrushOp : public FMeshSculptBrushOp
+{
+public:
+
+	virtual void ApplyStamp(const FDynamicMesh3* Mesh, const FSculptBrushStamp& Stamp, const TArray<int32>& Vertices, TArray<FVector3d>& NewPositionsOut) override
+	{
+		static const double BrushSpeedTuning = 1.0;
+		double UsePower = Stamp.Power * Stamp.Radius * Stamp.DeltaTime * BrushSpeedTuning;
+		double MaxOffset = Stamp.Radius;
+		if (!Mesh->HasAttributes())
+		{
+			return;
+		}
+
+		const UE::Geometry::FDynamicMeshSculptLayers* SculptLayers = Mesh->Attributes()->GetSculptLayers();
+		if (!SculptLayers)
+		{
+			return;
+		}
+
+		int32 ActiveLayer = SculptLayers->GetActiveLayer();
+		TConstArrayView<double> LayerWeights = SculptLayers->GetLayerWeights();
+		
+		ParallelFor(Vertices.Num(), [this, &Mesh, &Stamp, &Vertices, &NewPositionsOut, UsePower, MaxOffset, SculptLayers, ActiveLayer, LayerWeights](int32 k)
+		{
+			int32 VertIdx = Vertices[k];
+			FVector3d OrigPos = Mesh->GetVertex(VertIdx);
+
+			FVector3d BasePos(0);
+			for (int32 LayerIdx = 0, NumLayers = SculptLayers->NumLayers(); LayerIdx < NumLayers; ++LayerIdx)
+			{
+				if (LayerIdx != ActiveLayer)
+				{
+					FVector3d Offset;
+					SculptLayers->GetLayer(LayerIdx)->GetValue(VertIdx, Offset);
+					BasePos += Offset * LayerWeights[LayerIdx];
+				}
+			}
+		
+			FVector3d MoveVec = (BasePos - OrigPos);
+			double Falloff = GetFalloff().Evaluate(Stamp, OrigPos);
+			double MoveDist = Falloff * UsePower;
+			if (MoveVec.SquaredLength() < MoveDist * MoveDist)
+			{
+				NewPositionsOut[k] = BasePos;
+			}
+			else
+			{
+				UE::Geometry::Normalize(MoveVec);
+				NewPositionsOut[k] = OrigPos + MoveDist * MoveVec;
+			}
+		});
+	}
+
+
+	virtual ESculptBrushOpTargetType GetBrushTargetType() const override
+	{
+		return ESculptBrushOpTargetType::SculptMesh;
+	}
+
+
+	virtual bool IgnoreZeroMovements() const
+	{
+		return false;
+	}
+};
+
