Tiny refactor, remove MeshToWorld from MeshBasedModifierBase + constness

diff --git a/Engine/Plugins/Runtime/MeshModelingToolset/Source/ModelingComponentsEditorOnly/Private/ToolTargets/StaticMeshToolTarget.cpp b/Engine/Plugins/Runtime/MeshModelingToolset/Source/ModelingComponentsEditorOnly/Private/ToolTargets/StaticMeshToolTarget.cpp
--- a/Engine/Plugins/Runtime/MeshModelingToolset/Source/ModelingComponentsEditorOnly/Private/ToolTargets/StaticMeshToolTarget.cpp
+++ b/Engine/Plugins/Runtime/MeshModelingToolset/Source/ModelingComponentsEditorOnly/Private/ToolTargets/StaticMeshToolTarget.cpp
@@ -305,7 +305,7 @@ TArray<int32> UStaticMeshReadOnlyToolTarget::MapSectionToMaterialID(const UStati
 	return UE::Conversion::GetPolygonGroupToMaterialIndexMap(Mesh, LODType, UseLOD);
 }
 
-const FMeshDescription* UStaticMeshReadOnlyToolTarget::GetMeshDescriptionWithScaleApplied(UStaticMesh* StaticMesh, int32 UseLOD, FMeshDescriptionCache& CachedMeshDescriptions)
+const FMeshDescription* UStaticMeshReadOnlyToolTarget::GetMeshDescriptionWithScaleApplied(const UStaticMesh* StaticMesh, int32 UseLOD, FMeshDescriptionCache& CachedMeshDescriptions)
 {
 	static FMeshDescription EmptyMeshDescription;
 	static bool bFirst = true;
diff --git a/Engine/Plugins/Runtime/MeshModelingToolset/Source/ModelingComponentsEditorOnly/Public/ToolTargets/StaticMeshToolTarget.h b/Engine/Plugins/Runtime/MeshModelingToolset/Source/ModelingComponentsEditorOnly/Public/ToolTargets/StaticMeshToolTarget.h
--- a/Engine/Plugins/Runtime/MeshModelingToolset/Source/ModelingComponentsEditorOnly/Public/ToolTargets/StaticMeshToolTarget.h
+++ b/Engine/Plugins/Runtime/MeshModelingToolset/Source/ModelingComponentsEditorOnly/Public/ToolTargets/StaticMeshToolTarget.h
@@ -84,7 +84,7 @@ public:
 	};
 	using FMeshDescriptionCache = TMap<int32, FCachedMeshDescription>;
 	// helper to get the mesh description with build scale applied, using/updating the cache as needed
-	static const FMeshDescription* GetMeshDescriptionWithScaleApplied(UStaticMesh* Mesh, int32 UseLOD, FMeshDescriptionCache& Cache);
+	static const FMeshDescription* GetMeshDescriptionWithScaleApplied(const UStaticMesh* Mesh, int32 UseLOD, FMeshDescriptionCache& Cache);
 
 protected:
 	TWeakObjectPtr<UStaticMesh> StaticMesh = nullptr;
