note that the bTransact flag in the creation api also enabled RF_Transactional on the created component, so it can be included in subsequent undo/redo transactions as well

diff --git a/Engine/Plugins/Runtime/MeshModelingToolset/Source/ModelingComponents/Public/ModelingObjectsCreationAPI.h b/Engine/Plugins/Runtime/MeshModelingToolset/Source/ModelingComponents/Public/ModelingObjectsCreationAPI.h
--- a/Engine/Plugins/Runtime/MeshModelingToolset/Source/ModelingComponents/Public/ModelingObjectsCreationAPI.h
+++ b/Engine/Plugins/Runtime/MeshModelingToolset/Source/ModelingComponents/Public/ModelingObjectsCreationAPI.h
@@ -473,7 +473,7 @@ struct MODELINGCOMPONENTS_API FCreateComponentParams
 	UPROPERTY(Category = "CreateComponentParams", EditAnywhere)
 	bool bSetAsRoot = false;
 
-	/** If true, wrap the creation operation in a transaction. */
+	/** If true, wrap the creation operation in a transaction, and create the component with the RF_Transactional object flag set. */
 	UPROPERTY(Category = "CreateComponentParams", EditAnywhere)
 	bool bTransact = false;
 };
