Support 32-bit textures in modelling tools texture import.

diff --git a/Engine/Plugins/Runtime/MeshModelingToolset/Source/ModelingComponents/Private/AssetUtils/Texture2DUtil.cpp b/Engine/Plugins/Runtime/MeshModelingToolset/Source/ModelingComponents/Private/AssetUtils/Texture2DUtil.cpp
--- a/Engine/Plugins/Runtime/MeshModelingToolset/Source/ModelingComponents/Private/AssetUtils/Texture2DUtil.cpp
+++ b/Engine/Plugins/Runtime/MeshModelingToolset/Source/ModelingComponents/Private/AssetUtils/Texture2DUtil.cpp
@@ -141,14 +141,14 @@ static bool ReadTexture_SourceData(
 	const int32 Height = TextureSource.GetSizeY();
 	const FImageDimensions Dimensions = FImageDimensions(Width, Height);
 	DestImage.SetDimensions(Dimensions);
-	const int64 Num = Dimensions.Num();
+	const int64 Num = Dimensions.Num(); //< number of pixels
 
 	TArray64<uint8> SourceData;
 	verify( TextureMap->Source.GetMipData(SourceData, 0, 0, 0) );
 	const ETextureSourceFormat SourceFormat = TextureSource.GetFormat();
 	const int32 BytesPerPixel = TextureSource.GetBytesPerPixel();
 	const uint8* SourceDataPtr = SourceData.GetData();
-
+		
 	// code below is derived from UBlueprintMaterialTextureNodesBPLibrary::Texture2D_SampleUV_EditorOnly()
 	if (SourceFormat == TSF_BGRA8)
 	{
@@ -164,6 +164,7 @@ static bool ReadTexture_SourceData(
 
 			DestImage.SetPixel(i, ToVector4<float>(FloatColor));
 		}
+		return true;
 	}
 	else if (SourceFormat == TSF_BGRE8)
 	{
@@ -177,6 +178,7 @@ static bool ReadTexture_SourceData(
 
 			DestImage.SetPixel(i, ToVector4<float>(FloatColor));
 		}
+		return true;
 	}
 	// code below is also derived from FImage::CopyTo (CopyImage) - invoked during BuildTexture on import.
 	else if (SourceFormat == TSF_RGBA16)
@@ -193,6 +195,7 @@ static bool ReadTexture_SourceData(
 				PixelPtr[3] / 65535.0f
 				));
 		}
+		return true;
 	}
 	else if (SourceFormat == TSF_RGBA16F)
 	{
@@ -203,6 +206,28 @@ static bool ReadTexture_SourceData(
 			FLinearColor LinearColor = ((const FFloat16Color*)PixelPtr)->GetFloats();
 			DestImage.SetPixel(i, ToVector4<float>(LinearColor) );
 		}
+		return true;
+	}
+	else if (SourceFormat == TSF_RGBA32F)
+	{
+		check(BytesPerPixel == sizeof(FVector4f));
+		for (int64 i = 0; i < Num; ++i)
+		{
+			const uint8* PixelPtr = SourceDataPtr + (i * BytesPerPixel);
+			DestImage.SetPixel(i, *reinterpret_cast<const FVector4f*>(PixelPtr));
+		}
+		return true;
+	}
+	else if (SourceFormat == TSF_R32F)
+	{
+		check(BytesPerPixel == sizeof(float));
+		for (int64 i = 0; i < Num; ++i)
+		{
+			const uint8* PixelPtr = SourceDataPtr + (i * BytesPerPixel);
+			const float Value = *reinterpret_cast<const float*>(PixelPtr);
+			DestImage.SetPixel(i, FVector4f(Value, Value, Value, 1.f));
+		}
+		return true;
 	}
 	else if ((SourceFormat == TSF_G16))
 	{
@@ -214,6 +239,7 @@ static bool ReadTexture_SourceData(
 			const float FloatColor = float(ByteColor) / 65535.0f;
 			DestImage.SetPixel(i, FVector4f(FloatColor, FloatColor, FloatColor, 1.0));
 		}
+		return true;
 	}
 	else if (SourceFormat == TSF_G8)
 	{
@@ -228,9 +254,12 @@ static bool ReadTexture_SourceData(
 				FLinearColor(PixelColorf, PixelColorf, PixelColorf, 1.0);
 			DestImage.SetPixel(i, ToVector4<float>(FloatColor));
 		}
+		return true;
 	}
 
-	return true;
+	ensureAlwaysMsgf(false, TEXT("Texture source format not handled."));
+
+	return false;
 }
 #endif 
 
