ModelingTools: Fixed Displace tool to default to preferring source data when available.
Added a bUsePlatformData advanced setting if users prefer to exclusively displace using platform data.

diff --git a/Engine/Plugins/Runtime/MeshModelingToolset/Source/MeshModelingTools/Private/DisplaceMeshTool.cpp b/Engine/Plugins/Runtime/MeshModelingToolset/Source/MeshModelingTools/Private/DisplaceMeshTool.cpp
--- a/Engine/Plugins/Runtime/MeshModelingToolset/Source/MeshModelingTools/Private/DisplaceMeshTool.cpp
+++ b/Engine/Plugins/Runtime/MeshModelingToolset/Source/MeshModelingTools/Private/DisplaceMeshTool.cpp
@@ -858,9 +858,7 @@ namespace DisplaceMeshToolLocals{
 		}
 
 		TImageBuilder<FVector4f> DisplacementMapValues;
-		if (!UE::AssetUtils::ReadTexture(Parameters.DisplacementMap, DisplacementMapValues,
-			// need bPreferPlatformData to be true to respond to non-destructive changes to the texture in the editor
-			true)) 
+		if (!UE::AssetUtils::ReadTexture(Parameters.DisplacementMap, DisplacementMapValues, false)) 
 		{
 			Parameters.DisplaceField = FSampledScalarField2f();
 			Parameters.DisplaceField.GridValues.AssignAll(0);
