Mesh Topology Selection Mechanic: reset DrawnTriangleSetComponent on shutdown
- the TObjectPtr ref count on DrawnTriangleSetComponent is preventing GC: setting DrawnTriangleSetComponent to nullptr after the actor deletion resets the ref count to 0

#rnx

diff --git a/Engine/Plugins/Runtime/MeshModelingToolset/Source/ModelingComponents/Private/Selection/MeshTopologySelectionMechanic.cpp b/Engine/Plugins/Runtime/MeshModelingToolset/Source/ModelingComponents/Private/Selection/MeshTopologySelectionMechanic.cpp
--- a/Engine/Plugins/Runtime/MeshModelingToolset/Source/ModelingComponents/Private/Selection/MeshTopologySelectionMechanic.cpp
+++ b/Engine/Plugins/Runtime/MeshModelingToolset/Source/ModelingComponents/Private/Selection/MeshTopologySelectionMechanic.cpp
@@ -113,6 +113,7 @@ void UMeshTopologySelectionMechanic::Shutdown()
 	{
 		PreviewGeometryActor->Destroy();
 		PreviewGeometryActor = nullptr;
+		DrawnTriangleSetComponent = nullptr;
 	}
 }
 
