MeshModelingToolset: Do not load higher quality default vertex color material if debug view modes are not enabled

diff --git a/Engine/Plugins/Runtime/MeshModelingToolset/Source/ModelingComponents/Private/ModelingComponentsModule.cpp b/Engine/Plugins/Runtime/MeshModelingToolset/Source/ModelingComponents/Private/ModelingComponentsModule.cpp
--- a/Engine/Plugins/Runtime/MeshModelingToolset/Source/ModelingComponents/Private/ModelingComponentsModule.cpp
+++ b/Engine/Plugins/Runtime/MeshModelingToolset/Source/ModelingComponents/Private/ModelingComponentsModule.cpp
@@ -4,6 +4,7 @@
 #include "Components/BaseDynamicMeshComponent.h"
 #include "Materials/Material.h"
 #include "Misc/CoreDelegates.h"
+#include "ShaderCore.h"
 
 #define LOCTEXT_NAMESPACE "FModelingComponentsModule"
 
@@ -17,13 +18,16 @@ void FModelingComponentsModule::StartupModule()
 
 void FModelingComponentsModule::OnPostEngineInit()
 {
-	// Replace the standard UBaseDynamicMeshComponent vertex color material with something higher quality.
-	// This is done in ModelingComponents module (ie part of MeshModelingToolset plugin) to avoid having to
-	// make this Material a special "engine material", which has various undesirable implication
-	UMaterial* VertexColorMaterial = LoadObject<UMaterial>(nullptr, TEXT("/MeshModelingToolset/Materials/M_DynamicMeshComponentVtxColor"));
-	if (ensure(VertexColorMaterial))
+	if (AllowDebugViewmodes())
 	{
-		UBaseDynamicMeshComponent::SetDefaultVertexColorMaterial(VertexColorMaterial);
+		// Replace the standard UBaseDynamicMeshComponent vertex color material with something higher quality.
+		// This is done in ModelingComponents module (ie part of MeshModelingToolset plugin) to avoid having to
+		// make this Material a special "engine material", which has various undesirable implication
+		UMaterial* VertexColorMaterial = LoadObject<UMaterial>(nullptr, TEXT("/MeshModelingToolset/Materials/M_DynamicMeshComponentVtxColor"));
+		if (ensure(VertexColorMaterial))
+		{
+			UBaseDynamicMeshComponent::SetDefaultVertexColorMaterial(VertexColorMaterial);
+		}
 	}
 }
 
