Lattice Deformer Tool: add Mesh Sculpt Layer Properties and inherit from IModelingToolExternalDynamicMeshUpdateAPI. Together this allows the lattice deformer tool to be used from the Sculpt Layers Modifier

diff --git a/Engine/Plugins/Runtime/MeshModelingToolset/Source/MeshModelingTools/Private/LatticeDeformerTool.cpp b/Engine/Plugins/Runtime/MeshModelingToolset/Source/MeshModelingTools/Private/LatticeDeformerTool.cpp
--- a/Engine/Plugins/Runtime/MeshModelingToolset/Source/MeshModelingTools/Private/LatticeDeformerTool.cpp
+++ b/Engine/Plugins/Runtime/MeshModelingToolset/Source/MeshModelingTools/Private/LatticeDeformerTool.cpp
@@ -24,9 +24,13 @@
 #include "TargetInterfaces/MeshDescriptionCommitter.h"
 #include "TargetInterfaces/MeshDescriptionProvider.h"
 #include "TargetInterfaces/PrimitiveComponentBackedTarget.h"
+#include "TargetInterfaces/DynamicMeshCommitter.h"
+#include "TargetInterfaces/DynamicMeshProvider.h"
 #include "ModelingToolTargetUtil.h"
 #include "ToolTargetManager.h"
 #include "Selection/StoredMeshSelectionUtil.h"
+#include "Properties/MeshSculptLayerProperties.h"
+#include "MeshSculptLayersManagerAPI.h"
 
 #include UE_INLINE_GENERATED_CPP_BY_NAME(LatticeDeformerTool)
 
@@ -124,6 +128,17 @@ bool ULatticeDeformerToolBuilder::CanBuildTool(const FToolBuilderState& SceneSta
 	return SceneState.TargetManager->CountSelectedAndTargetable(SceneState, GetTargetRequirements()) == 1;
 }
 
+const FToolTargetTypeRequirements& ULatticeDeformerToolBuilder::GetTargetRequirements() const
+{
+	static FToolTargetTypeRequirements TypeRequirements({
+	UMaterialProvider::StaticClass(),
+	UDynamicMeshProvider::StaticClass(),
+	UDynamicMeshCommitter::StaticClass(),
+	USceneComponentBackedTarget::StaticClass()
+		});
+	return TypeRequirements;
+}
+
 
 // Operator factory
 
@@ -196,7 +211,7 @@ void ULatticeDeformerTool::InitializeLattice(TArray<FVector3d>& OutLatticePoints
 	Lattice->GenerateInitialLatticePositions(OutLatticePoints);
 
 	// Put the lattice in world space
-	FTransform3d LocalToWorld(Cast<IPrimitiveComponentBackedTarget>(Targets[0])->GetWorldTransform());
+	FTransform3d LocalToWorld(Cast<ISceneComponentBackedTarget>(Targets[0])->GetWorldTransform());
 	Algo::ForEach(OutLatticePoints, [&LocalToWorld](FVector3d& Point) {
 		Point = LocalToWorld.TransformPosition(Point);
 	});
@@ -229,7 +244,7 @@ void ULatticeDeformerTool::Setup()
 	}
 
 	// Note: Mesh will be implicitly transformed to world space by transforming the lattice; we account for whether that would invert the mesh here
-	MeshTransforms::ReverseOrientationIfNeeded(*OriginalMesh, (Cast<IPrimitiveComponentBackedTarget>(Targets[0])->GetWorldTransform()));
+	MeshTransforms::ReverseOrientationIfNeeded(*OriginalMesh, (Cast<ISceneComponentBackedTarget>(Targets[0])->GetWorldTransform()));
 
 	Settings = NewObject<ULatticeDeformerToolProperties>(this, TEXT("Lattice Deformer Tool Settings"));
 	Settings->Initialize(this);
@@ -274,7 +289,7 @@ void ULatticeDeformerTool::Setup()
 	ControlPointsMechanic = NewObject<ULatticeControlPointsMechanic>(this);
 	ControlPointsMechanic->Setup(this);
 	ControlPointsMechanic->SetWorld(GetTargetWorld());
-	FTransform3d LocalToWorld(Cast<IPrimitiveComponentBackedTarget>(Targets[0])->GetWorldTransform());
+	FTransform3d LocalToWorld(Cast<ISceneComponentBackedTarget>(Targets[0])->GetWorldTransform());
 	WorldTransform = LocalToWorld;
 	ControlPointsMechanic->Initialize(LatticePoints, LatticeEdges, LocalToWorld);
 
@@ -287,6 +302,10 @@ void ULatticeDeformerTool::Setup()
 		ResetConstrainedPoints();
 		Preview->InvalidateResult();
 		Settings->bCanChangeResolution = !ControlPointsMechanic->bHasChanged;
+		if (SculptLayerProperties)
+		{
+			SculptLayerProperties->bCanEditLayers = !ControlPointsMechanic->bHasChanged;
+		}
 	};
 	ControlPointsMechanic->OnPointsChanged.AddLambda(OnPointsChangedLambda);
 
@@ -316,6 +335,21 @@ void ULatticeDeformerTool::Setup()
 
 
 	StartPreview();
+
+
+	if (ISceneComponentBackedTarget* SceneComponentTarget = Cast<ISceneComponentBackedTarget>(Targets[0]))
+	{
+		if (IMeshSculptLayersManager* SculptLayersManager = Cast<IMeshSculptLayersManager>(SceneComponentTarget->GetOwnerSceneComponent()))
+		{
+			if (SculptLayersManager->HasSculptLayers())
+			{
+				SculptLayerProperties = NewObject<UMeshSculptLayerProperties>(this);
+				SculptLayerProperties->Init(this, SculptLayersManager->NumLockedBaseSculptLayers());
+				AddToolPropertySource(SculptLayerProperties);
+			}
+		}
+	}
+
 }
 
 
@@ -414,7 +448,7 @@ void ULatticeDeformerTool::OnShutdown(EToolShutdownType ShutdownType)
 	Settings->SaveProperties(this);
 	ControlPointsMechanic->Shutdown();
 
-	IPrimitiveComponentBackedTarget* TargetComponent = Cast<IPrimitiveComponentBackedTarget>(Targets[0]);
+	ISceneComponentBackedTarget* TargetComponent = Cast<ISceneComponentBackedTarget>(Targets[0]);
 	TargetComponent->SetOwnerVisibility(true);
 
 	if (Preview)
@@ -425,13 +459,32 @@ void ULatticeDeformerTool::OnShutdown(EToolShutdownType ShutdownType)
 		{
 			GetToolManager()->BeginUndoTransaction(LOCTEXT("LatticeDeformerTool", "Lattice Deformer"));
 
-			FDynamicMesh3* DynamicMeshResult = Result.Mesh.Get();
+			FDynamicMesh3* const DynamicMeshResult = Result.Mesh.Get();
 			check(DynamicMeshResult != nullptr);
 
 			// The lattice and its output mesh are in world space, so get them in local space.
 			// TODO: Would it make more sense to do all the lattice computation in local space?
-			FTransform3d LocalToWorld(TargetComponent->GetWorldTransform());
-			MeshTransforms::ApplyTransformInverse(*DynamicMeshResult, LocalToWorld, true);
+			const FTransformSRT3d LocalToWorld(TargetComponent->GetWorldTransform());
+			const bool bHasVertexNormals = DynamicMeshResult->HasVertexNormals();
+			const int32 NumVertices = DynamicMeshResult->MaxVertexID();
+			ParallelFor(NumVertices, [&DynamicMeshResult, &LocalToWorld, bHasVertexNormals](int32 VertexIndex)
+			{
+				if (DynamicMeshResult->IsVertex(VertexIndex))
+				{
+					FVector3d Position = DynamicMeshResult->GetVertex(VertexIndex);
+					Position = LocalToWorld.InverseTransformPosition(Position);
+					DynamicMeshResult->SetVertex(VertexIndex, Position);
+
+					if (bHasVertexNormals)
+					{
+						FVector3f Normal = DynamicMeshResult->GetVertexNormal(VertexIndex);
+						Normal = LocalToWorld.InverseTransformNormal(Normal);
+						DynamicMeshResult->SetVertexNormal(VertexIndex, Normalized(Normal));
+					}
+				}
+			});
+
+
 
 			UE::ToolTarget::CommitDynamicMeshUpdate(Targets[0], *DynamicMeshResult, true);
 
@@ -469,7 +522,7 @@ void ULatticeDeformerTool::StartPreview()
 	Preview->SetVisibility(true);
 	Preview->InvalidateResult();
 
-	Cast<IPrimitiveComponentBackedTarget>(Targets[0])->SetOwnerVisibility(false);
+	Cast<ISceneComponentBackedTarget>(Targets[0])->SetOwnerVisibility(false);
 }
 
 
@@ -505,7 +558,7 @@ void ULatticeDeformerTool::OnTick(float DeltaTime)
 			TArray<FVector3d> LatticePoints;
 			TArray<FVector2i> LatticeEdges;
 			InitializeLattice(LatticePoints, LatticeEdges);
-			FTransform3d LocalToWorld(Cast<IPrimitiveComponentBackedTarget>(Targets[0])->GetWorldTransform());
+			FTransform3d LocalToWorld(Cast<ISceneComponentBackedTarget>(Targets[0])->GetWorldTransform());
 			ControlPointsMechanic->Initialize(LatticePoints, LatticeEdges, LocalToWorld);
 			Preview->InvalidateResult();
 			bShouldRebuild = false;
@@ -575,6 +628,30 @@ void ULatticeDeformerTool::UpdateMechanicColorOverrides()
 	ControlPointsMechanic->UpdateDrawables();
 }
 
+bool ULatticeDeformerTool::AllowToolMeshUpdates() const
+{
+	return !ControlPointsMechanic->IsGizmoBeingDragged() && !ControlPointsMechanic->bHasChanged;
+}
+
+void ULatticeDeformerTool::UpdateToolMeshes(TFunctionRef<TUniquePtr<FMeshRegionChangeBase>(FDynamicMesh3&, int32 MeshIdx)> UpdateMesh)
+{
+	if (AllowToolMeshUpdates())
+	{
+		UpdateMesh(*OriginalMesh, 0);
+		bShouldRebuild = true;
+	}
+}
+
+void ULatticeDeformerTool::ProcessToolMeshes(TFunctionRef<void(const UE::Geometry::FDynamicMesh3&, int32 MeshIdx)> ProcessMesh) const
+{
+	ProcessMesh(*OriginalMesh, 0);
+}
+
+int32 ULatticeDeformerTool::NumToolMeshes() const
+{
+	return 1;
+}
+
 
 
 void FLatticeDeformerToolConstrainedPointsChange::Apply(UObject* Object)
diff --git a/Engine/Plugins/Runtime/MeshModelingToolset/Source/MeshModelingTools/Private/MeshVertexSculptTool.cpp b/Engine/Plugins/Runtime/MeshModelingToolset/Source/MeshModelingTools/Private/MeshVertexSculptTool.cpp
--- a/Engine/Plugins/Runtime/MeshModelingToolset/Source/MeshModelingTools/Private/MeshVertexSculptTool.cpp
+++ b/Engine/Plugins/Runtime/MeshModelingToolset/Source/MeshModelingTools/Private/MeshVertexSculptTool.cpp
@@ -24,6 +24,8 @@
 #include "BaseBehaviors/TwoAxisPropertyEditBehavior.h"
 #include "Generators/RectangleMeshGenerator.h"
 
+#include "Properties/MeshSculptLayerProperties.h"
+
 #include "Changes/MeshVertexChange.h"
 #include "Changes/MeshRegionChange.h"
 
@@ -95,160 +97,6 @@ public:
 	virtual void Revert(UObject* Object) override;
 };
 
-void UMeshSculptLayerProperties::Init(IModelingToolExternalDynamicMeshUpdateAPI* InTool, int32 InNumLockedBaseLayers)
-{
-	bHasSculptLayers = true;
-	LayerWeights.Reset();
-
-	Tool = InTool;
-	if (!Tool)
-	{
-		return;
-	}
-	NumLockedBaseLayers = InNumLockedBaseLayers;
-
-	Tool->ProcessToolMeshes([this](const FDynamicMesh3& Mesh, int32 MeshIdx)
-		{
-			// Sculpt layer UI only supports a single mesh for now
-			if (MeshIdx > 0)
-			{
-				return;
-			}
-			bHasSculptLayers = Mesh.HasAttributes() && Mesh.Attributes()->NumSculptLayers() > NumLockedBaseLayers;
-			if (bHasSculptLayers)
-			{
-				UpdateSettingsFromMesh(Mesh.Attributes()->GetSculptLayers());
-			}
-		}
-	);
-}
-
-void UMeshSculptLayerProperties::UpdateSettingsFromMesh(const UE::Geometry::FDynamicMeshSculptLayers* SculptLayers)
-{
-	TConstArrayView<double> SourceWeights = SculptLayers->GetLayerWeights();
-	int32 NumLayers = FMath::Max(0, SourceWeights.Num() - NumLockedBaseLayers);
-	LayerWeights.SetNum(NumLayers);
-	for (int32 Idx = 0; Idx < LayerWeights.Num(); ++Idx)
-	{
-		LayerWeights[Idx] = SourceWeights[Idx + NumLockedBaseLayers];
-	}
-	ActiveLayer = SculptLayers->GetActiveLayer();
-}
-
-void UMeshSculptLayerProperties::SetLayerWeights(FDynamicMeshSculptLayers* SculptLayers) const
-{
-	TArray<double> FullLayerWeights(SculptLayers->GetLayerWeights());
-	FullLayerWeights.SetNum(LayerWeights.Num() + NumLockedBaseLayers);
-	for (int32 Idx = 0; Idx < LayerWeights.Num(); ++Idx)
-	{
-		FullLayerWeights[Idx + NumLockedBaseLayers] = LayerWeights[Idx];
-	}
-	SculptLayers->UpdateLayerWeights(FullLayerWeights);
-}
-
-void UMeshSculptLayerProperties::EditSculptLayers(TFunctionRef<void(FDynamicMesh3& Mesh, FDynamicMeshSculptLayers* SculptLayers)> EditFn, bool bEmitChange)
-{
-	if (Tool && Tool->AllowToolMeshUpdates())
-	{
-		Tool->UpdateToolMeshes([this, &EditFn, bEmitChange](FDynamicMesh3& Mesh, int32 Idx) -> TUniquePtr<FMeshRegionChangeBase>
-			{
-				// Sculpt layer UI only supports a single mesh for now
-				if (Idx > 0)
-				{
-					return nullptr;
-				}
-
-				if (!InitialMesh)
-				{
-					InitialMesh = MakeShared<FDynamicMesh3>(Mesh);
-				}
-				if (FDynamicMeshAttributeSet* Attributes = Mesh.Attributes())
-				{
-					if (FDynamicMeshSculptLayers* SculptLayers = Attributes->GetSculptLayers())
-					{
-						EditFn(Mesh, SculptLayers);
-					}
-				}
-
-				if (!bEmitChange)
-				{
-					return nullptr;
-				}
-				TUniquePtr<FMeshRegionChangeBase> Result(new FMeshReplacementChange(InitialMesh, MakeShared<FDynamicMesh3>(Mesh)));
-				InitialMesh = nullptr;
-				return MoveTemp(Result);
-			}
-		);
-	}
-}
-
-#if WITH_EDITOR
-void UMeshSculptLayerProperties::PostEditChangeProperty(FPropertyChangedEvent& PropertyChangedEvent)
-{
-	Super::PostEditChangeProperty(PropertyChangedEvent);
-
-	EditSculptLayers([this, &PropertyChangedEvent](FDynamicMesh3& Mesh, FDynamicMeshSculptLayers* SculptLayers)
-		{
-			if (PropertyChangedEvent.GetMemberPropertyName() == GET_MEMBER_NAME_CHECKED(UMeshSculptLayerProperties, LayerWeights))
-			{
-				SetLayerWeights(SculptLayers);
-			}
-			if (PropertyChangedEvent.GetMemberPropertyName() == GET_MEMBER_NAME_CHECKED(UMeshSculptLayerProperties, ActiveLayer))
-			{
-				ActiveLayer = FMath::Clamp(ActiveLayer, NumLockedBaseLayers, SculptLayers->NumLayers() - 1);
-				SculptLayers->SetActiveLayer(ActiveLayer);
-			}
-		}, 
-		PropertyChangedEvent.ChangeType != EPropertyChangeType::Interactive
-	);
-}
-#endif
-
-void UMeshSculptLayerProperties::AddLayer()
-{
-	EditSculptLayers([this](FDynamicMesh3& Mesh, FDynamicMeshSculptLayers* SculptLayers)
-		{
-			Mesh.Attributes()->EnableSculptLayers(SculptLayers->NumLayers() + 1);
-			LayerWeights.Add(1.0);
-		}, true);
-}
-void UMeshSculptLayerProperties::RemoveLayer()
-{
-	if (LayerWeights.Num() <= 1)
-	{
-		return;
-	}
-	EditSculptLayers([this](FDynamicMesh3& Mesh, FDynamicMeshSculptLayers* SculptLayers)
-		{
-			if (!SculptLayers->DiscardSculptLayer(ActiveLayer))
-			{
-				return;
-			}
-			UpdateSettingsFromMesh(SculptLayers);
-			// If the system picked a locked active layer, try to pick a different layer instead
-			if (ActiveLayer < NumLockedBaseLayers)
-			{
-				// We shouldn't set a zero-weight layer as the active layer, so look for a non-zero weight layer to set
-				int32 NonZeroLayerIdx = INDEX_NONE;
-				for (int32 Idx = 0; Idx < LayerWeights.Num(); ++Idx)
-				{
-					if (LayerWeights[Idx] != 0.0)
-					{
-						NonZeroLayerIdx = Idx;
-						break;
-					}
-				}
-				// if all layers had zero weight, just pick the first layer and set its weight to 1.0 so it is ready to sculpt on
-				if (NonZeroLayerIdx == INDEX_NONE)
-				{
-					LayerWeights[0] = 1.0;
-					NonZeroLayerIdx = 0;
-					SetLayerWeights(SculptLayers);
-				}
-				ActiveLayer = SculptLayers->SetActiveLayer(NumLockedBaseLayers + NonZeroLayerIdx);
-			}
-		}, true);
-}
 
 /*
  * Tool
diff --git a/Engine/Plugins/Runtime/MeshModelingToolset/Source/MeshModelingTools/Private/Properties/MeshSculptLayerProperties.cpp b/Engine/Plugins/Runtime/MeshModelingToolset/Source/MeshModelingTools/Private/Properties/MeshSculptLayerProperties.cpp
new file mode 100644
--- /dev/null
+++ b/Engine/Plugins/Runtime/MeshModelingToolset/Source/MeshModelingTools/Private/Properties/MeshSculptLayerProperties.cpp
@@ -0,0 +1,171 @@
+// Copyright Epic Games, Inc. All Rights Reserved.
+
+#include "Properties/MeshSculptLayerProperties.h"
+#include "DynamicMesh/DynamicMesh3.h"
+#include "DynamicMesh/DynamicMeshAttributeSet.h"
+#include "ModelingToolExternalMeshUpdateAPI.h"
+#include "Changes/MeshReplacementChange.h"
+
+#include UE_INLINE_GENERATED_CPP_BY_NAME(MeshSculptLayerProperties)
+
+#define LOCTEXT_NAMESPACE "UMeshSculptLayerProperties"
+
+void UMeshSculptLayerProperties::Init(IModelingToolExternalDynamicMeshUpdateAPI* InTool, int32 InNumLockedBaseLayers)
+{
+	bCanEditLayers = true;
+	LayerWeights.Reset();
+
+	Tool = InTool;
+	if (!Tool)
+	{
+		return;
+	}
+	NumLockedBaseLayers = InNumLockedBaseLayers;
+
+	Tool->ProcessToolMeshes([this](const UE::Geometry::FDynamicMesh3& Mesh, int32 MeshIdx)
+		{
+			// Sculpt layer UI only supports a single mesh for now
+			if (MeshIdx > 0)
+			{
+				return;
+			}
+			bCanEditLayers = Mesh.HasAttributes() && Mesh.Attributes()->NumSculptLayers() > NumLockedBaseLayers;
+			if (bCanEditLayers)
+			{
+				UpdateSettingsFromMesh(Mesh.Attributes()->GetSculptLayers());
+			}
+		}
+	);
+}
+
+void UMeshSculptLayerProperties::UpdateSettingsFromMesh(const UE::Geometry::FDynamicMeshSculptLayers* SculptLayers)
+{
+	TConstArrayView<double> SourceWeights = SculptLayers->GetLayerWeights();
+	int32 NumLayers = FMath::Max(0, SourceWeights.Num() - NumLockedBaseLayers);
+	LayerWeights.SetNum(NumLayers);
+	for (int32 Idx = 0; Idx < LayerWeights.Num(); ++Idx)
+	{
+		LayerWeights[Idx] = SourceWeights[Idx + NumLockedBaseLayers];
+	}
+	ActiveLayer = SculptLayers->GetActiveLayer();
+}
+
+void UMeshSculptLayerProperties::SetLayerWeights(UE::Geometry::FDynamicMeshSculptLayers* SculptLayers) const
+{
+	TArray<double> FullLayerWeights(SculptLayers->GetLayerWeights());
+	FullLayerWeights.SetNum(LayerWeights.Num() + NumLockedBaseLayers);
+	for (int32 Idx = 0; Idx < LayerWeights.Num(); ++Idx)
+	{
+		FullLayerWeights[Idx + NumLockedBaseLayers] = LayerWeights[Idx];
+	}
+	SculptLayers->UpdateLayerWeights(FullLayerWeights);
+}
+
+void UMeshSculptLayerProperties::EditSculptLayers(TFunctionRef<void(UE::Geometry::FDynamicMesh3& Mesh, UE::Geometry::FDynamicMeshSculptLayers* SculptLayers)> EditFn, bool bEmitChange)
+{
+	if (Tool && Tool->AllowToolMeshUpdates())
+	{
+		Tool->UpdateToolMeshes([this, &EditFn, bEmitChange](UE::Geometry::FDynamicMesh3& Mesh, int32 Idx) -> TUniquePtr<FMeshRegionChangeBase>
+			{
+				// Sculpt layer UI only supports a single mesh for now
+				if (Idx > 0)
+				{
+					return nullptr;
+				}
+
+				if (!InitialMesh)
+				{
+					InitialMesh = MakeShared<UE::Geometry::FDynamicMesh3>(Mesh);
+				}
+				if (UE::Geometry::FDynamicMeshAttributeSet* Attributes = Mesh.Attributes())
+				{
+					if (UE::Geometry::FDynamicMeshSculptLayers* SculptLayers = Attributes->GetSculptLayers())
+					{
+						EditFn(Mesh, SculptLayers);
+					}
+				}
+
+				if (!bEmitChange)
+				{
+					return nullptr;
+				}
+				TUniquePtr<FMeshRegionChangeBase> Result(new FMeshReplacementChange(InitialMesh, MakeShared<FDynamicMesh3>(Mesh)));
+				InitialMesh = nullptr;
+				return MoveTemp(Result);
+			}
+		);
+	}
+}
+
+#if WITH_EDITOR
+void UMeshSculptLayerProperties::PostEditChangeProperty(FPropertyChangedEvent& PropertyChangedEvent)
+{
+	Super::PostEditChangeProperty(PropertyChangedEvent);
+
+	EditSculptLayers([this, &PropertyChangedEvent](UE::Geometry::FDynamicMesh3& Mesh, FDynamicMeshSculptLayers* SculptLayers)
+		{
+			if (PropertyChangedEvent.GetMemberPropertyName() == GET_MEMBER_NAME_CHECKED(UMeshSculptLayerProperties, LayerWeights))
+			{
+				SetLayerWeights(SculptLayers);
+			}
+			if (PropertyChangedEvent.GetMemberPropertyName() == GET_MEMBER_NAME_CHECKED(UMeshSculptLayerProperties, ActiveLayer))
+			{
+				ActiveLayer = FMath::Clamp(ActiveLayer, NumLockedBaseLayers, SculptLayers->NumLayers() - 1);
+				SculptLayers->SetActiveLayer(ActiveLayer);
+			}
+		},
+		PropertyChangedEvent.ChangeType != EPropertyChangeType::Interactive
+	);
+}
+#endif
+
+void UMeshSculptLayerProperties::AddLayer()
+{
+	EditSculptLayers([this](UE::Geometry::FDynamicMesh3& Mesh, UE::Geometry::FDynamicMeshSculptLayers* SculptLayers)
+		{
+			Mesh.Attributes()->EnableSculptLayers(SculptLayers->NumLayers() + 1);
+			LayerWeights.Add(1.0);
+		}, true);
+}
+
+void UMeshSculptLayerProperties::RemoveLayer()
+{
+	if (LayerWeights.Num() <= 1)
+	{
+		return;
+	}
+	EditSculptLayers([this](UE::Geometry::FDynamicMesh3& Mesh, UE::Geometry::FDynamicMeshSculptLayers* SculptLayers)
+		{
+			if (!SculptLayers->DiscardSculptLayer(ActiveLayer))
+			{
+				return;
+			}
+			UpdateSettingsFromMesh(SculptLayers);
+			// If the system picked a locked active layer, try to pick a different layer instead
+			if (ActiveLayer < NumLockedBaseLayers)
+			{
+				// We shouldn't set a zero-weight layer as the active layer, so look for a non-zero weight layer to set
+				int32 NonZeroLayerIdx = INDEX_NONE;
+				for (int32 Idx = 0; Idx < LayerWeights.Num(); ++Idx)
+				{
+					if (LayerWeights[Idx] != 0.0)
+					{
+						NonZeroLayerIdx = Idx;
+						break;
+					}
+				}
+				// if all layers had zero weight, just pick the first layer and set its weight to 1.0 so it is ready to sculpt on
+				if (NonZeroLayerIdx == INDEX_NONE)
+				{
+					LayerWeights[0] = 1.0;
+					NonZeroLayerIdx = 0;
+					SetLayerWeights(SculptLayers);
+				}
+				ActiveLayer = SculptLayers->SetActiveLayer(NumLockedBaseLayers + NonZeroLayerIdx);
+			}
+		}, true);
+}
+
+
+#undef LOCTEXT_NAMESPACE
+
diff --git a/Engine/Plugins/Runtime/MeshModelingToolset/Source/MeshModelingTools/Private/Properties/MeshSculptLayerProperties.h b/Engine/Plugins/Runtime/MeshModelingToolset/Source/MeshModelingTools/Private/Properties/MeshSculptLayerProperties.h
new file mode 100644
--- /dev/null
+++ b/Engine/Plugins/Runtime/MeshModelingToolset/Source/MeshModelingTools/Private/Properties/MeshSculptLayerProperties.h
@@ -0,0 +1,66 @@
+// Copyright Epic Games, Inc. All Rights Reserved.
+
+#pragma once
+
+#include "UObject/Object.h"
+#include "GeometryBase.h"
+
+#include "MeshSculptLayerProperties.generated.h"
+
+PREDECLARE_USE_GEOMETRY_CLASS(FDynamicMesh3);
+PREDECLARE_USE_GEOMETRY_CLASS(FDynamicMeshSculptLayers);
+class IModelingToolExternalDynamicMeshUpdateAPI;
+
+UCLASS(MinimalAPI)
+class UMeshSculptLayerProperties : public UObject
+{
+	GENERATED_BODY()
+
+public:
+
+	/** Set the active sculpt layer */
+	UPROPERTY(EditAnywhere, Category = SculptLayers, meta = (ClampMin = 0, HideEditConditionToggle, EditCondition = bCanEditLayers, ModelingQuickSettings))
+	int32 ActiveLayer = 0;
+
+	/** Set the sculpt layer weights */
+	UPROPERTY(EditAnywhere, EditFixedSize, NoClear, Category = SculptLayers, meta = (NoResetToDefault, HideEditConditionToggle, EditCondition = bCanEditLayers, ModelingQuickSettings))
+	TArray<double> LayerWeights;
+
+	UPROPERTY(meta = (TransientToolProperty))
+	bool bCanEditLayers = false;
+
+	UFUNCTION(CallInEditor, Category = SculptLayers)
+	void AddLayer();
+
+	UFUNCTION(CallInEditor, Category = SculptLayers)
+	void RemoveLayer();
+
+	MESHMODELINGTOOLS_API void Init(IModelingToolExternalDynamicMeshUpdateAPI* InTool, int32 InNumLockedBaseLayers);
+
+#if WITH_EDITOR
+	MESHMODELINGTOOLS_API virtual void PostEditChangeProperty(FPropertyChangedEvent& PropertyChangedEvent) override;
+#endif
+
+private:
+
+	IModelingToolExternalDynamicMeshUpdateAPI* Tool;
+
+	int32 NumLockedBaseLayers = 0;
+
+	// Mesh before layer change has been applied. Used for tracking mesh changes that occur over multiple frames (i.e., from an interactive drag)
+	// TODO: Add/update a mesh FChange to track sculpt layer changes, and use that instead of saving an entire mesh here
+	TSharedPtr<FDynamicMesh3> InitialMesh;
+
+	// Helper to set sculpt layers from the current LayerWeights property (accounting for the NumLockedBaseLayers)
+	void SetLayerWeights(FDynamicMeshSculptLayers* SculptLayers) const;
+
+	// Update the ActiveLayer and LayerWeights settings from the current sculpt layers
+	void UpdateSettingsFromMesh(const FDynamicMeshSculptLayers* SculptLayers);
+
+	// Helper to apply edits to the current sculpt layers if possible, with associated book keeping
+	// @param EditFn The edit to apply if possible
+	// @param bEmitChange Whether to emit a change object along with the edit
+	void EditSculptLayers(TFunctionRef<void(UE::Geometry::FDynamicMesh3& Mesh, UE::Geometry::FDynamicMeshSculptLayers* SculptLayers)> EditFn, bool bEmitChange);
+};
+
+
diff --git a/Engine/Plugins/Runtime/MeshModelingToolset/Source/MeshModelingTools/Public/LatticeDeformerTool.h b/Engine/Plugins/Runtime/MeshModelingToolset/Source/MeshModelingTools/Public/LatticeDeformerTool.h
--- a/Engine/Plugins/Runtime/MeshModelingToolset/Source/MeshModelingTools/Public/LatticeDeformerTool.h
+++ b/Engine/Plugins/Runtime/MeshModelingToolset/Source/MeshModelingTools/Public/LatticeDeformerTool.h
@@ -7,6 +7,7 @@
 #include "InteractiveToolBuilder.h" //UInteractiveToolBuilder
 #include "InteractiveToolChange.h" //FToolCommandChange
 #include "BaseTools/MultiTargetWithSelectionTool.h" // UMultiTargetWithSelectionTool
+#include "ModelingToolExternalMeshUpdateAPI.h"
 #include "DynamicMesh/MeshSharingUtil.h"
 #include "Operations/FFDLattice.h"
 #include "Solvers/ConstrainedMeshDeformer.h"
@@ -20,6 +21,7 @@ namespace UE::Geometry
 
 class ULatticeControlPointsMechanic;
 class UMeshOpPreviewWithBackgroundCompute;
+class UMeshSculptLayerProperties;
 
 UCLASS()
 class MESHMODELINGTOOLS_API ULatticeDeformerToolBuilder : public UMultiTargetWithSelectionToolBuilder
@@ -31,6 +33,9 @@ public:
 	virtual bool CanBuildTool(const FToolBuilderState& SceneState) const override;
 
 	virtual bool RequiresInputSelection() const override { return false; }
+
+private:
+	virtual const FToolTargetTypeRequirements& GetTargetRequirements() const override;
 };
 
 
@@ -135,7 +140,7 @@ public:
 
 /** Deform a mesh using a regular hexahedral lattice */
 UCLASS()
-class MESHMODELINGTOOLS_API ULatticeDeformerTool : public UMultiTargetWithSelectionTool, public IInteractiveToolManageGeometrySelectionAPI
+class MESHMODELINGTOOLS_API ULatticeDeformerTool : public UMultiTargetWithSelectionTool, public IInteractiveToolManageGeometrySelectionAPI, public IModelingToolExternalDynamicMeshUpdateAPI
 {
 	GENERATED_BODY()
 
@@ -183,6 +188,9 @@ protected:
 	UPROPERTY()
 	TObjectPtr<UMeshOpPreviewWithBackgroundCompute> Preview = nullptr;
 
+	UPROPERTY()
+	TObjectPtr<UMeshSculptLayerProperties> SculptLayerProperties;
+
 	UPROPERTY()
 	bool bLatticeDeformed = false;
 
@@ -190,6 +198,12 @@ protected:
 
 	bool bHasSelection = false;
 
+	// IInteractiveToolExternalDynamicMeshUpdateAPI methods
+	virtual bool AllowToolMeshUpdates() const override;
+	virtual void UpdateToolMeshes(TFunctionRef<TUniquePtr<FMeshRegionChangeBase>(UE::Geometry::FDynamicMesh3&, int32 MeshIdx)> UpdateMesh) override;
+	virtual void ProcessToolMeshes(TFunctionRef<void(const UE::Geometry::FDynamicMesh3&, int32 MeshIdx)> ProcessMesh) const override;
+	virtual int32 NumToolMeshes() const override;
+	
 	// Create and store an FFFDLattice. Pass out the lattice's positions and edges.
 	void InitializeLattice(TArray<FVector3d>& OutLatticePoints, TArray<UE::Geometry::FVector2i>& OutLatticeEdges);
 
diff --git a/Engine/Plugins/Runtime/MeshModelingToolset/Source/MeshModelingTools/Public/MeshVertexSculptTool.h b/Engine/Plugins/Runtime/MeshModelingToolset/Source/MeshModelingTools/Public/MeshVertexSculptTool.h
--- a/Engine/Plugins/Runtime/MeshModelingToolset/Source/MeshModelingTools/Public/MeshVertexSculptTool.h
+++ b/Engine/Plugins/Runtime/MeshModelingToolset/Source/MeshModelingTools/Public/MeshVertexSculptTool.h
@@ -26,7 +26,7 @@ class FMeshVertexChangeBuilder;
 class UPreviewMesh;
 namespace UE::Geometry { class FDynamicMeshSculptLayers; }
 PREDECLARE_GEOMETRY(class FMeshPlanarSymmetry);
-
+class UMeshSculptLayerProperties;
 
 
 /**
@@ -213,55 +213,6 @@ public:
 
 
 
-//~ TODO: These properties should be pulled out to a separate UI, and not tied to the vertex sculpt tool
-UCLASS(MinimalAPI)
-class UMeshSculptLayerProperties : public UObject
-{
-	GENERATED_BODY()
-public:
-	/** Set the active sculpt layer */
-	UPROPERTY(EditAnywhere, Category = SculptLayers, meta = (ClampMin = 0, HideEditConditionToggle, EditCondition = bHasSculptLayers, ModelingQuickSettings))
-	int32 ActiveLayer = 0;
-
-	/** Set the sculpt layer weights */
-	UPROPERTY(EditAnywhere, EditFixedSize, NoClear, Category = SculptLayers, meta = (NoResetToDefault, HideEditConditionToggle, EditCondition = bHasSculptLayers, ModelingQuickSettings))
-	TArray<double> LayerWeights;
-
-	UPROPERTY(meta = (TransientToolProperty))
-	bool bHasSculptLayers = false;
-
-	UFUNCTION(CallInEditor, Category = SculptLayers)
-	void AddLayer();
-
-	UFUNCTION(CallInEditor, Category = SculptLayers)
-	void RemoveLayer();
-
-	MESHMODELINGTOOLS_API void Init(IModelingToolExternalDynamicMeshUpdateAPI* InTool, int32 InNumLockedBaseLayers);
-
-#if WITH_EDITOR
-	MESHMODELINGTOOLS_API virtual void PostEditChangeProperty(FPropertyChangedEvent& PropertyChangedEvent) override;
-#endif
-
-private:
-	IModelingToolExternalDynamicMeshUpdateAPI* Tool;
-
-	int32 NumLockedBaseLayers = 0;
-
-	// Mesh before layer change has been applied. Used for tracking mesh changes that occur over multiple frames (i.e., from an interactive drag)
-	// TODO: Add/update a mesh FChange to track sculpt layer changes, and use that instead of saving an entire mesh here
-	TSharedPtr<FDynamicMesh3> InitialMesh;
-
-	// Helper to set sculpt layers from the current LayerWeights property (accounting for the NumLockedBaseLayers)
-	void SetLayerWeights(UE::Geometry::FDynamicMeshSculptLayers* SculptLayers) const;
-
-	// Update the ActiveLayer and LayerWeights settings from the current sculpt layers
-	void UpdateSettingsFromMesh(const UE::Geometry::FDynamicMeshSculptLayers* SculptLayers);
-
-	// Helper to apply edits to the current sculpt layers if possible, with associated book keeping
-	// @param EditFn The edit to apply if possible
-	// @param bEmitChange Whether to emit a change object along with the edit
-	void EditSculptLayers(TFunctionRef<void(UE::Geometry::FDynamicMesh3& Mesh, UE::Geometry::FDynamicMeshSculptLayers* SculptLayers)> EditFn, bool bEmitChange);
-};
 
 
 /**
diff --git a/Engine/Plugins/Runtime/MeshModelingToolset/Source/ModelingComponents/Public/Mechanics/LatticeControlPointsMechanic.h b/Engine/Plugins/Runtime/MeshModelingToolset/Source/ModelingComponents/Public/Mechanics/LatticeControlPointsMechanic.h
--- a/Engine/Plugins/Runtime/MeshModelingToolset/Source/ModelingComponents/Public/Mechanics/LatticeControlPointsMechanic.h
+++ b/Engine/Plugins/Runtime/MeshModelingToolset/Source/ModelingComponents/Public/Mechanics/LatticeControlPointsMechanic.h
@@ -103,6 +103,11 @@ public:
 
 	void UpdatePointLocations(const TMap<int32, FVector3d>& NewLocations);
 
+	bool IsGizmoBeingDragged() const
+	{
+		return bGizmoBeingDragged;
+	}
+
 	DECLARE_DELEGATE_RetVal(bool, FShouldHideGizmo);
 	FShouldHideGizmo ShouldHideGizmo;
 
