Modeling Mode: Set and implement quick settings priority values

diff --git a/Engine/Plugins/Runtime/MeshModelingToolset/Source/MeshModelingTools/Public/MeshVertexSculptTool.h b/Engine/Plugins/Runtime/MeshModelingToolset/Source/MeshModelingTools/Public/MeshVertexSculptTool.h
--- a/Engine/Plugins/Runtime/MeshModelingToolset/Source/MeshModelingTools/Public/MeshVertexSculptTool.h
+++ b/Engine/Plugins/Runtime/MeshModelingToolset/Source/MeshModelingTools/Public/MeshVertexSculptTool.h
@@ -141,7 +141,7 @@ public:
 	int32 PrimaryBrushID = 0;
 
 	/** Primary Brush Falloff Type, multiplied by Alpha Mask where applicable */
-	UPROPERTY(EditAnywhere, Category = Sculpting, meta = (DisplayName = "Falloff", ModelingQuickSettings))
+	UPROPERTY(EditAnywhere, Category = Sculpting, meta = (DisplayName = "Falloff", ModelingQuickSettings = 300))
 	EMeshSculptFalloffType PrimaryFalloffType = EMeshSculptFalloffType::Smooth;
 
 	/** Filter applied to Stamp Region Triangles, based on first Stroke Stamp */
@@ -172,7 +172,7 @@ class MESHMODELINGTOOLS_API UVertexBrushAlphaProperties : public UInteractiveToo
 
 public:
 	/** Alpha mask applied to brush stamp. Red channel is used. */
-	UPROPERTY(EditAnywhere, Category = Alpha, meta = (DisplayName = "Alpha Mask", ModelingQuickSettings))
+	UPROPERTY(EditAnywhere, Category = Alpha, meta = (DisplayName = "Alpha Mask", ModelingQuickSettings = 400))
 	TObjectPtr<UTexture2D> Alpha = nullptr;
 
 	/** Alpha is rotated by this angle, inside the brush stamp frame (vertically aligned) */
@@ -201,7 +201,7 @@ class MESHMODELINGTOOLS_API UMeshSymmetryProperties : public UInteractiveToolPro
 	GENERATED_BODY()
 public:
 	/** Enable/Disable symmetric sculpting. This option will not be available if symmetry cannot be detected, or a non-symmetric edit has been made */
-	UPROPERTY(EditAnywhere, Category = Symmetry, meta = (HideEditConditionToggle, EditCondition = bSymmetryCanBeEnabled, ModelingQuickSettings))
+	UPROPERTY(EditAnywhere, Category = Symmetry, meta = (HideEditConditionToggle, EditCondition = bSymmetryCanBeEnabled, ModelingQuickSettings = 500))
 	bool bEnableSymmetry = false;
 
 	// this flag is set/updated by the Tool to enable/disable the bEnableSymmetry toggle
diff --git a/Engine/Plugins/Runtime/MeshModelingToolset/Source/MeshModelingTools/Public/Sculpting/KelvinletBrushOp.h b/Engine/Plugins/Runtime/MeshModelingToolset/Source/MeshModelingTools/Public/Sculpting/KelvinletBrushOp.h
--- a/Engine/Plugins/Runtime/MeshModelingToolset/Source/MeshModelingTools/Public/Sculpting/KelvinletBrushOp.h
+++ b/Engine/Plugins/Runtime/MeshModelingToolset/Source/MeshModelingTools/Public/Sculpting/KelvinletBrushOp.h
@@ -194,7 +194,7 @@ class MESHMODELINGTOOLS_API UScaleKelvinletBrushOpProps : public   UBaseKelvinle
 
 public:
 	/** Strength of the Brush */
-	UPROPERTY(EditAnywhere, Category = KelvinScaleBrush, meta = (DisplayName = "Strength", UIMin = "0.0", UIMax = "10.", ClampMin = "0.0", ClampMax = "10.", ModelingQuickEdit, ModelingQuickSettings))
+	UPROPERTY(EditAnywhere, Category = KelvinScaleBrush, meta = (DisplayName = "Strength", UIMin = "0.0", UIMax = "10.", ClampMin = "0.0", ClampMax = "10.", ModelingQuickEdit, ModelingQuickSettings = 200))
 	float Strength = 0.5;
 
 	/** Amount of falloff to apply */
@@ -401,7 +401,7 @@ class MESHMODELINGTOOLS_API UTwistKelvinletBrushOpProps : public  UBaseKelvinlet
 
 public:
 	/** Twisting strength of the Brush */
-	UPROPERTY(EditAnywhere, Category = KelvinTwistBrush, meta = (DisplayName = "Strength", UIMin = "0.0", UIMax = "10.", ClampMin = "0.0", ClampMax = "10.", ModelingQuickSettings))
+	UPROPERTY(EditAnywhere, Category = KelvinTwistBrush, meta = (DisplayName = "Strength", UIMin = "0.0", UIMax = "10.", ClampMin = "0.0", ClampMax = "10.", ModelingQuickSettings = 200))
 	float Strength = 0.5;
 
 	/** Amount of falloff to apply */
diff --git a/Engine/Plugins/Runtime/MeshModelingToolset/Source/MeshModelingTools/Public/Sculpting/MeshInflateBrushOps.h b/Engine/Plugins/Runtime/MeshModelingToolset/Source/MeshModelingTools/Public/Sculpting/MeshInflateBrushOps.h
--- a/Engine/Plugins/Runtime/MeshModelingToolset/Source/MeshModelingTools/Public/Sculpting/MeshInflateBrushOps.h
+++ b/Engine/Plugins/Runtime/MeshModelingToolset/Source/MeshModelingTools/Public/Sculpting/MeshInflateBrushOps.h
@@ -16,7 +16,7 @@ class MESHMODELINGTOOLS_API UInflateBrushOpProps : public UMeshSculptBrushOpProp
 	GENERATED_BODY()
 public:
 	/** Strength of the Brush */
-	UPROPERTY(EditAnywhere, Category = InflateBrush, meta = (DisplayName = "Strength", UIMin = "0.0", UIMax = "1.0", ClampMin = "0.0", ClampMax = "1.0", ModelingQuickEdit, ModelingQuickSettings))
+	UPROPERTY(EditAnywhere, Category = InflateBrush, meta = (DisplayName = "Strength", UIMin = "0.0", UIMax = "1.0", ClampMin = "0.0", ClampMax = "1.0", ModelingQuickEdit, ModelingQuickSettings = 200))
 	float Strength = 0.5;
 
 	/** Amount of falloff to apply */
diff --git a/Engine/Plugins/Runtime/MeshModelingToolset/Source/MeshModelingTools/Public/Sculpting/MeshMoveBrushOps.h b/Engine/Plugins/Runtime/MeshModelingToolset/Source/MeshModelingTools/Public/Sculpting/MeshMoveBrushOps.h
--- a/Engine/Plugins/Runtime/MeshModelingToolset/Source/MeshModelingTools/Public/Sculpting/MeshMoveBrushOps.h
+++ b/Engine/Plugins/Runtime/MeshModelingToolset/Source/MeshModelingTools/Public/Sculpting/MeshMoveBrushOps.h
@@ -18,7 +18,7 @@ class MESHMODELINGTOOLS_API UMoveBrushOpProps : public UMeshSculptBrushOpProps
 	GENERATED_BODY()
 public:
 	/** Strength of the Brush */
-	UPROPERTY(EditAnywhere, Category = MoveBrush, meta = (DisplayName = "Strength", UIMin = "0.0", UIMax = "1.0", ClampMin = "0.0", ClampMax = "1.0", ModelingQuickEdit, ModelingQuickSettings))
+	UPROPERTY(EditAnywhere, Category = MoveBrush, meta = (DisplayName = "Strength", UIMin = "0.0", UIMax = "1.0", ClampMin = "0.0", ClampMax = "1.0", ModelingQuickEdit, ModelingQuickSettings = 200))
 	float Strength = 1.0;
 
 	/** Amount of falloff to apply */
diff --git a/Engine/Plugins/Runtime/MeshModelingToolset/Source/MeshModelingTools/Public/Sculpting/MeshPinchBrushOps.h b/Engine/Plugins/Runtime/MeshModelingToolset/Source/MeshModelingTools/Public/Sculpting/MeshPinchBrushOps.h
--- a/Engine/Plugins/Runtime/MeshModelingToolset/Source/MeshModelingTools/Public/Sculpting/MeshPinchBrushOps.h
+++ b/Engine/Plugins/Runtime/MeshModelingToolset/Source/MeshModelingTools/Public/Sculpting/MeshPinchBrushOps.h
@@ -15,7 +15,7 @@ class MESHMODELINGTOOLS_API UPinchBrushOpProps : public UMeshSculptBrushOpProps
 	GENERATED_BODY()
 public:
 	/** Strength of the Brush */
-	UPROPERTY(EditAnywhere, Category = PinchBrush, meta = (DisplayName = "Strength", UIMin = "0.0", UIMax = "1.0", ClampMin = "0.0", ClampMax = "1.0", ModelingQuickEdit, ModelingQuickSettings))
+	UPROPERTY(EditAnywhere, Category = PinchBrush, meta = (DisplayName = "Strength", UIMin = "0.0", UIMax = "1.0", ClampMin = "0.0", ClampMax = "1.0", ModelingQuickEdit, ModelingQuickSettings = 200))
 	float Strength = 0.5;
 
 	/** Amount of falloff to apply */
diff --git a/Engine/Plugins/Runtime/MeshModelingToolset/Source/MeshModelingTools/Public/Sculpting/MeshPlaneBrushOps.h b/Engine/Plugins/Runtime/MeshModelingToolset/Source/MeshModelingTools/Public/Sculpting/MeshPlaneBrushOps.h
--- a/Engine/Plugins/Runtime/MeshModelingToolset/Source/MeshModelingTools/Public/Sculpting/MeshPlaneBrushOps.h
+++ b/Engine/Plugins/Runtime/MeshModelingToolset/Source/MeshModelingTools/Public/Sculpting/MeshPlaneBrushOps.h
@@ -26,7 +26,7 @@ class MESHMODELINGTOOLS_API UPlaneBrushOpProps : public UBasePlaneBrushOpProps
 	GENERATED_BODY()
 public:
 	/** Strength of the Brush */
-	UPROPERTY(EditAnywhere, Category = PlaneBrush, meta = (DisplayName = "Strength", UIMin = "0.0", UIMax = "1.0", ClampMin = "0.0", ClampMax = "1.0", ModelingQuickEdit, ModelingQuickSettings))
+	UPROPERTY(EditAnywhere, Category = PlaneBrush, meta = (DisplayName = "Strength", UIMin = "0.0", UIMax = "1.0", ClampMin = "0.0", ClampMax = "1.0", ModelingQuickEdit, ModelingQuickSettings = 200))
 	float Strength = 0.5;
 
 	/** Amount of falloff to apply */
@@ -55,7 +55,7 @@ class MESHMODELINGTOOLS_API UViewAlignedPlaneBrushOpProps : public UBasePlaneBru
 	GENERATED_BODY()
 public:
 	/** Strength of the Brush */
-	UPROPERTY(EditAnywhere, Category = ViewPlaneBrush, meta = (DisplayName = "Strength", UIMin = "0.0", UIMax = "1.0", ClampMin = "0.0", ClampMax = "1.0", ModelingQuickEdit, ModelingQuickSettings))
+	UPROPERTY(EditAnywhere, Category = ViewPlaneBrush, meta = (DisplayName = "Strength", UIMin = "0.0", UIMax = "1.0", ClampMin = "0.0", ClampMax = "1.0", ModelingQuickEdit, ModelingQuickSettings = 200))
 	float Strength = 0.5;
 
 	/** Amount of falloff to apply */
@@ -84,7 +84,7 @@ class MESHMODELINGTOOLS_API UFixedPlaneBrushOpProps : public UBasePlaneBrushOpPr
 	GENERATED_BODY()
 public:
 	/** Strength of the Brush */
-	UPROPERTY(EditAnywhere, Category = FixedPlaneBrush, meta = (DisplayName = "Strength", UIMin = "0.0", UIMax = "1.0", ClampMin = "0.0", ClampMax = "1.0", ModelingQuickEdit, ModelingQuickSettings))
+	UPROPERTY(EditAnywhere, Category = FixedPlaneBrush, meta = (DisplayName = "Strength", UIMin = "0.0", UIMax = "1.0", ClampMin = "0.0", ClampMax = "1.0", ModelingQuickEdit, ModelingQuickSettings = 200))
 	float Strength = 0.5;
 
 	/** Amount of falloff to apply */
diff --git a/Engine/Plugins/Runtime/MeshModelingToolset/Source/MeshModelingTools/Public/Sculpting/MeshSculptBrushOps.h b/Engine/Plugins/Runtime/MeshModelingToolset/Source/MeshModelingTools/Public/Sculpting/MeshSculptBrushOps.h
--- a/Engine/Plugins/Runtime/MeshModelingToolset/Source/MeshModelingTools/Public/Sculpting/MeshSculptBrushOps.h
+++ b/Engine/Plugins/Runtime/MeshModelingToolset/Source/MeshModelingTools/Public/Sculpting/MeshSculptBrushOps.h
@@ -16,7 +16,7 @@ class MESHMODELINGTOOLS_API UStandardSculptBrushOpProps : public UMeshSculptBrus
 	GENERATED_BODY()
 public:
 	/** Strength of the Brush */
-	UPROPERTY(EditAnywhere, Category = SculptBrush, meta = (DisplayName = "Strength", UIMin = "0.0", UIMax = "1.0", ClampMin = "0.0", ClampMax = "1.0", ModelingQuickEdit, ModelingQuickSettings))
+	UPROPERTY(EditAnywhere, Category = SculptBrush, meta = (DisplayName = "Strength", UIMin = "0.0", UIMax = "1.0", ClampMin = "0.0", ClampMax = "1.0", ModelingQuickEdit, ModelingQuickSettings = 200))
 	float Strength = 0.5;
 
 	/** Amount of falloff to apply */
@@ -98,7 +98,7 @@ class MESHMODELINGTOOLS_API UViewAlignedSculptBrushOpProps : public UMeshSculptB
 	GENERATED_BODY()
 public:
 	/** Strength of the Brush */
-	UPROPERTY(EditAnywhere, Category = SculptToViewBrush, meta = (DisplayName = "Strength", UIMin = "0.0", UIMax = "1.0", ClampMin = "0.0", ClampMax = "1.0", ModelingQuickEdit, ModelingQuickSettings))
+	UPROPERTY(EditAnywhere, Category = SculptToViewBrush, meta = (DisplayName = "Strength", UIMin = "0.0", UIMax = "1.0", ClampMin = "0.0", ClampMax = "1.0", ModelingQuickEdit, ModelingQuickSettings = 200))
 	float Strength = 0.5;
 
 	/** Amount of falloff to apply */
@@ -188,7 +188,7 @@ class MESHMODELINGTOOLS_API USculptMaxBrushOpProps : public UMeshSculptBrushOpPr
 	GENERATED_BODY()
 public:
 	/** Strength of the Brush */
-	UPROPERTY(EditAnywhere, Category = SculptMaxBrush, meta = (DisplayName = "Strength", UIMin = "0.0", UIMax = "1.0", ClampMin = "0.0", ClampMax = "1.0", ModelingQuickEdit, ModelingQuickSettings))
+	UPROPERTY(EditAnywhere, Category = SculptMaxBrush, meta = (DisplayName = "Strength", UIMin = "0.0", UIMax = "1.0", ClampMin = "0.0", ClampMax = "1.0", ModelingQuickEdit, ModelingQuickSettings = 200))
 	float Strength = 0.5;
 
 	/** Amount of falloff to apply */
diff --git a/Engine/Plugins/Runtime/MeshModelingToolset/Source/MeshModelingTools/Public/Sculpting/MeshSculptToolBase.h b/Engine/Plugins/Runtime/MeshModelingToolset/Source/MeshModelingTools/Public/Sculpting/MeshSculptToolBase.h
--- a/Engine/Plugins/Runtime/MeshModelingToolset/Source/MeshModelingTools/Public/Sculpting/MeshSculptToolBase.h
+++ b/Engine/Plugins/Runtime/MeshModelingToolset/Source/MeshModelingTools/Public/Sculpting/MeshSculptToolBase.h
@@ -107,7 +107,7 @@ class MESHMODELINGTOOLS_API USculptBrushProperties : public UInteractiveToolProp
 	GENERATED_BODY()
 public:
 
-	UPROPERTY(EditAnywhere, Category = Brush, meta = (DisplayPriority = 1, ModelingQuickEdit, ModelingQuickSettings))
+	UPROPERTY(EditAnywhere, Category = Brush, meta = (DisplayPriority = 1, ModelingQuickEdit, ModelingQuickSettings = 100))
 	FBrushToolRadius BrushSize;
 
 	/** Amount of falloff to apply (0.0 - 1.0) */
diff --git a/Engine/Plugins/Runtime/MeshModelingToolset/Source/MeshModelingTools/Public/Sculpting/MeshSmoothingBrushOps.h b/Engine/Plugins/Runtime/MeshModelingToolset/Source/MeshModelingTools/Public/Sculpting/MeshSmoothingBrushOps.h
--- a/Engine/Plugins/Runtime/MeshModelingToolset/Source/MeshModelingTools/Public/Sculpting/MeshSmoothingBrushOps.h
+++ b/Engine/Plugins/Runtime/MeshModelingToolset/Source/MeshModelingTools/Public/Sculpting/MeshSmoothingBrushOps.h
@@ -26,7 +26,7 @@ class MESHMODELINGTOOLS_API USmoothBrushOpProps : public UBaseSmoothBrushOpProps
 	GENERATED_BODY()
 public:
 	/** Strength of the Brush */
-	UPROPERTY(EditAnywhere, Category = SmoothBrush, meta = (DisplayName = "Strength", UIMin = "0.0", UIMax = "1.0", ClampMin = "0.0", ClampMax = "1.0", ModelingQuickEdit, ModelingQuickSettings))
+	UPROPERTY(EditAnywhere, Category = SmoothBrush, meta = (DisplayName = "Strength", UIMin = "0.0", UIMax = "1.0", ClampMin = "0.0", ClampMax = "1.0", ModelingQuickEdit, ModelingQuickSettings = 200))
 	float Strength = 0.5;
 
 	/** Amount of falloff to apply (0.0 - 1.0) */
@@ -115,7 +115,7 @@ class MESHMODELINGTOOLS_API USmoothFillBrushOpProps : public UBaseSmoothBrushOpP
 	GENERATED_BODY()
 public:
 	/** Strength of the Brush */
-	UPROPERTY(EditAnywhere, Category = SmoothFillBrush, meta = (DisplayName = "Strength", UIMin = "0.0", UIMax = "1.0", ClampMin = "0.0", ClampMax = "1.0", ModelingQuickEdit, ModelingQuickSettings))
+	UPROPERTY(EditAnywhere, Category = SmoothFillBrush, meta = (DisplayName = "Strength", UIMin = "0.0", UIMax = "1.0", ClampMin = "0.0", ClampMax = "1.0", ModelingQuickEdit, ModelingQuickSettings = 200))
 	float Strength = 0.5;
 
 	/** Amount of falloff to apply (0.0 - 1.0) */
@@ -188,7 +188,7 @@ class MESHMODELINGTOOLS_API UFlattenBrushOpProps : public UMeshSculptBrushOpProp
 	GENERATED_BODY()
 public:
 	/** Strength of the Brush */
-	UPROPERTY(EditAnywhere, Category = FlattenBrush, meta = (DisplayName = "Strength", UIMin = "0.0", UIMax = "1.0", ClampMin = "0.0", ClampMax = "1.0", ModelingQuickEdit, ModelingQuickSettings))
+	UPROPERTY(EditAnywhere, Category = FlattenBrush, meta = (DisplayName = "Strength", UIMin = "0.0", UIMax = "1.0", ClampMin = "0.0", ClampMax = "1.0", ModelingQuickEdit, ModelingQuickSettings = 200))
 	float Strength = 0.5;
 
 	/** Amount of falloff to apply (0.0 - 1.0) */
@@ -269,7 +269,7 @@ class MESHMODELINGTOOLS_API UEraseBrushOpProps : public UMeshSculptBrushOpProps
 	GENERATED_BODY()
 public:
 	/** Strength of the Brush */
-	UPROPERTY(EditAnywhere, Category = EraseBrush, meta = (DisplayName = "Strength", UIMin = "0.0", UIMax = "1.0", ClampMin = "0.0", ClampMax = "1.0", ModelingQuickEdit, ModelingQuickSettings))
+	UPROPERTY(EditAnywhere, Category = EraseBrush, meta = (DisplayName = "Strength", UIMin = "0.0", UIMax = "1.0", ClampMin = "0.0", ClampMax = "1.0", ModelingQuickEdit, ModelingQuickSettings = 200))
 	float Strength = 0.5;
 
 	/** Amount of falloff to apply (0.0 - 1.0) */
