add option to use new mesh cluster simplifier in the simplify mesh tool

diff --git a/Engine/Plugins/Runtime/MeshModelingToolset/Source/MeshModelingToolsEditorOnly/Private/SimplifyMeshTool.cpp b/Engine/Plugins/Runtime/MeshModelingToolset/Source/MeshModelingToolsEditorOnly/Private/SimplifyMeshTool.cpp
--- a/Engine/Plugins/Runtime/MeshModelingToolset/Source/MeshModelingToolsEditorOnly/Private/SimplifyMeshTool.cpp
+++ b/Engine/Plugins/Runtime/MeshModelingToolset/Source/MeshModelingToolsEditorOnly/Private/SimplifyMeshTool.cpp
@@ -227,6 +227,7 @@ TUniquePtr<FDynamicMeshOperator> USimplifyMeshTool::MakeNewOperator()
 	Op->bGeometricDeviationConstraint = SimplifyProperties->bGeometricConstraint;
 	Op->GeometricTolerance = SimplifyProperties->GeometricTolerance;
 	Op->PolyEdgeAngleTolerance = SimplifyProperties->PolyEdgeAngleTolerance;
+	Op->BoundaryEdgeAngleTolerance = SimplifyProperties->BoundaryEdgeAngleTolerance;
 	FTransform LocalToWorld = (FTransform)UE::ToolTarget::GetLocalToWorldTransform(Target);
 	Op->SetTransform(LocalToWorld);
 
diff --git a/Engine/Plugins/Runtime/MeshModelingToolset/Source/MeshModelingToolsEditorOnly/Public/SimplifyMeshTool.h b/Engine/Plugins/Runtime/MeshModelingToolset/Source/MeshModelingToolsEditorOnly/Public/SimplifyMeshTool.h
--- a/Engine/Plugins/Runtime/MeshModelingToolset/Source/MeshModelingToolsEditorOnly/Public/SimplifyMeshTool.h
+++ b/Engine/Plugins/Runtime/MeshModelingToolset/Source/MeshModelingToolsEditorOnly/Public/SimplifyMeshTool.h
@@ -44,27 +44,27 @@ public:
 	ESimplifyType SimplifierType;
 
 	/** Simplification Target Type  */
-	UPROPERTY(EditAnywhere, Category = Options, meta = (EditCondition = "SimplifierType != ESimplifyType::MinimalPlanar && SimplifierType != ESimplifyType::MinimalPolygroup"))
+	UPROPERTY(EditAnywhere, Category = Options, meta = (EditCondition = "SimplifierType != ESimplifyType::MinimalPlanar && SimplifierType != ESimplifyType::MinimalPolygroup && SimplifierType != ESimplifyType::ClusterBased"))
 	ESimplifyTargetType TargetMode;
 
 	/** Target percentage of original triangle count */
 	UPROPERTY(EditAnywhere, Category = Options, meta = (UIMin = "0", UIMax = "100", ClampMin = "0.0", ClampMax = "100.0", 
-				EditCondition = "SimplifierType != ESimplifyType::MinimalPolygroup && SimplifierType != ESimplifyType::MinimalPlanar && TargetMode == ESimplifyTargetType::Percentage"))
+				EditCondition = "SimplifierType != ESimplifyType::ClusterBased && SimplifierType != ESimplifyType::MinimalPolygroup && SimplifierType != ESimplifyType::MinimalPlanar && TargetMode == ESimplifyTargetType::Percentage"))
 	int TargetPercentage;
 
 	/** Target edge length */
 	UPROPERTY(EditAnywhere, Category = Options, meta = (UIMin = "3.0", UIMax = "10.0", ClampMin = "0.001", ClampMax = "1000.0",
-		EditCondition = "TargetMode == ESimplifyTargetType::EdgeLength && SimplifierType != ESimplifyType::UEStandard && SimplifierType != ESimplifyType::MinimalPlanar && SimplifierType != ESimplifyType::MinimalPolygroup"))
+		EditCondition = "SimplifierType == ESimplifyType::ClusterBased || (TargetMode == ESimplifyTargetType::EdgeLength && SimplifierType != ESimplifyType::UEStandard && SimplifierType != ESimplifyType::MinimalPlanar && SimplifierType != ESimplifyType::MinimalPolygroup)"))
 	float TargetEdgeLength;
 
 	/** Target triangle count */
 	UPROPERTY(EditAnywhere, Category = Options, meta = (UIMin = "4", UIMax = "10000", ClampMin = "1", ClampMax = "9999999999",
-				EditCondition = "TargetMode == ESimplifyTargetType::TriangleCount && SimplifierType != ESimplifyType::MinimalPlanar && SimplifierType != ESimplifyType::MinimalPolygroup"))
+				EditCondition = "TargetMode == ESimplifyTargetType::TriangleCount && SimplifierType != ESimplifyType::MinimalPlanar && SimplifierType != ESimplifyType::MinimalPolygroup && SimplifierType != ESimplifyType::ClusterBased"))
 	int TargetTriangleCount;
 
 	/** Target vertex count */
 	UPROPERTY(EditAnywhere, Category = Options, meta = (UIMin = "4", UIMax = "10000", ClampMin = "1", ClampMax = "9999999999",
-				EditCondition = "TargetMode == ESimplifyTargetType::VertexCount && SimplifierType != ESimplifyType::MinimalPlanar"))
+				EditCondition = "TargetMode == ESimplifyTargetType::VertexCount && SimplifierType != ESimplifyType::MinimalPlanar && SimplifierType != ESimplifyType::ClusterBased"))
 	int TargetVertexCount;
 
 	/** Angle threshold in degrees used for testing if two triangles should be considered coplanar, or two lines collinear */
@@ -76,6 +76,10 @@ public:
 	UPROPERTY(EditAnywhere, Category = Options, meta = (UIMin = "0.001", ClampMin = "0.0", UIMax = "90.0", ClampMax = "180.0", EditCondition = "SimplifierType == ESimplifyType::MinimalPolygroup"))
 	float PolyEdgeAngleTolerance = .1;
 
+	/** Threshold angle change (in degrees) between adjacent boundary edges, above which the common vertex must be preserved */
+	UPROPERTY(EditAnywhere, Category = Options, meta = (UIMin = "0.001", ClampMin = "0.0", UIMax = "90.0", ClampMax = "180.0", EditCondition = "SimplifierType == ESimplifyType::ClusterBased"))
+	float BoundaryEdgeAngleTolerance = 30.0f;
+
 	/** Discard UVs and existing normals, allowing the simplifier to ignore any UV and normal seams. New per-vertex normals are computed. */
 	UPROPERTY(EditAnywhere, Category = Options)
 	bool bDiscardAttributes;
diff --git a/Engine/Plugins/Runtime/MeshModelingToolset/Source/ModelingOperatorsEditorOnly/Private/CleaningOps/SimplifyMeshOp.cpp b/Engine/Plugins/Runtime/MeshModelingToolset/Source/ModelingOperatorsEditorOnly/Private/CleaningOps/SimplifyMeshOp.cpp
--- a/Engine/Plugins/Runtime/MeshModelingToolset/Source/ModelingOperatorsEditorOnly/Private/CleaningOps/SimplifyMeshOp.cpp
+++ b/Engine/Plugins/Runtime/MeshModelingToolset/Source/ModelingOperatorsEditorOnly/Private/CleaningOps/SimplifyMeshOp.cpp
@@ -19,6 +19,7 @@
 #include "DynamicMesh/Operations/MergeCoincidentMeshEdges.h"
 #include "GroupTopology.h"
 #include "Operations/PolygroupRemesh.h"
+#include "Operations/MeshClusterSimplifier.h"
 #include "ConstrainedDelaunay2.h"
 #include "OverlappingCorners.h"
 #include "StaticMeshOperations.h"
@@ -313,6 +314,50 @@ void FSimplifyMeshOp::CalculateResult(FProgressCancel* Progress)
 		Remesh.SimplificationAngleTolerance = PolyEdgeAngleTolerance;
 		Remesh.Compute();
 	}
+	else if (SimplifierType == ESimplifyType::ClusterBased)
+	{
+		MeshClusterSimplify::FSimplifyOptions SimplifyOptions;
+		SimplifyOptions.TargetEdgeLength = TargetEdgeLength;
+		SimplifyOptions.FixBoundaryAngleTolerance = BoundaryEdgeAngleTolerance;
+
+		auto RefineFlagsToConstraintLevel = [](EEdgeRefineFlags Flags)
+		{
+			if (bool(int(Flags) & (int)EEdgeRefineFlags::NoCollapse))
+			{
+				return MeshClusterSimplify::FSimplifyOptions::EConstraintLevel::Fixed;
+			}
+			else if (bool(int(Flags) & (int)EEdgeRefineFlags::NoFlip))
+			{
+				return MeshClusterSimplify::FSimplifyOptions::EConstraintLevel::Constrained;
+			}
+			return MeshClusterSimplify::FSimplifyOptions::EConstraintLevel::Free;
+		};
+		SimplifyOptions.PreserveEdges.Boundary = RefineFlagsToConstraintLevel(MeshBoundaryConstraint);
+		SimplifyOptions.PreserveEdges.PolyGroup = RefineFlagsToConstraintLevel(GroupBoundaryConstraint);
+		SimplifyOptions.PreserveEdges.Material = RefineFlagsToConstraintLevel(MaterialBoundaryConstraint);
+
+		if (bDiscardAttributes)
+		{
+			// if discarding attributes, also discard constraints from the attribute layer
+			SimplifyOptions.PreserveEdges.SetSeamConstraints(MeshClusterSimplify::FSimplifyOptions::EConstraintLevel::Free);
+			SimplifyOptions.PreserveEdges.Material = MeshClusterSimplify::FSimplifyOptions::EConstraintLevel::Free;
+			SimplifyOptions.bTransferAttributes = false;
+		}
+		else
+		{
+			SimplifyOptions.PreserveEdges.SetSeamConstraints(MeshClusterSimplify::FSimplifyOptions::EConstraintLevel::Constrained);
+			// typically don't actually want tangent seams to prevent simplification
+			SimplifyOptions.PreserveEdges.TangentSeam = MeshClusterSimplify::FSimplifyOptions::EConstraintLevel::Free;
+			SimplifyOptions.bTransferAttributes = true;
+		}
+
+		// drive normal seams by the bPreserveSharpEdges flag, rather than the more general bDiscardAttributes
+		SimplifyOptions.PreserveEdges.NormalSeam = bPreserveSharpEdges ? 
+			MeshClusterSimplify::FSimplifyOptions::EConstraintLevel::Constrained :
+			MeshClusterSimplify::FSimplifyOptions::EConstraintLevel::Free;
+		
+		MeshClusterSimplify::Simplify(*OriginalMesh, *ResultMesh, SimplifyOptions);
+	}
 	else // SimplifierType == ESimplifyType::UEStandard
 	{
 		if (!MeshReduction)
diff --git a/Engine/Plugins/Runtime/MeshModelingToolset/Source/ModelingOperatorsEditorOnly/Public/CleaningOps/SimplifyMeshOp.h b/Engine/Plugins/Runtime/MeshModelingToolset/Source/ModelingOperatorsEditorOnly/Public/CleaningOps/SimplifyMeshOp.h
--- a/Engine/Plugins/Runtime/MeshModelingToolset/Source/ModelingOperatorsEditorOnly/Public/CleaningOps/SimplifyMeshOp.h
+++ b/Engine/Plugins/Runtime/MeshModelingToolset/Source/ModelingOperatorsEditorOnly/Public/CleaningOps/SimplifyMeshOp.h
@@ -54,6 +54,9 @@ enum class ESimplifyType : uint8
 	/** Only preserve polygroup boundaries; ignore all other shape features */
 	MinimalPolygroup = 5 UMETA(DisplayName = "Minimal PolyGroup-Preserving"),
 
+	/** Simplify by locally clustering vertices, and re-creating the mesh triangles from the cluster connectivity */
+	ClusterBased = 6,
+
 };
 
 namespace UE
@@ -85,6 +88,9 @@ public:
 	// For minimal polygroup-preserving simplification: Threshold angle change (in degrees) along a polygroup edge, above which a vertex must be added
 	float PolyEdgeAngleTolerance = 0.1f;
 
+	// For cluster-based simplification, equivalent to PolyEdgeAngleTolerance but only for boundary edges. Helpful for preserving open boundary shape.
+	float BoundaryEdgeAngleTolerance = 30.0f;
+
 	bool bGeometricDeviationConstraint = false;
 	float GeometricTolerance = 0.0f;
 
