Moved dll exports to the right place (not on the type).

diff --git a/Engine/Plugins/Runtime/MeshModelingToolset/Source/ModelingComponents/Public/Scene/WorldRenderCapture.h b/Engine/Plugins/Runtime/MeshModelingToolset/Source/ModelingComponents/Public/Scene/WorldRenderCapture.h
--- a/Engine/Plugins/Runtime/MeshModelingToolset/Source/ModelingComponents/Public/Scene/WorldRenderCapture.h
+++ b/Engine/Plugins/Runtime/MeshModelingToolset/Source/ModelingComponents/Public/Scene/WorldRenderCapture.h
@@ -5,6 +5,8 @@
 #include "SceneTypes.h"		// FPrimitiveComponentId
 #include "Image/SpatialPhotoSet.h"
 
+#define UE_API MODELINGCOMPONENTS_API
+
 class UTextureRenderTarget2D;
 class AActor;
 
@@ -46,7 +48,7 @@ void ForEachCaptureType(const FunctionType& Function)
 	Function(ERenderCaptureType::DeviceDepth);
 }
 
-struct MODELINGCOMPONENTS_API FRenderCaptureConfig
+struct FRenderCaptureConfig
 {
 	// You might want to disable this if you're using FMeshMapBaker because it supports multi-sampling in a way which
 	// will avoid blending pixels at different depths. This option is ignored for ERenderCaptureType::DeviceDepth
@@ -68,7 +70,7 @@ FRenderCaptureConfig MODELINGCOMPONENTS_API GetDefaultRenderCaptureConfig(ERende
 /**
  * FRenderCaptureTypeFlags is a set of per-capture-type booleans
  */
-struct MODELINGCOMPONENTS_API FRenderCaptureTypeFlags
+struct FRenderCaptureTypeFlags
 {
 	bool bBaseColor = false;
 	bool bRoughness = false;
@@ -82,36 +84,36 @@ struct MODELINGCOMPONENTS_API FRenderCaptureTypeFlags
 	bool bSubsurfaceColor = false;
 
 	/** @return FRenderCaptureTypeFlags with all types enabled/true */
-	static FRenderCaptureTypeFlags All(bool bCombinedMRS = false);
+	static UE_API FRenderCaptureTypeFlags All(bool bCombinedMRS = false);
 
 	/** @return FRenderCaptureTypeFlags with all types disabled/false */
-	static FRenderCaptureTypeFlags None();
+	static UE_API FRenderCaptureTypeFlags None();
 
 	/** @return FRenderCaptureTypeFlags with only BaseColor enabled/true */
-	static FRenderCaptureTypeFlags BaseColor();
+	static UE_API FRenderCaptureTypeFlags BaseColor();
 
 	/** @return FRenderCaptureTypeFlags with only WorldNormal enabled/true */
-	static FRenderCaptureTypeFlags WorldNormal();
+	static UE_API FRenderCaptureTypeFlags WorldNormal();
 
 	/** @return FRenderCaptureTypeFlags with a single capture type enabled/true */
-	static FRenderCaptureTypeFlags Single(ERenderCaptureType CaptureType);
+	static UE_API FRenderCaptureTypeFlags Single(ERenderCaptureType CaptureType);
 
 	/** Set the indicated CaptureType to enabled */
-	void SetEnabled(ERenderCaptureType CaptureType, bool bEnabled);
+	UE_API void SetEnabled(ERenderCaptureType CaptureType, bool bEnabled);
 
 	/** @return mutable flag corresponding to the given CaptureType */
-	bool& operator[](ERenderCaptureType CaptureType);
+	UE_API bool& operator[](ERenderCaptureType CaptureType);
 
 	/** @return constant flag corresponding to the given CaptureType */
-	const bool& operator[](ERenderCaptureType CaptureType) const;
+	UE_API const bool& operator[](ERenderCaptureType CaptureType) const;
 
 	/** @return true if the flags for this object match the Other object and false otherwise */
-	bool operator==(const FRenderCaptureTypeFlags& Other) const;
-	bool operator!=(const FRenderCaptureTypeFlags& Other) const;
+	UE_API bool operator==(const FRenderCaptureTypeFlags& Other) const;
+	UE_API bool operator!=(const FRenderCaptureTypeFlags& Other) const;
 };
 
 template <typename T>
-struct MODELINGCOMPONENTS_API TRenderCaptureTypeData
+struct TRenderCaptureTypeData
 {
 	// FRenderCaptureTypeFlags uses the correct member naming convention and has additional boolean helpers
 	static_assert(std::is_same_v<T, bool> == false, "Use FRenderCaptureTypeFlags instead");
@@ -211,7 +213,7 @@ FViewMatrices GetRenderCaptureViewMatrices(const FFrame3d& ViewFrame, double Hor
  * specific viewpoint. Various types of rendering are supported, as defined by ERenderCaptureType.
  * Currently rendering an entire World, ie without an explicit list of Actors or Components, is not supported.
  */
-class MODELINGCOMPONENTS_API FWorldRenderCapture
+class FWorldRenderCapture
 {
 public:
 
@@ -221,30 +223,30 @@ PRAGMA_DISABLE_DEPRECATION_WARNINGS // Silence deprecation warnings for deprecat
 	FWorldRenderCapture(FWorldRenderCapture&&) noexcept = default;
 	FWorldRenderCapture& operator=(const FWorldRenderCapture&) = default;
 	FWorldRenderCapture& operator=(FWorldRenderCapture&&) noexcept = default;
-PRAGMA_ENABLE_DEPRECATION_WARNINGS
+UE_API PRAGMA_ENABLE_DEPRECATION_WARNINGS
 
 	~FWorldRenderCapture();
 
 	/** Explicitly release any allocated textures or other data structres */
-	void Shutdown();
+	UE_API void Shutdown();
 
 	/** Set the target World */
-	void SetWorld(UWorld* World);
+	UE_API void SetWorld(UWorld* World);
 
 	/** 
 	 *  Set the set of Actors in the target World that should be included in the Rendering.
 	 */
-	void SetVisibleActors(const TArray<AActor*>& Actors);
+	UE_API void SetVisibleActors(const TArray<AActor*>& Actors);
 
 	/**
 	 *  Set the set of Components in the target World that should be included in the Rendering.
 	 */
-	void SetVisibleComponents(const TArray<UActorComponent*>& Components);
+	UE_API void SetVisibleComponents(const TArray<UActorComponent*>& Components);
 
 	/**
 	 *  Set the set of Actors and Components in the target World that should be included in the Rendering.
 	 */
-	void SetVisibleActorsAndComponents(const TArray<AActor*>& VisibleActors, const TArray<UActorComponent*>& VisibleComponents);
+	UE_API void SetVisibleActorsAndComponents(const TArray<AActor*>& VisibleActors, const TArray<UActorComponent*>& VisibleComponents);
 
 	/** Get bounding-box of the Visible actors */
 	FBoxSphereBounds GetVisibleActorBounds() const { return VisibleBounds; }
@@ -254,10 +256,10 @@ PRAGMA_ENABLE_DEPRECATION_WARNINGS
 	 * will be fully visible (ie a square capture will not have any clipping), for the given Field of View.
 	 * SafetyBoundsScale multiplier is applied to the bounding box.
 	 */
-	FSphere ComputeContainingRenderSphere(float HorzFOVDegrees, float SafetyBoundsScale = 1.25f) const;
+	UE_API FSphere ComputeContainingRenderSphere(float HorzFOVDegrees, float SafetyBoundsScale = 1.25f) const;
 
 	/** Set desired pixel dimensions of the rendered target image */
-	void SetDimensions(const FImageDimensions& Dimensions);
+	UE_API void SetDimensions(const FImageDimensions& Dimensions);
 
 	/** @return pixel dimensions that target image will be rendered at */
 	const FImageDimensions& GetDimensions() const { return Dimensions; }
@@ -280,7 +282,7 @@ PRAGMA_ENABLE_DEPRECATION_WARNINGS
 	 * @param ResultImageOut output image of size GetDimensions() is stored here.
 	 * @return true if capture could be rendered successfully
 	 */
-	bool CaptureFromPosition(
+	UE_API bool CaptureFromPosition(
 		ERenderCaptureType CaptureType,
 		const FFrame3d& ViewFrame,
 		double HorzFOVDegrees,
@@ -293,7 +295,7 @@ PRAGMA_ENABLE_DEPRECATION_WARNINGS
 	 * If FolderName is not specified, "WorldRenderCapture" is used by default.
 	 * If an ImageCounter is not specified, an internal static counter is used that increments on every write
 	 */
-	void SetEnableWriteDebugImage(bool bEnable, int32 ImageCounter = -1, FString FolderName = FString());
+	UE_API void SetEnableWriteDebugImage(bool bEnable, int32 ImageCounter = -1, FString FolderName = FString());
 
 protected:
 	UWorld* World = nullptr;
@@ -309,8 +311,8 @@ protected:
 	UTextureRenderTarget2D* DepthRenderTexture = nullptr;
 
 	FImageDimensions RenderTextureDimensions;
-	UTextureRenderTarget2D* GetRenderTexture(bool bLinear);
-	UTextureRenderTarget2D* GetDepthRenderTexture();
+	UE_API UTextureRenderTarget2D* GetRenderTexture(bool bLinear);
+	UE_API UTextureRenderTarget2D* GetDepthRenderTexture();
 
 	UE_DEPRECATED(5.4, "Please use GetRenderCaptureViewMatrices instead")
 	FViewMatrices LastCaptureViewMatrices;
@@ -320,7 +322,7 @@ protected:
 
 	/** Emissive is a special case and uses different code than capture of color/property channels */
 	UE_DEPRECATED(5.4, "Please use CaptureBufferVisualizationFromPosition with VisualizationMode=PreTonemapHDRColor instead")
-	bool CaptureEmissiveFromPosition(
+	UE_API bool CaptureEmissiveFromPosition(
 		const FFrame3d& Frame,
 		double HorzFOVDegrees,
 		double NearPlaneDist,
@@ -328,7 +330,7 @@ protected:
 		const FRenderCaptureConfig& Config = {});
 
 	/** Combined Metallic/Roughness/Specular uses a custom postprocess material */
-	bool CaptureMRSFromPosition(
+	UE_API bool CaptureMRSFromPosition(
 		const FFrame3d& Frame,
 		double HorzFOVDegrees,
 		double NearPlaneDist,
@@ -336,7 +338,7 @@ protected:
 		const FRenderCaptureConfig& Config = {});
 
 	/** Depth is a special case and uses different code than capture of color/property channels */
-	bool CaptureDeviceDepthFromPosition(
+	UE_API bool CaptureDeviceDepthFromPosition(
 		const FFrame3d& Frame,
 		double HorzFOVDegrees,
 		double NearPlaneDist,
@@ -350,7 +352,7 @@ protected:
 	 * Its not clear where the valid VisualizationMode FNames are defined, it may be in the [Engine.BufferVisualizationMaterials] section of
 	 * Config/BaseEngine.ini, which is read by FBufferVisualizationData::Initialize() and called when the engine initializes in UEngine::Init()
 	 */
-	bool CaptureBufferVisualizationFromPosition(
+	UE_API bool CaptureBufferVisualizationFromPosition(
 		const FName& VisualizationMode,
 		const FFrame3d& Frame,
 		double HorzFOVDegrees,
@@ -361,7 +363,7 @@ protected:
 	bool bWriteDebugImage = false;
 	int32 DebugImageCounter = -1;
 	FString DebugImageFolderName = TEXT("WorldRenderCapture");
-	void WriteDebugImage(const FImageAdapter& ResultImageOut, const FString& ImageTypeName);
+	UE_API void WriteDebugImage(const FImageAdapter& ResultImageOut, const FString& ImageTypeName);
 };
 
 
@@ -370,3 +372,5 @@ protected:
 
 } // end namespace UE::Geometry
 } // end namespace UE
+
+#undef UE_API
diff --git a/Engine/Plugins/Runtime/Metasound/Source/MetasoundFrontend/Public/Analysis/MetasoundFrontendVertexAnalyzerForwardValue.h b/Engine/Plugins/Runtime/Metasound/Source/MetasoundFrontend/Public/Analysis/MetasoundFrontendVertexAnalyzerForwardValue.h
--- a/Engine/Plugins/Runtime/Metasound/Source/MetasoundFrontend/Public/Analysis/MetasoundFrontendVertexAnalyzerForwardValue.h
+++ b/Engine/Plugins/Runtime/Metasound/Source/MetasoundFrontend/Public/Analysis/MetasoundFrontendVertexAnalyzerForwardValue.h
@@ -89,7 +89,7 @@ namespace Metasound
 		};
 
 		#define METASOUND_DECLARE_FORWARD_VALUE_ANALYZER(ClassType, AnalyzerFName, DataType) \
-			class METASOUNDFRONTEND_API ClassType : public TVertexAnalyzerForwardValue<ClassType, DataType> \
+			class ClassType : public TVertexAnalyzerForwardValue<ClassType, DataType> \
 			{ \
 			public: \
 				static const FName& GetAnalyzerName() \
diff --git a/Engine/Plugins/Runtime/Metasound/Source/MetasoundFrontend/Public/MetasoundEnum.h b/Engine/Plugins/Runtime/Metasound/Source/MetasoundFrontend/Public/MetasoundEnum.h
--- a/Engine/Plugins/Runtime/Metasound/Source/MetasoundFrontend/Public/MetasoundEnum.h
+++ b/Engine/Plugins/Runtime/Metasound/Source/MetasoundFrontend/Public/MetasoundEnum.h
@@ -97,7 +97,7 @@ namespace Metasound
 	/** Metasound Enum String Helper
 	 */
 	template<typename T>
-	struct METASOUNDFRONTEND_API TEnumStringHelper : TEnumStringHelperBase<TEnumStringHelper<T>, T>
+	struct TEnumStringHelper : TEnumStringHelperBase<TEnumStringHelper<T>, T>
 	{
 		static_assert(TIsEnum<T>::Value, "Please define a specialization of this class. The DECLARE_METASOUND_ENUM macros will do this for you");
 	};
@@ -303,12 +303,12 @@ namespace Metasound
 #define DECLARE_METASOUND_ENUM(ENUMNAME, DEFAULT, API, ENUMTYPEDEF, TYPEINFO, READREF, WRITEREF)\
 	using ENUMTYPEDEF = Metasound::TEnum<ENUMNAME, DEFAULT>; \
 	DECLARE_METASOUND_DATA_REFERENCE_TYPES(ENUMTYPEDEF, API, TYPEINFO, READREF, WRITEREF);\
-	template<> struct API Metasound::TEnumStringHelper<ENUMNAME> : Metasound::TEnumStringHelperBase<Metasound::TEnumStringHelper<ENUMNAME>, ENUMNAME>\
+	template<> struct Metasound::TEnumStringHelper<ENUMNAME> : Metasound::TEnumStringHelperBase<Metasound::TEnumStringHelper<ENUMNAME>, ENUMNAME>\
 	{\
 		static FName GetNamespace()\
 		{\
 			static const FName ThisName { TEXT(#ENUMNAME) };\
 			return ThisName;\
 		}\
-		static TArrayView<const Metasound::TEnumEntry<ENUMNAME>> GetAllEntries();\
+		static API TArrayView<const Metasound::TEnumEntry<ENUMNAME>> GetAllEntries();\
 	};
diff --git a/Engine/Source/Runtime/Core/Public/FramePro/FramePro.h b/Engine/Source/Runtime/Core/Public/FramePro/FramePro.h
--- a/Engine/Source/Runtime/Core/Public/FramePro/FramePro.h
+++ b/Engine/Source/Runtime/Core/Public/FramePro/FramePro.h
@@ -1607,9 +1607,9 @@ namespace FramePro
 
 	FRAMEPRO_API extern RelaxedAtomic<unsigned int> g_ConditionalScopeMinTime;
 
-	FRAMEPRO_API inline bool IsConnected() { return g_Connected; }
+	inline bool IsConnected() { return g_Connected; }
 
-	FRAMEPRO_API inline unsigned int GetConditionalScopeMinTime() { return g_ConditionalScopeMinTime; }
+	inline unsigned int GetConditionalScopeMinTime() { return g_ConditionalScopeMinTime; }
 
 	FRAMEPRO_API void Shutdown();
 
diff --git a/Engine/Source/Runtime/Experimental/ChaosVehicles/ChaosVehiclesCore/Public/SimModule/ModuleInput.h b/Engine/Source/Runtime/Experimental/ChaosVehicles/ChaosVehiclesCore/Public/SimModule/ModuleInput.h
--- a/Engine/Source/Runtime/Experimental/ChaosVehicles/ChaosVehiclesCore/Public/SimModule/ModuleInput.h
+++ b/Engine/Source/Runtime/Experimental/ChaosVehicles/ChaosVehiclesCore/Public/SimModule/ModuleInput.h
@@ -9,6 +9,7 @@
 
 DECLARE_LOG_CATEGORY_EXTERN(LogModularInput, Log, All);
 
+#define UE_API CHAOSVEHICLESCORE_API
 
 UENUM(BlueprintType)
 enum class EModuleInputValueType : uint8
@@ -166,7 +167,7 @@ bool WriteCompressedFloat(const T Value, FArchive& Ar)
 }
 
 USTRUCT(BlueprintType)
-struct CHAOSVEHICLESCORE_API FModuleInputValue
+struct FModuleInputValue
 {
 	GENERATED_BODY()
 
@@ -346,7 +347,7 @@ public:
 	float operator[](int32 Index) const { return Value[Index]; }
 
 	bool IsQuantizedNonZero(float Tolerance = KINDA_SMALL_NUMBER) const { return ReturnQuantized().IsNonZero(Tolerance); }
-	bool IsNonZero(float Tolerance = KINDA_SMALL_NUMBER) const;
+	UE_API bool IsNonZero(float Tolerance = KINDA_SMALL_NUMBER) const;
 
 	// In-place type conversion
 	FModuleInputValue& ConvertToType(EModuleInputValueType Type)
@@ -361,23 +362,23 @@ public:
 
 	EModuleInputValueType GetValueType() const { return ValueType; }
 
-	float GetMagnitudeSq() const;
-	float GetMagnitude() const;
-	int32 GetMagnitudeInt() const;
+	UE_API float GetMagnitudeSq() const;
+	UE_API float GetMagnitude() const;
+	UE_API int32 GetMagnitudeInt() const;
 
 	// Serialize values
-	void Serialize(FArchive& Ar, UPackageMap* Map, bool& bOutSuccess);
-	void NetSerialize(FArchive& Ar, UPackageMap* Map, bool& bOutSuccess);
-	void DeltaNetSerialize(FArchive& Ar, UPackageMap* Map, bool& bOutSuccess,const FModuleInputValue& PreviousInputValue);
+	UE_API void Serialize(FArchive& Ar, UPackageMap* Map, bool& bOutSuccess);
+	UE_API void NetSerialize(FArchive& Ar, UPackageMap* Map, bool& bOutSuccess);
+	UE_API void DeltaNetSerialize(FArchive& Ar, UPackageMap* Map, bool& bOutSuccess,const FModuleInputValue& PreviousInputValue);
 
-	void Lerp(const FModuleInputValue& Min, const FModuleInputValue& Max, float Alpha);
+	UE_API void Lerp(const FModuleInputValue& Min, const FModuleInputValue& Max, float Alpha);
 
-	void Merge(const FModuleInputValue& From);
+	UE_API void Merge(const FModuleInputValue& From);
 
-	void Decay(const float DecayAmount);
+	UE_API void Decay(const float DecayAmount);
 
 	// Type sensitive debug stringify
-	FString ToString() const;
+	UE_API FString ToString() const;
 
 	/** During physics resimulation, apply decay for this input while it's being extrapolated */
 	void SetApplyInputDecay(const bool bInApplyInputDecay) { bApplyInputDecay = bInApplyInputDecay; }
@@ -432,7 +433,7 @@ inline FModuleInputValue::MInteger FModuleInputValue::Get() const
 }
 
 
-class CHAOSVEHICLESCORE_API FModuleInputConversion
+class FModuleInputConversion
 {
 public:
 
@@ -473,8 +474,8 @@ public:
 };
 
 
-UCLASS(BlueprintType, Blueprintable)
-class CHAOSVEHICLESCORE_API UDefaultModularVehicleInputModifier : public UObject
+UCLASS(BlueprintType, Blueprintable, MinimalAPI)
+class UDefaultModularVehicleInputModifier : public UObject
 {
 	GENERATED_BODY()
 
@@ -513,13 +514,13 @@ public:
 	//FRuntimeFloatCurve UserCurve;
 
 	/** Change an output value using max rise and fall rates */
-	virtual FModuleInputValue InterpInputValue(float DeltaTime, const FModuleInputValue& CurrentValue, const FModuleInputValue& NewValue) const;
-	virtual float CalcControlFunction(float InputValue);
+	UE_API virtual FModuleInputValue InterpInputValue(float DeltaTime, const FModuleInputValue& CurrentValue, const FModuleInputValue& NewValue) const;
+	UE_API virtual float CalcControlFunction(float InputValue);
 };
 
 
 USTRUCT(BlueprintType)
-struct CHAOSVEHICLESCORE_API FModuleInputSetup
+struct FModuleInputSetup
 {
 	GENERATED_BODY()
 
@@ -555,7 +556,7 @@ struct CHAOSVEHICLESCORE_API FModuleInputSetup
 	bool bApplyInputDecay;
 };
 
-class CHAOSVEHICLESCORE_API FScopedModuleInputInitializer
+class FScopedModuleInputInitializer
 {
 public:
 	FScopedModuleInputInitializer(TArray<struct FModuleInputSetup>& InSetupData)
@@ -573,12 +574,12 @@ public:
 	
 private:
 
-	static TArray<struct FModuleInputSetup>* InitSetupData;
+	static UE_API TArray<struct FModuleInputSetup>* InitSetupData;
 };
 
 
 USTRUCT()
-struct CHAOSVEHICLESCORE_API FModuleInputContainer
+struct FModuleInputContainer
 {
 	GENERATED_BODY()
 
@@ -624,22 +625,22 @@ public:
 		return *this;
 	}
 
-	void Initialize(TArray<FModuleInputSetup>& SetupData, FInputNameMap& NameMapOut);
+	UE_API void Initialize(TArray<FModuleInputSetup>& SetupData, FInputNameMap& NameMapOut);
 
-	void ZeroValues();
+	UE_API void ZeroValues();
 
-	void Serialize(FArchive& Ar, UPackageMap* Map, bool& bOutSuccess);
+	UE_API void Serialize(FArchive& Ar, UPackageMap* Map, bool& bOutSuccess);
 
-	int AddInput(EModuleInputValueType Type, TSubclassOf<UDefaultModularVehicleInputModifier>& InputModifierClass);
+	UE_API int AddInput(EModuleInputValueType Type, TSubclassOf<UDefaultModularVehicleInputModifier>& InputModifierClass);
 
-	void RemoveAllInputs();
+	UE_API void RemoveAllInputs();
 
-	void Lerp(const FModuleInputContainer& Min, const FModuleInputContainer& Max, float Alpha);
+	UE_API void Lerp(const FModuleInputContainer& Min, const FModuleInputContainer& Max, float Alpha);
 
-	void Merge(const FModuleInputContainer& From);
+	UE_API void Merge(const FModuleInputContainer& From);
 
 	/**  Decay input during resimulation by DecayAmount which increases over resimulation frames from 0.0 -> 1.0 when the input is being reused */
-	void Decay(const float DecayAmount);
+	UE_API void Decay(const float DecayAmount);
 
 	TArray<FModuleInputValue>& AccessInputValues() { return InputValues; }
 
@@ -650,7 +651,7 @@ private:
 };
 
 
-class CHAOSVEHICLESCORE_API FInputInterface
+class FInputInterface
 {
 public:
 
@@ -662,13 +663,13 @@ public:
 	{
 	}
 
-	void SetValue(const FName& InName, const FModuleInputValue& InValue, bool Quantize = true);
-	void MergeValue(const FName& InName, const FModuleInputValue& InValue);
-	FModuleInputValue GetValue(const FName& InName) const;
-	EModuleInputValueType GetValueType(const FName& InName) const;
-	float GetMagnitude(const FName& InName) const;
-	int32 GetMagnitudeInt(const FName& InName) const;
-	bool InputsNonZero() const;
+	UE_API void SetValue(const FName& InName, const FModuleInputValue& InValue, bool Quantize = true);
+	UE_API void MergeValue(const FName& InName, const FModuleInputValue& InValue);
+	UE_API FModuleInputValue GetValue(const FName& InName) const;
+	UE_API EModuleInputValueType GetValueType(const FName& InName) const;
+	UE_API float GetMagnitude(const FName& InName) const;
+	UE_API int32 GetMagnitudeInt(const FName& InName) const;
+	UE_API bool InputsNonZero() const;
 
 	// Quick Access to data type
 	bool GetBool(const FName& InName) const
@@ -731,8 +732,8 @@ public:
 	FModuleInputContainer& ValueContainer;	// per vehicle instance
 };
 
-UCLASS(Abstract, BlueprintType, Blueprintable, EditInlineNew)
-class CHAOSVEHICLESCORE_API UVehicleInputProducerBase : public UObject
+UCLASS(Abstract, BlueprintType, Blueprintable, EditInlineNew, MinimalAPI)
+class UVehicleInputProducerBase : public UObject
 {
 	GENERATED_BODY()
 	
@@ -756,3 +757,4 @@ public:
 
 };
 
+#undef UE_API 
\ No newline at end of file
diff --git a/Engine/Source/Runtime/Experimental/Iris/Core/Public/Iris/ReplicationSystem/NetTokenStructDefines.h b/Engine/Source/Runtime/Experimental/Iris/Core/Public/Iris/ReplicationSystem/NetTokenStructDefines.h
--- a/Engine/Source/Runtime/Experimental/Iris/Core/Public/Iris/ReplicationSystem/NetTokenStructDefines.h
+++ b/Engine/Source/Runtime/Experimental/Iris/Core/Public/Iris/ReplicationSystem/NetTokenStructDefines.h
@@ -31,7 +31,7 @@ struct TNetTokenStructOpsTypeTraits : public TStructOpsTypeTraitsBase2<T>
 // Optional, declares a default native NetSerializer method, GetTokenStoreName method and Equality operators (for equality via identity) for a NetToken Struct type
 #define UE_NET_NETTOKEN_GENERATED_BODY(NAME, API) \
 	inline static FName TokenStoreName = TEXT( PREPROCESSOR_TO_STRING(F##NAME) );\
-	API static FName GetTokenStoreName() \
+	static FName GetTokenStoreName() \
 	{ \
 		return TokenStoreName; \
 	} \
diff --git a/Engine/Source/Runtime/Slate/Private/Framework/Commands/GenericCommands.cpp b/Engine/Source/Runtime/Slate/Private/Framework/Commands/GenericCommands.cpp
--- a/Engine/Source/Runtime/Slate/Private/Framework/Commands/GenericCommands.cpp
+++ b/Engine/Source/Runtime/Slate/Private/Framework/Commands/GenericCommands.cpp
@@ -4,6 +4,8 @@
 
 #define LOCTEXT_NAMESPACE "GenericCommands"
 
+UE_DEFINE_COMMANDS_TLS(FGenericCommands)
+
 void FGenericCommands::RegisterCommands()
 {
 	UI_COMMAND(Cut, "Cut", "Cut selection", EUserInterfaceActionType::Button, FInputChord(EModifierKey::Control, EKeys::X))
diff --git a/Engine/Source/Runtime/Slate/Public/Framework/Application/SlateApplication.h b/Engine/Source/Runtime/Slate/Public/Framework/Application/SlateApplication.h
--- a/Engine/Source/Runtime/Slate/Public/Framework/Application/SlateApplication.h
+++ b/Engine/Source/Runtime/Slate/Public/Framework/Application/SlateApplication.h
@@ -826,7 +826,7 @@ public:
 	 * @return True if the application is currently routing high precision mouse movement events (OS specific)
 	 * If this value is true, the mouse is captured and hidden by the widget that originally made the request.
 	 */
-	SLATE_API bool IsUsingHighPrecisionMouseMovment() const { return PlatformApplication.IsValid() ? PlatformApplication->IsUsingHighPrecisionMouseMode() : false; }
+	bool IsUsingHighPrecisionMouseMovment() const { return PlatformApplication.IsValid() ? PlatformApplication->IsUsingHighPrecisionMouseMode() : false; }
 	
 	/**
 	 * @return True if the last mouse event was from a trackpad.
@@ -915,10 +915,10 @@ public:
 	SLATE_API void LeaveDebuggingMode( bool bLeavingDebugForSingleStep = false );
 
 #if WITH_EDITOR
-	struct SLATE_API FScopedPreventDebuggingMode
+	struct FScopedPreventDebuggingMode
 	{
-		UE_NODISCARD_CTOR FScopedPreventDebuggingMode(FText Reason);
-		~FScopedPreventDebuggingMode();
+		UE_NODISCARD_CTOR SLATE_API FScopedPreventDebuggingMode(FText Reason);
+		SLATE_API ~FScopedPreventDebuggingMode();
 	private:
 		int32 Id;
 	};
diff --git a/Engine/Source/Runtime/Slate/Public/Framework/Commands/Commands.h b/Engine/Source/Runtime/Slate/Public/Framework/Commands/Commands.h
--- a/Engine/Source/Runtime/Slate/Public/Framework/Commands/Commands.h
+++ b/Engine/Source/Runtime/Slate/Public/Framework/Commands/Commands.h
@@ -24,6 +24,9 @@ SLATE_API void MakeUICommand_InternalUseOnly(FBindingContext* This, TSharedPtr<
 
 #undef LOC_DEFINE_REGION
 
+#define UE_DECLARE_COMMANDS_TLS(Type, Api) template<> Api TWeakPtr<Type>& TCommands<Type>::GetInstance();
+#define UE_DEFINE_COMMANDS_TLS(Type) template<> TWeakPtr<Type>& TCommands<Type>::GetInstance() { static TWeakPtr<Type> Inst; return Inst; }
+
 /** A base class for a set of commands. Inherit from it to make a set of commands. See MainFrameActions for an example. */
 template<typename CommandContextType>
 class TCommands : public FBindingContext
@@ -32,7 +35,7 @@ public:
 	
 
 	/** Use this method to register commands. Usually done in StartupModule(). */
-	FORCENOINLINE static void Register()
+	inline static void Register()
 	{
 		if ( !Instance.IsValid() )
 		{
@@ -60,13 +63,13 @@ public:
 		}
 	}
 
-	FORCENOINLINE static bool IsRegistered()
+	inline static bool IsRegistered()
 	{
 		return Instance.IsValid();
 	}
 
 	/** Get the singleton instance of this set of commands. */
-	FORCENOINLINE static const CommandContextType& Get()
+	inline static const CommandContextType& Get()
 	{
 		// If you get a crash here, then either your commands object wasn't registered, or it was not kept
 		//  alive. The latter could happen if your RegisterCommands() function does not actually register 
@@ -75,7 +78,7 @@ public:
 	}
 
 	/** Use this method to clean up any resources used by the command set. Usually done in ShutdownModule() */
-	FORCENOINLINE static void Unregister()
+	inline static void Unregister()
 	{
 		// The instance may not be valid if it was never used.
 		if( Instance.IsValid() )
@@ -92,7 +95,7 @@ public:
 	}
 
 	/** Get the BindingContext for this set of commands. */
-	FORCENOINLINE static const FBindingContext& GetContext()
+	inline static const FBindingContext& GetContext()
 	{
 		check( Instance.IsValid() );
 		return *(Instance.Pin());
@@ -111,15 +114,25 @@ protected:
 
 	/** A static instance of the command set. */
 #if PLATFORM_UNIX || PLATFORM_APPLE
-	static SLATE_API TWeakPtr< CommandContextType > Instance;
+	static SLATE_API TWeakPtr<CommandContextType>& GetInstance()
+	{
+		static TWeakPtr<CommandContextType> Inst;
+		return Inst;
+	}
 #else
-	static TWeakPtr< CommandContextType > Instance;
+	static TWeakPtr<CommandContextType>& GetInstance()
+	{
+		static TWeakPtr<CommandContextType> Inst;
+		return Inst;
+	}
 #endif
 
+	static TWeakPtr<CommandContextType>& Instance;
+
 	/** Pure virtual to override; describe and instantiate the commands in here by using the UI COMMAND macro. */
 	virtual void RegisterCommands() = 0;
 
 };
 
 template<typename CommandContextType>
-TWeakPtr< CommandContextType > TCommands<CommandContextType>::Instance = NULL;
+TWeakPtr<CommandContextType>& TCommands<CommandContextType>::Instance = TCommands<CommandContextType>::GetInstance();
diff --git a/Engine/Source/Runtime/Slate/Public/Framework/Commands/GenericCommands.h b/Engine/Source/Runtime/Slate/Public/Framework/Commands/GenericCommands.h
--- a/Engine/Source/Runtime/Slate/Public/Framework/Commands/GenericCommands.h
+++ b/Engine/Source/Runtime/Slate/Public/Framework/Commands/GenericCommands.h
@@ -6,7 +6,9 @@
 #include "Styling/CoreStyle.h"
 #include "Framework/Commands/Commands.h"
 
-class SLATE_API FGenericCommands : public TCommands<FGenericCommands>
+UE_DECLARE_COMMANDS_TLS(class FGenericCommands, SLATE_API)
+
+class FGenericCommands : public TCommands<FGenericCommands>
 {
 public:
 	
@@ -19,7 +21,7 @@ public:
 	{
 	}
 
-	virtual void RegisterCommands() override;	
+	SLATE_API virtual void RegisterCommands() override;	
 
 	TSharedPtr<FUICommandInfo> Cut;
 	TSharedPtr<FUICommandInfo> Copy;
