Skeleton Editing Tool: only merge the bones tree when the skeleton is referenced by some other skeletal meshes
- we assume that the user is fully editing the skeleton when there's a one-to-one relationship

TODO: we need to think about what to do when there are several references and have a shared pattern with interchange

#rnx

diff --git a/Engine/Plugins/Runtime/MeshModelingToolset/Source/SkeletalMeshModifiers/Private/SkeletonModifier.cpp b/Engine/Plugins/Runtime/MeshModelingToolset/Source/SkeletalMeshModifiers/Private/SkeletonModifier.cpp
--- a/Engine/Plugins/Runtime/MeshModelingToolset/Source/SkeletalMeshModifiers/Private/SkeletonModifier.cpp
+++ b/Engine/Plugins/Runtime/MeshModelingToolset/Source/SkeletalMeshModifiers/Private/SkeletonModifier.cpp
@@ -560,7 +560,7 @@ ESkeletonModificationType USkeletonModifier::PreCommitSkeleton(const ESkeletalMe
 		});
 
 	const int32 Choice = OptionsDialog->ShowModal();
-	if (Choice == 1)
+	if (Choice == 1 || Choice == -1)
 	{
 		return ESkeletonModificationType::Cancel;
 	}
@@ -597,7 +597,7 @@ bool USkeletonModifier::CommitSkeletonToSkeletalMesh()
 	}
 
 	// check modifications at the skeleton level
-	const ESkeletonModificationType SkeletonModifications = PreCommitSkeleton(Modifications);
+	ESkeletonModificationType SkeletonModifications = PreCommitSkeleton(Modifications);
 	if (SkeletonModifications == ESkeletonModificationType::Cancel)
 	{
 		UE_LOG(LogAnimation, Warning, TEXT("Skeleton Modifier: Skeleton can't be modified."));
@@ -674,12 +674,17 @@ bool USkeletonModifier::CommitSkeletonToSkeletalMesh()
 	{
 		NotifyFromSkeletonChanges();
 		
-		auto UpdateSkeleton = [this, SkeletonModifications, &RetargetingModes]()
+		auto UpdateSkeleton = [this, &SkeletonModifications, &RetargetingModes]()
 		{
 			USkeleton* Skeleton = SkeletalMesh->GetSkeleton();
 			if (SkeletonModifications == ESkeletonModificationType::SimpleMerge)
 			{
-				return Skeleton->MergeAllBonesToBoneTree(SkeletalMesh.Get());
+				constexpr bool bSkeletalMeshReferencesOnly = false;
+				if (!HasAnyOtherReferences(bSkeletalMeshReferencesOnly))
+				{
+					return Skeleton->MergeAllBonesToBoneTree(SkeletalMesh.Get());
+				}
+				SkeletonModifications = ESkeletonModificationType::FullMerge;
 			}
 		
 			bool bSkeletonModified = false;	
@@ -845,6 +850,47 @@ void USkeletonModifier::PostCommitSkeleton(const ESkeletonModificationType InSke
 	}
 }
 
+bool USkeletonModifier::HasAnyOtherReferences(const bool bSkeletalMeshOnly) const
+{
+	USkeleton* Skeleton = SkeletalMesh->GetSkeleton();
+	UPackage* SkeletonPackage = Skeleton->GetPackage();
+	UPackage* SkeletalMeshPackage = SkeletalMesh->GetPackage();
+
+	if (SkeletonPackage && SkeletalMeshPackage)
+	{
+		TArray<FSoftObjectPath> Result;
+		FAssetRegistryModule& AssetRegistryModule = FModuleManager::LoadModuleChecked<FAssetRegistryModule>(TEXT("AssetRegistry"));
+		IAssetRegistry& AssetRegistry = AssetRegistryModule.GetRegistry();
+
+		TArray<FName> References;
+		AssetRegistry.GetReferencers(SkeletonPackage->GetFName(), References);
+
+		const FName SkeletalMeshPackageName = SkeletalMeshPackage->GetFName();
+		
+		TArray<FName> Dependencies;
+		for (const FName& Reference : References)
+		{
+			if (SkeletalMeshPackageName != Reference)
+			{
+				const FString PackageString = Reference.ToString();
+				const FSoftObjectPath FullAssetPath(FTopLevelAssetPath(*PackageString, *FPackageName::GetLongPackageAssetName(PackageString)));
+				FAssetData AssetData = AssetRegistry.GetAssetByObjectPath(FullAssetPath);
+				if (UClass* Class = AssetData.GetClass())
+				{
+					if (!bSkeletalMeshOnly || Class == USkeletalMesh::StaticClass())
+					{
+						Dependencies.AddUnique(*AssetData.GetObjectPathString());
+					}
+				}
+			}
+		}
+		
+		return !Dependencies.IsEmpty();
+	}
+	
+	return false;
+}
+
 void USkeletonModifier::NotifyFromSkeletonChanges() const
 {
 	// check assets using this skeleton ?
diff --git a/Engine/Plugins/Runtime/MeshModelingToolset/Source/SkeletalMeshModifiers/Public/SkeletonModifier.h b/Engine/Plugins/Runtime/MeshModelingToolset/Source/SkeletalMeshModifiers/Public/SkeletonModifier.h
--- a/Engine/Plugins/Runtime/MeshModelingToolset/Source/SkeletalMeshModifiers/Public/SkeletonModifier.h
+++ b/Engine/Plugins/Runtime/MeshModelingToolset/Source/SkeletalMeshModifiers/Public/SkeletonModifier.h
@@ -268,6 +268,9 @@ private:
 	ESkeletonModificationType PreCommitSkeleton(const ESkeletalMeshModificationType InSkeletalMeshModifications) const;
 	void PostCommitSkeleton(const ESkeletonModificationType InSkeletonModifications) const;
 
+	// does the skeleton have any references other than the skeletal mesh currently being edited
+	bool HasAnyOtherReferences(const bool bSkeletalMeshOnly) const;
+
 	// update mesh description on commit
 	void CommitChangesToMeshDescription(const ESkeletalMeshModificationType InSkeletalMeshModifications);
 	
