Modeling Mode: Deprecate SupportsVariableSpacing() on sculpting brushes for a new SupportsStrokeType(), which supports more stroke behaviors in a future looking way.

diff --git a/Engine/Plugins/Runtime/MeshModelingToolset/Source/MeshModelingTools/Private/Sculpting/MeshSculptToolBase.cpp b/Engine/Plugins/Runtime/MeshModelingToolset/Source/MeshModelingTools/Private/Sculpting/MeshSculptToolBase.cpp
--- a/Engine/Plugins/Runtime/MeshModelingToolset/Source/MeshModelingTools/Private/Sculpting/MeshSculptToolBase.cpp
+++ b/Engine/Plugins/Runtime/MeshModelingToolset/Source/MeshModelingTools/Private/Sculpting/MeshSculptToolBase.cpp
@@ -866,7 +866,7 @@ void UMeshSculptToolBase::UpdateStampPendingState()
 	}
 
 	// if brush does not support variable spacing, we only use flow
-	if (GetActiveBrushOp()->SupportsVariableSpacing() == false)
+	if (GetActiveBrushOp()->SupportsStrokeType(EMeshSculptStrokeType::Spacing) == false)
 	{
 		bIsStampPending = bFlowStampPending;
 		return;
diff --git a/Engine/Plugins/Runtime/MeshModelingToolset/Source/MeshModelingTools/Public/Sculpting/MeshBrushOpBase.h b/Engine/Plugins/Runtime/MeshModelingToolset/Source/MeshModelingTools/Public/Sculpting/MeshBrushOpBase.h
--- a/Engine/Plugins/Runtime/MeshModelingToolset/Source/MeshModelingTools/Public/Sculpting/MeshBrushOpBase.h
+++ b/Engine/Plugins/Runtime/MeshModelingToolset/Source/MeshModelingTools/Public/Sculpting/MeshBrushOpBase.h
@@ -105,7 +105,16 @@ public:
 	bool bIsStrengthPressureEnabled = true;
 };
 
+/** Mesh Sculpting Brush Stroke Behaviors */
+UENUM()
+enum class EMeshSculptStrokeType : uint8
+{
+	Spacing = 0,
+	Airbrush = 1,
+	Dots = 2,
 
+	LastValue UMETA(Hidden)
+};
 
 class FMeshSculptBrushOp
 {
@@ -257,10 +266,25 @@ public:
 		return false;
 	}
 
+	UE_DEPRECATED(5.7, "Deprecated in favor of SupportsStrokeType to provide more fine grained description of what strokes brushes support.")
 	virtual bool SupportsVariableSpacing() const
 	{
 		return false;
 	}
+
+	virtual bool SupportsStrokeType(EMeshSculptStrokeType StrokeBehavior) const
+	{
+		switch (StrokeBehavior)
+		{
+		case EMeshSculptStrokeType::Airbrush:
+		case EMeshSculptStrokeType::Dots:
+			return true;
+		case EMeshSculptStrokeType::Spacing:
+			return false;
+		default:
+			return false;
+		}
+	} 
 };
 
 
diff --git a/Engine/Plugins/Runtime/MeshModelingToolset/Source/MeshModelingTools/Public/Sculpting/MeshSculptBrushOps.h b/Engine/Plugins/Runtime/MeshModelingToolset/Source/MeshModelingTools/Public/Sculpting/MeshSculptBrushOps.h
--- a/Engine/Plugins/Runtime/MeshModelingToolset/Source/MeshModelingTools/Public/Sculpting/MeshSculptBrushOps.h
+++ b/Engine/Plugins/Runtime/MeshModelingToolset/Source/MeshModelingTools/Public/Sculpting/MeshSculptBrushOps.h
@@ -51,6 +51,20 @@ public:
 	{
 		return true;
 	}
+
+	virtual bool SupportsStrokeType(EMeshSculptStrokeType StrokeBehavior) const override
+	{
+		switch (StrokeBehavior)
+		{
+		case EMeshSculptStrokeType::Airbrush:
+		case EMeshSculptStrokeType::Dots:
+		case EMeshSculptStrokeType::Spacing:
+			return true;
+		default:
+			return false;
+		}
+	}
+
 	virtual bool UsesAlpha() const override
 	{
 		return true;
@@ -139,6 +153,18 @@ public:
 	{
 		return true;
 	}
+	virtual bool SupportsStrokeType(EMeshSculptStrokeType StrokeBehavior) const override
+	{
+		switch (StrokeBehavior)
+		{
+		case EMeshSculptStrokeType::Airbrush:
+		case EMeshSculptStrokeType::Dots:
+		case EMeshSculptStrokeType::Spacing:
+			return true;
+		default:
+			return false;
+		}
+	}
 	virtual bool UsesAlpha() const override
 	{
 		return true;
@@ -237,6 +263,18 @@ public:
 	{
 		return true;
 	}
+	virtual bool SupportsStrokeType(EMeshSculptStrokeType StrokeBehavior) const override
+	{
+		switch (StrokeBehavior)
+		{
+		case EMeshSculptStrokeType::Airbrush:
+		case EMeshSculptStrokeType::Dots:
+		case EMeshSculptStrokeType::Spacing:
+			return true;
+		default:
+			return false;
+		}
+	}
 	virtual bool UsesAlpha() const override
 	{
 		return true;
diff --git a/Engine/Plugins/Runtime/MeshModelingToolset/Source/MeshModelingTools/Public/Sculpting/MeshVertexPaintBrushOps.h b/Engine/Plugins/Runtime/MeshModelingToolset/Source/MeshModelingTools/Public/Sculpting/MeshVertexPaintBrushOps.h
--- a/Engine/Plugins/Runtime/MeshModelingToolset/Source/MeshModelingTools/Public/Sculpting/MeshVertexPaintBrushOps.h
+++ b/Engine/Plugins/Runtime/MeshModelingToolset/Source/MeshModelingTools/Public/Sculpting/MeshVertexPaintBrushOps.h
@@ -267,6 +267,24 @@ public:
 		});
 
 	}
+
+	virtual bool SupportsVariableSpacing() const override
+	{
+		return true;
+	}
+
+	virtual bool SupportsStrokeType(EMeshSculptStrokeType StrokeBehavior) const override
+	{
+		switch (StrokeBehavior)
+		{
+		case EMeshSculptStrokeType::Airbrush:
+		case EMeshSculptStrokeType::Dots:
+		case EMeshSculptStrokeType::Spacing:
+			return true;
+		default:
+			return false;
+		}
+	}
 };
 
 
