Modeling Mode: Provides new sculpting stroke modes: Airbrush, Spacing, and Dot, which break up the formerly unified concept of a stroke.

Airbrush follows the old behavior pattern of temporal based stamping via a flow parameter, Dot provides the simpliest approach - applying a stamp for the mouse position at each tick, and Spacing generates geometrically spaced stamps based on the mouse path taken, including handling cases where the user moves the mouse farther in one tick than the spacing distance, resulting in multiple stamps per tick.

Also includes an optional visualization mode for the mouse path, enabled by a CVAR.

diff --git a/Engine/Plugins/ChaosClothAssetEditor/Source/ChaosClothAssetEditorTools/Private/ChaosClothAsset/ClothWeightMapPaintTool.cpp b/Engine/Plugins/ChaosClothAssetEditor/Source/ChaosClothAssetEditorTools/Private/ChaosClothAsset/ClothWeightMapPaintTool.cpp
--- a/Engine/Plugins/ChaosClothAssetEditor/Source/ChaosClothAssetEditorTools/Private/ChaosClothAsset/ClothWeightMapPaintTool.cpp
+++ b/Engine/Plugins/ChaosClothAssetEditor/Source/ChaosClothAssetEditorTools/Private/ChaosClothAsset/ClothWeightMapPaintTool.cpp
@@ -1698,10 +1698,8 @@ void UClothEditorWeightMapPaintTool::ApplyVisibilityFilter(const TArray<int32>&
 
 
 
-int32 UClothEditorWeightMapPaintTool::FindHitSculptMeshTriangle(const FRay3d& LocalRay)
+int32 UClothEditorWeightMapPaintTool::FindHitSculptMeshTriangleConst(const FRay3d& LocalRay) const
 {
-	CurrentBaryCentricCoords = FVector3d::Zero();
-
 	if (!IsInBrushSubMode())
 	{
 		return IndexConstants::InvalidID;
@@ -1741,24 +1739,31 @@ int32 UClothEditorWeightMapPaintTool::FindHitSculptMeshTriangle(const FRay3d& Lo
 			});
 	}
 
-	if (Mesh->IsTriangle(HitTID))
-	{
-		FTriangle3d Triangle;
-		Mesh->GetTriVertices(HitTID, Triangle.V[0], Triangle.V[1], Triangle.V[2]);
-		FIntrRay3Triangle3d Query(LocalRay, Triangle);
-		Query.Find();
-		CurrentBaryCentricCoords = Query.TriangleBaryCoords;
-	}
-
 	return HitTID;
 }
 
-int32 UClothEditorWeightMapPaintTool::FindHitTargetMeshTriangle(const FRay3d& LocalRay)
+int32 UClothEditorWeightMapPaintTool::FindHitTargetMeshTriangleConst(const FRay3d& LocalRay) const
 {
 	check(false);
 	return IndexConstants::InvalidID;
 }
 
+void UClothEditorWeightMapPaintTool::UpdateHitSculptMeshTriangle(int32 TriangleID, const FRay3d& LocalRay)
+{
+	CurrentBaryCentricCoords = FVector3d::Zero();
+
+	const FDynamicMesh3* const Mesh = GetSculptMesh();
+
+	if (Mesh->IsTriangle(TriangleID))
+	{
+		FTriangle3d Triangle;
+		Mesh->GetTriVertices(TriangleID, Triangle.V[0], Triangle.V[1], Triangle.V[2]);
+		FIntrRay3Triangle3d Query(LocalRay, Triangle);
+		Query.Find();
+		CurrentBaryCentricCoords = Query.TriangleBaryCoords;
+	}
+
+}
 
 
 bool UClothEditorWeightMapPaintTool::UpdateBrushPosition(const FRay& WorldRay)
@@ -1908,7 +1913,7 @@ void UClothEditorWeightMapPaintTool::OnTick(float DeltaTime)
 	if (bUndoUpdatePending)
 	{
 		// wait for updates
-		WaitForPendingUndoRedo();
+		WaitForPendingUndoRedoUpdate();
 
 		// post rendering update
 		DynamicMeshComponent->FastNotifyTriangleVerticesUpdated(AccumulatedTriangleROI, EMeshRenderAttributeFlags::VertexColors);
@@ -1948,33 +1953,19 @@ void UClothEditorWeightMapPaintTool::OnTick(float DeltaTime)
 		bPendingPickWeight = false;
 	}
 
-
-	if (IsInBrushSubMode())
-	{
-		if (InStroke())
+	auto ExecuteStampOperation = [this](const FRay& StampRay)
 		{
 			SCOPE_CYCLE_COUNTER(WeightMapPaintTool_Tick_ApplyStampBlock);
 
-			// update brush position
-			if (UpdateStampPosition(GetPendingStampRayWorld()) == false)
-			{
-				return;
-			}
-			UpdateStampPendingState();
-			if (IsStampPending() == false)
-			{
-				return;
-			}
-
 			// update sculpt ROI
 			UpdateROI(CurrentStamp);
 
 			// append updated ROI to modified region (async)
 			FDynamicMesh3* Mesh = GetSculptMesh();
 			TFuture<void> AccumulateROI = Async(UE::Chaos::ClothAsset::Private::WeightPaintToolAsyncExecTarget, [&]()
-			{
-				UE::Geometry::VertexToTriangleOneRing(Mesh, VertexROI, AccumulatedTriangleROI);
-			});
+				{
+					UE::Geometry::VertexToTriangleOneRing(Mesh, VertexROI, AccumulatedTriangleROI);
+				});
 
 			// apply the stamp
 			bool bWeightsModified = ApplyStamp();
@@ -1989,7 +1980,14 @@ void UClothEditorWeightMapPaintTool::OnTick(float DeltaTime)
 
 			// we don't really need to wait for these to happen to end Tick()...
 			AccumulateROI.Wait();
-		}
+		};
+
+	if (IsInBrushSubMode())
+	{
+		ProcessPerTickStamps(
+			[this](const FRay& StampRay) -> bool {
+				return UpdateStampPosition(StampRay);
+			}, ExecuteStampOperation);
 	}
 
 }
@@ -2323,7 +2321,7 @@ void UClothEditorWeightMapPaintTool::EndChange()
 	NewChange->WrappedChange = MoveTemp(EditResult);
 	NewChange->BeforeModify = [this](bool bRevert)
 	{
-		this->WaitForPendingUndoRedo();
+		this->WaitForPendingUndoRedoUpdate();
 	};
 	NewChange->AfterModify = [this](bool bRevert)
 	{
@@ -2337,7 +2335,7 @@ void UClothEditorWeightMapPaintTool::EndChange()
 }
 
 
-void UClothEditorWeightMapPaintTool::WaitForPendingUndoRedo()
+void UClothEditorWeightMapPaintTool::WaitForPendingUndoRedoUpdate()
 {
 	if (bUndoUpdatePending)
 	{
@@ -2353,8 +2351,8 @@ void UClothEditorWeightMapPaintTool::OnDynamicMeshComponentChanged(UDynamicMeshC
 	// make sure any previous async computations are done, and update the undo ROI
 	if (bUndoUpdatePending)
 	{
-		// we should never hit this anymore, because of pre-change calling WaitForPendingUndoRedo()
-		WaitForPendingUndoRedo();
+		// we should never hit this anymore, because of pre-change calling WaitForPendingUndoRedoUpdate()
+		WaitForPendingUndoRedoUpdate();
 
 		// this is not right because now we are going to do extra recomputation, but it's very messy otherwise...
 		UE::Geometry::VertexToTriangleOneRing(Mesh, Change->Vertices, AccumulatedTriangleROI);
diff --git a/Engine/Plugins/ChaosClothAssetEditor/Source/ChaosClothAssetEditorTools/Private/ChaosClothAsset/ClothWeightMapPaintTool.h b/Engine/Plugins/ChaosClothAssetEditor/Source/ChaosClothAssetEditorTools/Private/ChaosClothAsset/ClothWeightMapPaintTool.h
--- a/Engine/Plugins/ChaosClothAssetEditor/Source/ChaosClothAssetEditorTools/Private/ChaosClothAsset/ClothWeightMapPaintTool.h
+++ b/Engine/Plugins/ChaosClothAssetEditor/Source/ChaosClothAssetEditorTools/Private/ChaosClothAsset/ClothWeightMapPaintTool.h
@@ -401,8 +401,10 @@ protected:
 	virtual FDynamicMesh3* GetBaseMesh() { check(false); return nullptr; }
 	virtual const FDynamicMesh3* GetBaseMesh() const { check(false); return nullptr; }
 
-	virtual int32 FindHitSculptMeshTriangle(const FRay3d& LocalRay) override;
-	virtual int32 FindHitTargetMeshTriangle(const FRay3d& LocalRay) override;
+	virtual int32 FindHitSculptMeshTriangleConst(const FRay3d& LocalRay) const override;
+	virtual int32 FindHitTargetMeshTriangleConst(const FRay3d& LocalRay) const override;
+	virtual void UpdateHitSculptMeshTriangle(int32 TriangleID, const FRay3d& LocalRay) override;
+
 
 	virtual void OnBeginStroke(const FRay& WorldRay) override;
 	virtual void OnCancelStroke() override;
@@ -537,7 +539,7 @@ protected:
 	TSet<int32> AccumulatedTriangleROI;
 	bool bUndoUpdatePending = false;
 	TArray<int> NormalsBuffer;
-	void WaitForPendingUndoRedo();
+	void WaitForPendingUndoRedoUpdate();
 
 	TArray<int> TempROIBuffer;
 	TArray<int> VertexROI;
diff --git a/Engine/Plugins/Experimental/Animation/SkeletalMeshMorphTargetEditingTools/Source/SkeletalMeshMorphTargetEditingTools/Private/MorphTargetVertexSculptTool.cpp b/Engine/Plugins/Experimental/Animation/SkeletalMeshMorphTargetEditingTools/Source/SkeletalMeshMorphTargetEditingTools/Private/MorphTargetVertexSculptTool.cpp
--- a/Engine/Plugins/Experimental/Animation/SkeletalMeshMorphTargetEditingTools/Source/SkeletalMeshMorphTargetEditingTools/Private/MorphTargetVertexSculptTool.cpp
+++ b/Engine/Plugins/Experimental/Animation/SkeletalMeshMorphTargetEditingTools/Source/SkeletalMeshMorphTargetEditingTools/Private/MorphTargetVertexSculptTool.cpp
@@ -174,7 +174,7 @@ void UMorphTargetVertexSculptTool::OnEndStroke()
 	NewChange->WrappedChange = MoveTemp(ActiveVertexChange->Change);
 	NewChange->BeforeModify = [this](bool bRevert)
 	{
-		this->WaitForPendingUndoRedo();
+		this->WaitForPendingUndoRedoUpdate();
 		// Any sculpt change needs to be applied to the latest posed mesh during undo / redo, so pose the mesh immediately instead of waiting
 		// for the next tick
 		this->PoseToolMesh();
@@ -194,11 +194,6 @@ void UMorphTargetVertexSculptTool::OnEndStroke()
 	ActiveVertexChange = nullptr;
 }
 
-int32 UMorphTargetVertexSculptTool::FindHitSculptMeshTriangle(const FRay3d& LocalRay)
-{
-	WaitForPendingUndoRedo();
-	return Super::FindHitSculptMeshTriangle(LocalRay);
-}
 
 void UMorphTargetVertexSculptTool::OnToolMeshChanged(UDynamicMeshComponent* Component, const FMeshRegionChangeBase* Change, bool bRevert)
 {
@@ -258,8 +253,8 @@ void UMorphTargetVertexSculptTool::InitializeCache()
 
 void UMorphTargetVertexSculptTool::UpdateCacheIfNeeded()
 {
-	WaitForPendingStampUpdate();
-	WaitForPendingUndoRedo();
+	WaitForPendingStampUpdateConst();
+	WaitForPendingUndoRedoUpdate();
 	UPrimitiveComponent* Component = Cast<IPrimitiveComponentBackedTarget>(GetTarget())->GetOwnerComponent();
 	USkeletalMeshComponent* SkeletalMeshComponent = Cast<USkeletalMeshComponent>(Component);
 
@@ -445,8 +440,8 @@ void UMorphTargetVertexSculptTool::PoseToolMesh()
 	PreviousMorphWeights = MorphTargetWeights;
 
 	// have to wait for any outstanding stamp/undo update to finish...
-	WaitForPendingStampUpdate();
-	WaitForPendingUndoRedo();
+	WaitForPendingStampUpdateConst();
+	WaitForPendingUndoRedoUpdate();
 
 	TArray<FTransform> RefComponentSpaceTransforms;
 	for (int32 BoneIndex = 0; BoneIndex < ComponentSpaceTransforms.Num(); ++BoneIndex)
diff --git a/Engine/Plugins/Experimental/Animation/SkeletalMeshMorphTargetEditingTools/Source/SkeletalMeshMorphTargetEditingTools/Public/MorphTargetVertexSculptTool.h b/Engine/Plugins/Experimental/Animation/SkeletalMeshMorphTargetEditingTools/Source/SkeletalMeshMorphTargetEditingTools/Public/MorphTargetVertexSculptTool.h
--- a/Engine/Plugins/Experimental/Animation/SkeletalMeshMorphTargetEditingTools/Source/SkeletalMeshMorphTargetEditingTools/Public/MorphTargetVertexSculptTool.h
+++ b/Engine/Plugins/Experimental/Animation/SkeletalMeshMorphTargetEditingTools/Source/SkeletalMeshMorphTargetEditingTools/Public/MorphTargetVertexSculptTool.h
@@ -48,8 +48,6 @@ public:
 	
 protected:
 	friend class FMorphTargetVertexSculptNonSymmetricChange;
-	virtual int32 FindHitSculptMeshTriangle(const FRay3d& LocalRay) override;
-	
 
 	void OnToolMeshChanged(UDynamicMeshComponent* Component, const FMeshRegionChangeBase* Change, bool bRevert);
 
diff --git a/Engine/Plugins/Experimental/Dataflow/Source/DataflowEditor/Private/DataflowEditorTools/DataflowEditorWeightMapPaintTool.cpp b/Engine/Plugins/Experimental/Dataflow/Source/DataflowEditor/Private/DataflowEditorTools/DataflowEditorWeightMapPaintTool.cpp
--- a/Engine/Plugins/Experimental/Dataflow/Source/DataflowEditor/Private/DataflowEditorTools/DataflowEditorWeightMapPaintTool.cpp
+++ b/Engine/Plugins/Experimental/Dataflow/Source/DataflowEditor/Private/DataflowEditorTools/DataflowEditorWeightMapPaintTool.cpp
@@ -1341,11 +1341,8 @@ void UDataflowEditorWeightMapPaintTool::ApplyVisibilityFilter(const TArray<int32
 
 
 
-int32 UDataflowEditorWeightMapPaintTool::FindHitSculptMeshTriangle(const FRay3d& LocalRay)
+int32 UDataflowEditorWeightMapPaintTool::FindHitSculptMeshTriangleConst(const FRay3d& LocalRay) const
 {
-	// TODO: Figure out what the actual position on the triangle is when hit.
-	CurrentBaryCentricCoords = FVector3d(1.0 / 3.0, 1.0 / 3.0, 1.0 / 3.0);
-
 	if (!IsInBrushSubMode())
 	{
 		return IndexConstants::InvalidID;
@@ -1357,7 +1354,7 @@ int32 UDataflowEditorWeightMapPaintTool::FindHitSculptMeshTriangle(const FRay3d&
 	}
 	else
 	{
-		FDynamicMesh3* Mesh = GetSculptMesh();
+		const FDynamicMesh3* Mesh = GetSculptMesh();
 
 		FViewCameraState StateOut;
 		GetToolManager()->GetContextQueriesAPI()->GetCurrentViewState(StateOut);
@@ -1373,13 +1370,17 @@ int32 UDataflowEditorWeightMapPaintTool::FindHitSculptMeshTriangle(const FRay3d&
 	}
 }
 
-int32 UDataflowEditorWeightMapPaintTool::FindHitTargetMeshTriangle(const FRay3d& LocalRay)
+int32 UDataflowEditorWeightMapPaintTool::FindHitTargetMeshTriangleConst(const FRay3d& LocalRay) const
 {
 	check(false);
 	return IndexConstants::InvalidID;
 }
 
-
+void UDataflowEditorWeightMapPaintTool::UpdateHitSculptMeshTriangle(int32 TriangleID, const FRay3d& LocalRay)
+{
+	// TODO: Figure out what the actual position on the triangle is when hit.
+	CurrentBaryCentricCoords = FVector3d(1.0 / 3.0, 1.0 / 3.0, 1.0 / 3.0);
+}
 
 bool UDataflowEditorWeightMapPaintTool::UpdateBrushPosition(const FRay& WorldRay)
 {
@@ -1528,7 +1529,7 @@ void UDataflowEditorWeightMapPaintTool::OnTick(float DeltaTime)
 	if (bUndoUpdatePending)
 	{
 		// wait for updates
-		WaitForPendingUndoRedo();
+		WaitForPendingUndoRedoUpdate();
 
 		// post rendering update
 		DynamicMeshComponent->FastNotifyTriangleVerticesUpdated(AccumulatedTriangleROI, EMeshRenderAttributeFlags::VertexColors);
@@ -1563,33 +1564,19 @@ void UDataflowEditorWeightMapPaintTool::OnTick(float DeltaTime)
 		bPendingPickWeight = false;
 	}
 
-
-	if (IsInBrushSubMode())
-	{
-		if (InStroke())
+	auto ExecuteStampOperation = [this](const FRay& StampRay)
 		{
 			SCOPE_CYCLE_COUNTER(WeightMapPaintTool_Tick_ApplyStampBlock);
 
-			// update brush position
-			if (UpdateStampPosition(GetPendingStampRayWorld()) == false)
-			{
-				return;
-			}
-			UpdateStampPendingState();
-			if (IsStampPending() == false)
-			{
-				return;
-			}
-
 			// update sculpt ROI
 			UpdateROI(CurrentStamp);
 
 			// append updated ROI to modified region (async)
 			FDynamicMesh3* Mesh = GetSculptMesh();
 			TFuture<void> AccumulateROI = Async(UE::Dataflow::Private::WeightPaintToolAsyncExecTarget, [&]()
-			{
-				UE::Geometry::VertexToTriangleOneRing(Mesh, VertexROI, AccumulatedTriangleROI);
-			});
+				{
+					UE::Geometry::VertexToTriangleOneRing(Mesh, VertexROI, AccumulatedTriangleROI);
+				});
 
 			// apply the stamp
 			bool bWeightsModified = ApplyStamp();
@@ -1604,7 +1591,15 @@ void UDataflowEditorWeightMapPaintTool::OnTick(float DeltaTime)
 
 			// we don't really need to wait for these to happen to end Tick()...
 			AccumulateROI.Wait();
-		}
+		};
+
+
+	if (IsInBrushSubMode())
+	{
+		ProcessPerTickStamps(
+			[this](const FRay& StampRay) -> bool {
+				return UpdateStampPosition(StampRay);
+			}, ExecuteStampOperation);
 	}
 
 }
@@ -1878,7 +1873,7 @@ void UDataflowEditorWeightMapPaintTool::EndChange()
 	NewChange->WrappedChange = MoveTemp(DataflowWeightPaintMeshChange);
 	NewChange->BeforeModify = [this](bool bRevert)
 	{
-		this->WaitForPendingUndoRedo();
+		this->WaitForPendingUndoRedoUpdate();
 	};
 
 	GetToolManager()->EmitObjectChange(DynamicMeshComponent, MoveTemp(NewChange), LOCTEXT("WeightPaintChange", "Weight Stroke"));
@@ -1886,7 +1881,7 @@ void UDataflowEditorWeightMapPaintTool::EndChange()
 }
 
 
-void UDataflowEditorWeightMapPaintTool::WaitForPendingUndoRedo()
+void UDataflowEditorWeightMapPaintTool::WaitForPendingUndoRedoUpdate()
 {
 	if (bUndoUpdatePending)
 	{
@@ -1902,8 +1897,8 @@ void UDataflowEditorWeightMapPaintTool::OnDynamicMeshComponentChanged(UDynamicMe
 	// make sure any previous async computations are done, and update the undo ROI
 	if (bUndoUpdatePending)
 	{
-		// we should never hit this anymore, because of pre-change calling WaitForPendingUndoRedo()
-		WaitForPendingUndoRedo();
+		// we should never hit this anymore, because of pre-change calling WaitForPendingUndoRedoUpdate()
+		WaitForPendingUndoRedoUpdate();
 
 		// this is not right because now we are going to do extra recomputation, but it's very messy otherwise...
 		UE::Geometry::VertexToTriangleOneRing(Mesh, Change->Vertices, AccumulatedTriangleROI);
diff --git a/Engine/Plugins/Experimental/Dataflow/Source/DataflowEditor/Public/DataflowEditorTools/DataflowEditorWeightMapPaintTool.h b/Engine/Plugins/Experimental/Dataflow/Source/DataflowEditor/Public/DataflowEditorTools/DataflowEditorWeightMapPaintTool.h
--- a/Engine/Plugins/Experimental/Dataflow/Source/DataflowEditor/Public/DataflowEditorTools/DataflowEditorWeightMapPaintTool.h
+++ b/Engine/Plugins/Experimental/Dataflow/Source/DataflowEditor/Public/DataflowEditorTools/DataflowEditorWeightMapPaintTool.h
@@ -342,8 +342,9 @@ protected:
 	virtual FDynamicMesh3* GetBaseMesh() { check(false); return nullptr; }
 	virtual const FDynamicMesh3* GetBaseMesh() const { check(false); return nullptr; }
 
-	virtual int32 FindHitSculptMeshTriangle(const FRay3d& LocalRay) override;
-	virtual int32 FindHitTargetMeshTriangle(const FRay3d& LocalRay) override;
+	virtual int32 FindHitSculptMeshTriangleConst(const FRay3d& LocalRay) const override;
+	virtual int32 FindHitTargetMeshTriangleConst(const FRay3d& LocalRay) const override;
+	virtual void UpdateHitSculptMeshTriangle(int32 TriangleID, const FRay3d& LocalRay) override;
 
 	virtual void OnBeginStroke(const FRay& WorldRay) override;
 	virtual void OnEndStroke() override;
@@ -441,7 +442,7 @@ protected:
 	TSet<int32> AccumulatedTriangleROI;
 	bool bUndoUpdatePending = false;
 	TArray<int> NormalsBuffer;
-	void WaitForPendingUndoRedo();
+	void WaitForPendingUndoRedoUpdate();
 
 	TArray<int> TempROIBuffer;
 	TArray<int> VertexROI;
diff --git a/Engine/Plugins/Runtime/MeshModelingToolset/Source/MeshModelingTools/Private/MeshGroupPaintTool.cpp b/Engine/Plugins/Runtime/MeshModelingToolset/Source/MeshModelingTools/Private/MeshGroupPaintTool.cpp
--- a/Engine/Plugins/Runtime/MeshModelingToolset/Source/MeshModelingTools/Private/MeshGroupPaintTool.cpp
+++ b/Engine/Plugins/Runtime/MeshModelingToolset/Source/MeshModelingTools/Private/MeshGroupPaintTool.cpp
@@ -852,7 +852,7 @@ void UMeshGroupPaintTool::ApplyVisibilityFilter(const TArray<int32>& Triangles,
 
 
 
-int32 UMeshGroupPaintTool::FindHitSculptMeshTriangle(const FRay3d& LocalRay)
+int32 UMeshGroupPaintTool::FindHitSculptMeshTriangleConst(const FRay3d& LocalRay) const
 {
 	if (!IsInBrushSubMode())
 	{
@@ -865,7 +865,7 @@ int32 UMeshGroupPaintTool::FindHitSculptMeshTriangle(const FRay3d& LocalRay)
 	}
 	else
 	{
-		FDynamicMesh3* Mesh = GetSculptMesh();
+		const FDynamicMesh3* Mesh = GetSculptMesh();
 
 		FViewCameraState StateOut;
 		GetToolManager()->GetContextQueriesAPI()->GetCurrentViewState(StateOut);
@@ -881,7 +881,7 @@ int32 UMeshGroupPaintTool::FindHitSculptMeshTriangle(const FRay3d& LocalRay)
 	}
 }
 
-int32 UMeshGroupPaintTool::FindHitTargetMeshTriangle(const FRay3d& LocalRay)
+int32 UMeshGroupPaintTool::FindHitTargetMeshTriangleConst(const FRay3d& LocalRay) const
 {
 	check(false);
 	return IndexConstants::InvalidID;
@@ -1012,7 +1012,7 @@ void UMeshGroupPaintTool::OnTick(float DeltaTime)
 	if (bUndoUpdatePending)
 	{
 		// wait for updates
-		WaitForPendingUndoRedo();
+		WaitForPendingUndoRedoUpdate();
 
 		// post rendering update
 		DynamicMeshComponent->FastNotifyTriangleVerticesUpdated(AccumulatedTriangleROI, EMeshRenderAttributeFlags::VertexColors);
@@ -1045,32 +1045,18 @@ void UMeshGroupPaintTool::OnTick(float DeltaTime)
 		bPendingPickGroup = bPendingToggleFreezeGroup = false;
 	}
 
-
-	if (IsInBrushSubMode())
-	{
-		if (InStroke())
+	auto ExecuteStampOperation = [this](const FRay& StampRay)
 		{
 			SCOPE_CYCLE_COUNTER(GroupPaintTool_Tick_ApplyStampBlock);
 
-			// update brush position
-			if (UpdateStampPosition(GetPendingStampRayWorld()) == false)
-			{
-				return;
-			}
-			UpdateStampPendingState();
-			if (IsStampPending() == false)
-			{
-				return;
-			}
-
 			// update sculpt ROI
 			UpdateROI(CurrentStamp);
 
 			// append updated ROI to modified region (async)
 			TFuture<void> AccumulateROI = Async(GroupPaintToolAsyncExecTarget, [&]()
-			{
-				AccumulatedTriangleROI.Append(TriangleROI);
-			});
+				{
+					AccumulatedTriangleROI.Append(TriangleROI);
+				});
 
 			// apply the stamp
 			bool bGroupsModified = ApplyStamp();
@@ -1084,7 +1070,14 @@ void UMeshGroupPaintTool::OnTick(float DeltaTime)
 
 			// we don't really need to wait for these to happen to end Tick()...
 			AccumulateROI.Wait();
-		}
+		};
+
+	if (IsInBrushSubMode())
+	{
+		ProcessPerTickStamps(
+			[this](const FRay& StampRay) -> bool {
+				return UpdateStampPosition(StampRay);
+			}, ExecuteStampOperation);
 	}
 
 }
@@ -1357,7 +1350,7 @@ void UMeshGroupPaintTool::EndChange()
 	NewChange->WrappedChange = MakeUnique<FMeshPolygroupChange>(MoveTemp(EditResult));
 	NewChange->BeforeModify = [this](bool bRevert)
 	{
-		this->WaitForPendingUndoRedo();
+		this->WaitForPendingUndoRedoUpdate();
 	};
 
 	GetToolManager()->EmitObjectChange(DynamicMeshComponent, MoveTemp(NewChange), LOCTEXT("GroupPaintChange", "Group Stroke"));
@@ -1368,7 +1361,7 @@ void UMeshGroupPaintTool::EndChange()
 }
 
 
-void UMeshGroupPaintTool::WaitForPendingUndoRedo()
+void UMeshGroupPaintTool::WaitForPendingUndoRedoUpdate()
 {
 	if (bUndoUpdatePending)
 	{
@@ -1384,8 +1377,8 @@ void UMeshGroupPaintTool::OnDynamicMeshComponentChanged()
 	// make sure any previous async computations are done, and update the undo ROI
 	if (bUndoUpdatePending)
 	{
-		// we should never hit this anymore, because of pre-change calling WaitForPendingUndoRedo()
-		WaitForPendingUndoRedo();
+		// we should never hit this anymore, because of pre-change calling WaitForPendingUndoRedoUpdate()
+		WaitForPendingUndoRedoUpdate();
 
 		// TODO: do we need to read from mesh change here??
 		//UE::Geometry::VertexToTriangleOneRing(Mesh, Change->Vertices, AccumulatedTriangleROI);
diff --git a/Engine/Plugins/Runtime/MeshModelingToolset/Source/MeshModelingTools/Private/MeshVertexPaintTool.cpp b/Engine/Plugins/Runtime/MeshModelingToolset/Source/MeshModelingTools/Private/MeshVertexPaintTool.cpp
--- a/Engine/Plugins/Runtime/MeshModelingToolset/Source/MeshModelingTools/Private/MeshVertexPaintTool.cpp
+++ b/Engine/Plugins/Runtime/MeshModelingToolset/Source/MeshModelingTools/Private/MeshVertexPaintTool.cpp
@@ -279,7 +279,7 @@ void UMeshVertexPaintTool::Setup()
 	UMeshSculptToolBase::BrushProperties->bShowFalloff = true;
 	UMeshSculptToolBase::BrushProperties->bShowLazyness = true;
 	UMeshSculptToolBase::BrushProperties->bShowFlowRate = true;
-	UMeshSculptToolBase::BrushProperties->bShowSpacing = false;
+	UMeshSculptToolBase::BrushProperties->bShowSpacing = true;
 	UMeshSculptToolBase::BrushProperties->BrushFalloffAmount = 0.66f;
 	UMeshSculptToolBase::BrushProperties->BrushSize.bToolSupportsPressureSensitivity = true;
 	CalculateBrushRadius();
@@ -563,7 +563,7 @@ void UMeshVertexPaintTool::OnBeginStroke(const FRay& WorldRay)
 	else
 	{
 		UseBrushOp->PropertySet->SetStrength(
-			FMathf::Pow(UMeshSculptToolBase::BrushProperties->FlowRate, 2.0f));
+			FMathf::Pow(BasicProperties->Strength, 2.0f));
 	}
 	UseBrushOp->PropertySet->SetFalloff(UMeshSculptToolBase::BrushProperties->BrushFalloffAmount);
 
@@ -1162,7 +1162,7 @@ void UMeshVertexPaintTool::ApplyVisibilityFilter(const TArray<int32>& Triangles,
 
 
 
-int32 UMeshVertexPaintTool::FindHitSculptMeshTriangle(const FRay3d& LocalRay)
+int32 UMeshVertexPaintTool::FindHitSculptMeshTriangleConst(const FRay3d& LocalRay) const
 {
 	if (!IsInBrushSubMode())
 	{
@@ -1180,7 +1180,7 @@ int32 UMeshVertexPaintTool::FindHitSculptMeshTriangle(const FRay3d& LocalRay)
 	}
 	else
 	{
-		FDynamicMesh3* Mesh = GetSculptMesh();
+		const FDynamicMesh3* Mesh = GetSculptMesh();
 
 		FViewCameraState StateOut;
 		GetToolManager()->GetContextQueriesAPI()->GetCurrentViewState(StateOut);
@@ -1200,7 +1200,7 @@ int32 UMeshVertexPaintTool::FindHitSculptMeshTriangle(const FRay3d& LocalRay)
 	}
 }
 
-int32 UMeshVertexPaintTool::FindHitTargetMeshTriangle(const FRay3d& LocalRay)
+int32 UMeshVertexPaintTool::FindHitTargetMeshTriangleConst(const FRay3d& LocalRay) const
 {
 	check(false);
 	return IndexConstants::InvalidID;
@@ -1248,7 +1248,7 @@ bool UMeshVertexPaintTool::OnUpdateHover(const FInputDeviceRay& DevicePos)
 		// update strength and falloff for current brush
 		TUniquePtr<FMeshSculptBrushOp>& UseBrushOp = GetActiveBrushOp();
 		UseBrushOp->PropertySet->SetStrength(
-			FMathf::Pow(UMeshSculptToolBase::BrushProperties->FlowRate, 2.0f) );
+			FMathf::Pow(BasicProperties->Strength, 2.0f) );
 		UseBrushOp->PropertySet->SetFalloff(UMeshSculptToolBase::BrushProperties->BrushFalloffAmount);
 
 	}
@@ -1319,7 +1319,7 @@ void UMeshVertexPaintTool::OnTick(float DeltaTime)
 	if (bUndoUpdatePending)
 	{
 		// wait for updates
-		WaitForPendingUndoRedo();
+		WaitForPendingUndoRedoUpdate();
 
 		// post rendering update
 		DynamicMeshComponent->FastNotifyTriangleVerticesUpdated(AccumulatedTriangleROI, EMeshRenderAttributeFlags::VertexColors);
@@ -1367,32 +1367,18 @@ void UMeshVertexPaintTool::OnTick(float DeltaTime)
 		bPendingPickColor = bPendingPickEraseColor = false;
 	}
 
-
-	if (IsInBrushSubMode())
-	{
-		if (InStroke())
+	auto ExecuteStampOperation = [this](const FRay& StampRay)
 		{
 			TRACE_CPUPROFILER_EVENT_SCOPE(MeshVPaint_OnTick_InStroke);
 
-			// update brush position
-			if (UpdateStampPosition(GetPendingStampRayWorld()) == false)
-			{
-				return;
-			}
-			UpdateStampPendingState();
-			if (IsStampPending() == false)
-			{
-				return;
-			}
-
 			// update sculpt ROI
 			UpdateROI(CurrentStamp);
 
 			// append updated ROI to modified region (async)
 			TFuture<void> AccumulateROI = Async(VertexPaintToolAsyncExecTarget, [&]()
-			{
-				AccumulatedTriangleROI.Append(TriangleROI);
-			});
+				{
+					AccumulatedTriangleROI.Append(TriangleROI);
+				});
 
 			// apply the stamp
 			bool bColorModified = ApplyStamp();
@@ -1406,7 +1392,14 @@ void UMeshVertexPaintTool::OnTick(float DeltaTime)
 
 			// we don't really need to wait for these to happen to end Tick()...
 			AccumulateROI.Wait();
-		}
+		};
+
+	if (IsInBrushSubMode())
+	{
+		ProcessPerTickStamps(
+			[this](const FRay& StampRay) -> bool {
+				return UpdateStampPosition(StampRay);
+			}, ExecuteStampOperation);
 	}
 
 }
@@ -1950,7 +1943,7 @@ void UMeshVertexPaintTool::EndChange()
 	NewChange->WrappedChange = MoveTemp(EditResult);
 	NewChange->BeforeModify = [this](bool bRevert)
 	{
-		this->WaitForPendingUndoRedo();
+		this->WaitForPendingUndoRedoUpdate();
 	};
 
 	GetToolManager()->EmitObjectChange(this, MoveTemp(NewChange), LOCTEXT("VertexPaintChange", "Paint Stroke"));
@@ -1978,7 +1971,7 @@ void UMeshVertexPaintTool::ExternalUpdateValues(const TArray<int32>& ElementIDs,
 
 
 
-void UMeshVertexPaintTool::WaitForPendingUndoRedo()
+void UMeshVertexPaintTool::WaitForPendingUndoRedoUpdate()
 {
 	if (bUndoUpdatePending)
 	{
@@ -1994,8 +1987,8 @@ void UMeshVertexPaintTool::OnDynamicMeshComponentChanged()
 	// make sure any previous async computations are done, and update the undo ROI
 	if (bUndoUpdatePending)
 	{
-		// we should never hit this anymore, because of pre-change calling WaitForPendingUndoRedo()
-		WaitForPendingUndoRedo();
+		// we should never hit this anymore, because of pre-change calling WaitForPendingUndoRedoUpdate()
+		WaitForPendingUndoRedoUpdate();
 
 		// TODO: do we need to read from mesh change here??
 		//UE::Geometry::VertexToTriangleOneRing(Mesh, Change->Vertices, AccumulatedTriangleROI);
diff --git a/Engine/Plugins/Runtime/MeshModelingToolset/Source/MeshModelingTools/Private/MeshVertexSculptTool.cpp b/Engine/Plugins/Runtime/MeshModelingToolset/Source/MeshModelingTools/Private/MeshVertexSculptTool.cpp
--- a/Engine/Plugins/Runtime/MeshModelingToolset/Source/MeshModelingTools/Private/MeshVertexSculptTool.cpp
+++ b/Engine/Plugins/Runtime/MeshModelingToolset/Source/MeshModelingTools/Private/MeshVertexSculptTool.cpp
@@ -437,8 +437,8 @@ void UMeshVertexSculptTool::UpdateToolMeshes(TFunctionRef<TUniquePtr<FMeshRegion
 	if (AllowToolMeshUpdates())
 	{
 		// have to wait for any outstanding stamp/undo update to finish...
-		WaitForPendingStampUpdate();
-		WaitForPendingUndoRedo();
+		WaitForPendingStampUpdateConst();
+		WaitForPendingUndoRedoUpdate();
 
 		TUniquePtr<FMeshRegionChangeBase> Change = UpdateMesh(*GetSculptMesh(), 0);
 		// A change was created -- emit it to the tool manager and update associated data structures, etc
@@ -451,7 +451,7 @@ void UMeshVertexSculptTool::UpdateToolMeshes(TFunctionRef<TUniquePtr<FMeshRegion
 			NewChange->WrappedChange = MoveTemp(Change);
 			NewChange->BeforeModify = [this](bool bRevert)
 			{
-				this->WaitForPendingUndoRedo();
+				this->WaitForPendingUndoRedoUpdate();
 			};
 			// Note this change should be in the context of a larger transaction, so the text isn't that important
 			GetToolManager()->EmitObjectChange(DynamicMeshComponent, MoveTemp(NewChange), LOCTEXT("UpdateVertexSculptMesh", "Updated Mesh"));
@@ -550,7 +550,7 @@ void UMeshVertexSculptTool::SetRegionFilterType(int32 Identifier)
 
 void UMeshVertexSculptTool::OnBeginStroke(const FRay& WorldRay)
 {
-	WaitForPendingUndoRedo();		// cannot start stroke if there is an outstanding undo/redo update
+	WaitForPendingUndoRedoUpdate();		// cannot start stroke if there is an outstanding undo/redo update
 
 	UpdateBrushPosition(WorldRay);
 
@@ -1205,7 +1205,7 @@ TFuture<void> UMeshVertexSculptTool::ApplyStamp()
 
 
 
-bool UMeshVertexSculptTool::IsHitTriangleBackFacing(int32 TriangleID, const FDynamicMesh3* QueryMesh)
+bool UMeshVertexSculptTool::IsHitTriangleBackFacing(int32 TriangleID, const FDynamicMesh3* QueryMesh) const
 {
 	if (TriangleID != IndexConstants::InvalidID)
 	{
@@ -1223,10 +1223,12 @@ bool UMeshVertexSculptTool::IsHitTriangleBackFacing(int32 TriangleID, const FDyn
 }
 
 
-int32 UMeshVertexSculptTool::FindHitSculptMeshTriangle(const FRay3d& LocalRay)
+int32 UMeshVertexSculptTool::FindHitSculptMeshTriangleConst(const FRay3d& LocalRay) const
 {
-	// need this to finish before we can touch Octree
-	WaitForPendingStampUpdate();
+	// have to wait for any outstanding stamp update to finish...
+	WaitForPendingStampUpdateConst();
+	// wait for previous Undo to finish (possibly never hit because the change records do it?)
+	WaitForPendingUndoRedoUpdate();
 
 	int32 HitTID = Octree.FindNearestHitObject(LocalRay);
 	if (GetBrushCanHitBackFaces() == false && IsHitTriangleBackFacing(HitTID, GetSculptMesh()))
@@ -1236,7 +1238,7 @@ int32 UMeshVertexSculptTool::FindHitSculptMeshTriangle(const FRay3d& LocalRay)
 	return HitTID;
 }
 
-int32 UMeshVertexSculptTool::FindHitTargetMeshTriangle(const FRay3d& LocalRay)
+int32 UMeshVertexSculptTool::FindHitTargetMeshTriangleConst(const FRay3d& LocalRay) const
 {
 	int32 HitTID = BaseMeshSpatial.FindNearestHitObject(LocalRay);
 	if (GetBrushCanHitBackFaces() == false && IsHitTriangleBackFacing(HitTID, GetBaseMesh()))
@@ -1384,7 +1386,7 @@ void UMeshVertexSculptTool::OnTick(float DeltaTime)
 	if (bUndoUpdatePending)
 	{
 		// wait for updates
-		WaitForPendingUndoRedo();
+		WaitForPendingUndoRedoUpdate();
 
 		// post rendering update
 		DynamicMeshComponent->FastNotifyTriangleVerticesUpdated(AccumulatedTriangleROI,
@@ -1392,7 +1394,6 @@ void UMeshVertexSculptTool::OnTick(float DeltaTime)
 		GetToolManager()->PostInvalidation();
 
 		// ignore stamp and wait for next tick to do anything else
-		bUndoUpdatePending = false;
 		return;
 	}
 
@@ -1403,25 +1404,14 @@ void UMeshVertexSculptTool::OnTick(float DeltaTime)
 		bTargetDirty = false;
 	}
 
-	if (InStroke())
+
+	FDynamicMesh3* Mesh = GetSculptMesh();
+	auto ExecuteStampOperation = [this, Mesh](const FRay& StampRay)
 	{
-		//UE_LOG(LogTemp, Warning, TEXT("dt is %.3f, tick fps %.2f - roi size %d/%d"), DeltaTime, 1.0 / DeltaTime, VertexROI.Num(), TriangleROI.Num());
 		TRACE_CPUPROFILER_EVENT_SCOPE(VtxSculptTool_Tick_StrokeUpdate);
-		FDynamicMesh3* Mesh = GetSculptMesh();
-
-		// update brush position
-		if (UpdateStampPosition(GetPendingStampRayWorld()) == false)
-		{
-			return;
-		}
-		UpdateStampPendingState();
-		if (IsStampPending() == false)
-		{
-			return;
-		}
 
 		// need to make sure previous stamp finished
-		WaitForPendingStampUpdate();
+		WaitForPendingStampUpdateConst();
 
 		// update sculpt ROI
 		UpdateROI(CurrentStamp.LocalFrame);
@@ -1429,24 +1419,24 @@ void UMeshVertexSculptTool::OnTick(float DeltaTime)
 		// Append updated ROI to modified region (async). For some reason this is very expensive,
 		// maybe because of TSet? but we have a lot of time to do it.
 		TFuture<void> AccumulateROI = Async(VertexSculptToolAsyncExecTarget, [this]()
-		{
-			TRACE_CPUPROFILER_EVENT_SCOPE(VtxSculptTool_Tick_AccumROI);
-			for (int32 tid : TriangleROIArray)
 			{
-				AccumulatedTriangleROI.Add(tid);
-			}
-		});
+				TRACE_CPUPROFILER_EVENT_SCOPE(VtxSculptTool_Tick_AccumROI);
+				for (int32 tid : TriangleROIArray)
+				{
+					AccumulatedTriangleROI.Add(tid);
+				}
+			});
 
 		// Start precomputing the normals ROI. This is currently the most expensive single thing we do next
 		// to Octree re-insertion, despite it being almost trivial. Why?!?
 		bool bUsingOverlayNormalsOut = false;
 		TFuture<void> NormalsROI = Async(VertexSculptToolAsyncExecTarget, [Mesh, &bUsingOverlayNormalsOut, this]()
-		{
-			TRACE_CPUPROFILER_EVENT_SCOPE(VtxSculptTool_Tick_NormalsROI);
+			{
+				TRACE_CPUPROFILER_EVENT_SCOPE(VtxSculptTool_Tick_NormalsROI);
 
-			//UE::SculptUtil::PrecalculateNormalsROI(Mesh, TriangleROIArray, NormalsROIBuilder, bUsingOverlayNormals, false);
-			UE::SculptUtil::PrecalculateNormalsROI(Mesh, TriangleROIArray, NormalsFlags, bUsingOverlayNormalsOut, false);
-		});
+				//UE::SculptUtil::PrecalculateNormalsROI(Mesh, TriangleROIArray, NormalsROIBuilder, bUsingOverlayNormals, false);
+				UE::SculptUtil::PrecalculateNormalsROI(Mesh, TriangleROIArray, NormalsFlags, bUsingOverlayNormalsOut, false);
+			});
 
 		// NOTE: you might try to speculatively do the octree remove here, to save doing it later on Reinsert().
 		// This will not improve things, as Reinsert() checks if it needs to actually re-insert, which avoids many
@@ -1458,11 +1448,16 @@ void UMeshVertexSculptTool::OnTick(float DeltaTime)
 
 		// begin octree rebuild calculation
 		StampUpdateOctreeFuture = Async(VertexSculptToolAsyncExecTarget, [this]()
-		{
-			TRACE_CPUPROFILER_EVENT_SCOPE(VtxSculptTool_Tick_OctreeReinsert);
-			Octree.ReinsertTrianglesParallel(TriangleROIArray, OctreeUpdateTempBuffer, OctreeUpdateTempFlagBuffer);
-		});
+			{
+				TRACE_CPUPROFILER_EVENT_SCOPE(VtxSculptTool_Tick_OctreeReinsert);
+				Octree.ReinsertTrianglesParallel(TriangleROIArray, OctreeUpdateTempBuffer, OctreeUpdateTempFlagBuffer);
+			});
 		bStampUpdatePending = true;
+		StampUpdateOctreeFuture.Then([this](TFuture<void> Future)
+			{
+				bStampUpdatePending = false;
+			});
+
 		//TFuture<void> OctreeRebuild = Async(VertexSculptToolAsyncExecTarget, [&]()
 		//{
 		//	TRACE_CPUPROFILER_EVENT_SCOPE(VtxSculptTool_Tick_OctreeReinsert);
@@ -1475,7 +1470,7 @@ void UMeshVertexSculptTool::OnTick(float DeltaTime)
 		// precompute dynamic mesh update info
 		TArray<int32> RenderUpdateSets; FAxisAlignedBox3d RenderUpdateBounds;
 		TFuture<bool> RenderUpdatePrecompute = DynamicMeshComponent->FastNotifyTriangleVerticesUpdated_TryPrecompute(
-				TriangleROIArray, RenderUpdateSets, RenderUpdateBounds);
+			TriangleROIArray, RenderUpdateSets, RenderUpdateBounds);
 
 		// recalculate normals. This has to complete before we can update component
 		// (in fact we could do it per-chunk...)
@@ -1499,8 +1494,16 @@ void UMeshVertexSculptTool::OnTick(float DeltaTime)
 		// we don't really need to wait for these to happen to end Tick()...
 		UpdateChangeFuture.Wait();
 		AccumulateROI.Wait();
-	} 
-	else if (bTargetDirty)
+	};
+
+	
+	ProcessPerTickStamps(
+		[this](const FRay& StampRay) -> bool {
+			return UpdateStampPosition(StampRay);
+		}, ExecuteStampOperation);
+
+
+	if (InStroke() == false && bTargetDirty)
 	{
 		TRACE_CPUPROFILER_EVENT_SCOPE(VtxSculptTool_Tick_UpdateTarget);
 		check(InStroke() == false);
@@ -1516,12 +1519,11 @@ void UMeshVertexSculptTool::OnTick(float DeltaTime)
 
 
 
-void UMeshVertexSculptTool::WaitForPendingStampUpdate()
+void UMeshVertexSculptTool::WaitForPendingStampUpdateConst() const
 {
 	if (bStampUpdatePending)
 	{
 		StampUpdateOctreeFuture.Wait();
-		bStampUpdatePending = false;
 	}
 }
 
@@ -1697,7 +1699,7 @@ void UMeshVertexSculptTool::EndChange()
 	NewChange->WrappedChange = MoveTemp(ActiveVertexChange->Change);
 	NewChange->BeforeModify = [this](bool bRevert)
 	{
-		this->WaitForPendingUndoRedo();
+		this->WaitForPendingUndoRedoUpdate();
 	};
 
 	GetToolManager()->EmitObjectChange(DynamicMeshComponent, MoveTemp(NewChange), LOCTEXT("VertexSculptChange", "Brush Stroke"));
@@ -1720,23 +1722,17 @@ void UMeshVertexSculptTool::SetDefaultPrimaryBrushID(const int32 InPrimaryBrushI
 }
 
 
-void UMeshVertexSculptTool::WaitForPendingUndoRedo()
+void UMeshVertexSculptTool::WaitForPendingUndoRedoUpdate() const
 {
-	if (bUndoUpdatePending)
-	{
-		UndoNormalsFuture.Wait();
-		UndoUpdateOctreeFuture.Wait();
-		UndoUpdateBaseMeshFuture.Wait();
-		bUndoUpdatePending = false;
-	}
+	UndoUpdateFuture.Wait();
 }
 
 void UMeshVertexSculptTool::OnDynamicMeshComponentChanged(UDynamicMeshComponent* Component, const FMeshRegionChangeBase* Change, bool bRevert)
 {
 	// have to wait for any outstanding stamp update to finish...
-	WaitForPendingStampUpdate();
+	WaitForPendingStampUpdateConst();
 	// wait for previous Undo to finish (possibly never hit because the change records do it?)
-	WaitForPendingUndoRedo();
+	WaitForPendingUndoRedoUpdate();
 
 	FDynamicMesh3* Mesh = GetSculptMesh();
 
@@ -1747,6 +1743,9 @@ void UMeshVertexSculptTool::OnDynamicMeshComponentChanged(UDynamicMeshComponent*
 		UE::Geometry::VertexToTriangleOneRing(Mesh, Vertices, AccumulatedTriangleROI);
 	}, bRevert);
 
+	UndoUpdateFuture.Reset();
+	bUndoUpdatePending = true;
+
 	// start the normal recomputation
 	UndoNormalsFuture = Async(VertexSculptToolAsyncExecTarget, [this, Mesh]()
 	{
@@ -1768,8 +1767,14 @@ void UMeshVertexSculptTool::OnDynamicMeshComponentChanged(UDynamicMeshComponent*
 		return true;
 	});
 
-	// note that we have a pending update
-	bUndoUpdatePending = true;
+	UndoUpdateFuture = Async(VertexSculptToolAsyncExecTarget, [this]()
+	{
+		UndoNormalsFuture.Wait();
+		UndoUpdateOctreeFuture.Wait();
+		UndoUpdateBaseMeshFuture.Wait();
+		bUndoUpdatePending = false;
+	});
+
 }
 
 
diff --git a/Engine/Plugins/Runtime/MeshModelingToolset/Source/MeshModelingTools/Private/Sculpting/MeshSculptToolBase.cpp b/Engine/Plugins/Runtime/MeshModelingToolset/Source/MeshModelingTools/Private/Sculpting/MeshSculptToolBase.cpp
--- a/Engine/Plugins/Runtime/MeshModelingToolset/Source/MeshModelingTools/Private/Sculpting/MeshSculptToolBase.cpp
+++ b/Engine/Plugins/Runtime/MeshModelingToolset/Source/MeshModelingTools/Private/Sculpting/MeshSculptToolBase.cpp
@@ -10,6 +10,8 @@
 #include "BaseGizmos/GizmoComponents.h"
 #include "BaseGizmos/TransformGizmoUtil.h"
 #include "Drawing/MeshDebugDrawing.h"
+#include "Drawing/PreviewGeometryActor.h"
+
 
 #include "Sculpting/StampFalloffs.h"
 
@@ -17,6 +19,8 @@
 
 #include "ModelingToolTargetUtil.h"
 #include "BaseBehaviors/TwoAxisPropertyEditBehavior.h"
+#include "Drawing/PreviewGeometryActor.h"
+#include "Polyline3.h"
 
 #include UE_INLINE_GENERATED_CPP_BY_NAME(MeshSculptToolBase)
 
@@ -30,6 +34,13 @@ namespace
 	const FString VertexSculptIndicatorGizmoType = TEXT("VertexSculptIndicatorGizmoType");
 }
 
+namespace MeshSculptBaseLocals
+{
+	const FString& StrokePointSetID(TEXT("StrokePointSet"));
+	const FString& StrokeLineSetID(TEXT("StrokeLineSet"));
+
+	static FAutoConsoleVariable EnableMouseTrail(TEXT("modeling.Sculpting.EnableMouseTrail"), 0.0f, TEXT("Enable visualizing the path the mouse takes while sculpting."));
+}
 
 void FBrushToolRadius::InitializeWorldSizeRange(TInterval<float> Range, bool bValidateWorldRadius)
 {
@@ -154,6 +165,14 @@ void UMeshSculptToolBase::Setup()
 	ViewProperties->WatchProperty(ViewProperties->Image,
 		[this](UTexture2D* NewImage) { UpdateImageSetting(NewImage); });
 	
+	// Set up all the components we need to visualize things.
+	StrokeGeometry = NewObject<UPreviewGeometry>();
+	StrokeGeometry->CreateInWorld(TargetWorld, FTransform());
+
+	// These visualize the current spline edges that would be extracted
+	StrokeGeometry->AddPointSet(MeshSculptBaseLocals::StrokePointSetID);
+	StrokeGeometry->AddLineSet(MeshSculptBaseLocals::StrokeLineSetID);
+
 	// add input behavior to click-drag while holding hotkey to adjust brush size and strength
     {
         BrushEditBehavior = NewObject<ULocalTwoAxisPropertyEditInputBehavior>(this);
@@ -218,6 +237,12 @@ void UMeshSculptToolBase::OnCompleteSetup()
 
 void UMeshSculptToolBase::Shutdown(EToolShutdownType ShutdownType)
 {
+	if (StrokeGeometry)
+	{
+		StrokeGeometry->Disconnect();
+		StrokeGeometry = nullptr;
+	}
+
 	if (ShutdownType == EToolShutdownType::Accept && AreAllTargetsValid() == false)
 	{
 		UE_LOG(LogTemp, Error, TEXT("Tool Target has become Invalid (possibly it has been Force Deleted). Aborting Tool."));
@@ -308,6 +333,104 @@ void UMeshSculptToolBase::OnTick(float DeltaTime)
 
 	UpdateHoverStamp(GetBrushFrameWorld());
 
+	if (MeshSculptBaseLocals::EnableMouseTrail->GetBool() && StrokePath.Num() > 1)
+	{
+		TArray<FVector3d> UnProjectedStampPoints;
+
+		double Spacing = FMath::Clamp(BrushProperties->Spacing * BrushProperties->BrushSize.GetWorldRadius() * 2 , 0.1, 100000);
+
+		TArray<UE::Geometry::FPolyline3d> SpacingSections;
+		int32 SectionIndex = 0;
+		SpacingSections.Add(UE::Geometry::FPolyline3d());
+		SpacingSections[SectionIndex].AppendVertex(StrokePath[0].Position);
+		UnProjectedStampPoints.Add(StrokePath[0].Position);
+		for (int32 Index = 0; Index < StrokePath.Num() - 1; Index++)
+		{
+			FVector3d Span = StrokePath[Index + 1].Position - StrokePath[Index].Position;
+			FVector3d SpanDir = Span.GetSafeNormal();
+
+			double RemainingSpace = FMath::Max(0, Spacing - SpacingSections[SectionIndex].Length());
+			double OverflowSpace = Span.Length() - RemainingSpace;
+			if (OverflowSpace < 0)
+			{
+				SpacingSections[SectionIndex].AppendVertex(StrokePath[Index + 1].Position);
+			}
+			else
+			{
+				SpacingSections[SectionIndex].AppendVertex(StrokePath[Index].Position + (SpanDir * RemainingSpace));
+				
+				while (OverflowSpace > Spacing)
+				{
+					FVector3d TerminalPosition = SpacingSections[SectionIndex].End();
+					SpacingSections.Add(UE::Geometry::FPolyline3d());
+					SectionIndex++;
+					UnProjectedStampPoints.Add(TerminalPosition);
+					SpacingSections[SectionIndex].AppendVertex(TerminalPosition);
+					SpacingSections[SectionIndex].AppendVertex(TerminalPosition + (SpanDir * Spacing));
+					OverflowSpace -= Spacing;
+				}
+				
+				FVector3d TerminalPosition = SpacingSections[SectionIndex].End();
+				SpacingSections.Add(UE::Geometry::FPolyline3d());
+				SectionIndex++;
+				UnProjectedStampPoints.Add(TerminalPosition);
+				SpacingSections[SectionIndex].AppendVertex(TerminalPosition);
+				SpacingSections[SectionIndex].AppendVertex(TerminalPosition + (SpanDir * OverflowSpace));
+			}
+		}
+
+		ULineSetComponent* LineSet = StrokeGeometry->FindLineSet(MeshSculptBaseLocals::StrokeLineSetID);
+		if (LineSet)
+		{
+			LineSet->Clear();
+			for (int32 StrokeSectionIndex = 0; StrokeSectionIndex < SpacingSections.Num(); ++StrokeSectionIndex)
+			{
+				FColor SectionColor = (StrokeSectionIndex % 2) ? FColor::Red : FColor::Cyan;
+				for (int32 SegmentIndex = 0; SegmentIndex < SpacingSections[StrokeSectionIndex].SegmentCount(); ++SegmentIndex)
+				{
+
+					LineSet->AddLine(SpacingSections[StrokeSectionIndex].GetSegment(SegmentIndex).StartPoint(),
+									 SpacingSections[StrokeSectionIndex].GetSegment(SegmentIndex).EndPoint(),
+									 SectionColor, 3.0);
+				}
+			}
+		}
+
+		UPointSetComponent* PointSet = StrokeGeometry->FindPointSet(MeshSculptBaseLocals::StrokePointSetID);
+		if (PointSet)
+		{
+			PointSet->Clear();
+			for (int32 UnProjectedStampPointsIndex = 0; UnProjectedStampPointsIndex < UnProjectedStampPoints.Num(); ++UnProjectedStampPointsIndex)
+			{
+				FRay3d CameraToPathRay(CameraState.Position, (UnProjectedStampPoints[UnProjectedStampPointsIndex] - CameraState.Position).GetSafeNormal());
+
+				TUniquePtr<FMeshSculptBrushOp>& UseBrushOp = GetActiveBrushOp();
+				ESculptBrushOpTargetType TargetType = UseBrushOp->GetBrushTargetType();
+
+				FVector3d ProjectedPosition;
+				FVector3d ProjectedNormal;
+				int32 HitTriangle = IndexConstants::InvalidID;
+
+				switch (TargetType)
+				{
+				case ESculptBrushOpTargetType::SculptMesh:
+					ProjectWorldRayOnSculptMesh(CameraToPathRay, true, ProjectedPosition, ProjectedNormal, HitTriangle);
+					break;
+				case ESculptBrushOpTargetType::TargetMesh:
+					ProjectWorldRayOnTargetMesh(CameraToPathRay, true, ProjectedPosition, ProjectedNormal, HitTriangle);
+					break;
+				case ESculptBrushOpTargetType::ActivePlane:
+					ProjectWorldRayOnActivePlane(CameraToPathRay, ProjectedPosition, ProjectedNormal);
+					break;
+				}				
+
+				FColor PointColor = FColor::Green;
+				PointSet->AddPoint(ProjectedPosition, PointColor, 6.0, 0.0);
+			}
+		}
+	}
+
+
 	// always using HoverStamp here because it's position should always be up-to-date for the current apply-stamp...
 	BrushIndicator->Update((float)GetCurrentBrushRadius(),
 		HoverStamp.WorldFrame.ToFTransform(), 1.0f - (float)GetCurrentBrushFalloff());
@@ -610,6 +733,19 @@ bool UMeshSculptToolBase::HitTest(const FRay& Ray, FHitResult& OutHit)
 }
 
 
+void UMeshSculptToolBase::ClearStrokePath()
+{
+	StrokePath.Empty();
+}
+
+void UMeshSculptToolBase::AppendStrokePath(const FVector3d& NextPoint)
+{
+	if(StrokePath.Num() == 0 || !FMath::IsNearlyZero((StrokePath.Last().Position - NextPoint).SquaredLength()))
+	{
+		StrokePath.Add({ .Timestamp = ActiveStrokeTime, .Position = NextPoint, .Pressure = GetActivePressure() });
+	}
+}
+
 void UMeshSculptToolBase::OnBeginDrag(const FRay& WorldRay)
 {
 	SaveActiveStrokeModifiers();
@@ -638,6 +774,28 @@ void UMeshSculptToolBase::OnUpdateDrag(const FRay& WorldRay)
 	if (InStroke())
 	{
 		PendingStampRay = WorldRay;
+
+		TUniquePtr<FMeshSculptBrushOp>& UseBrushOp = GetActiveBrushOp();
+		ESculptBrushOpTargetType TargetType = UseBrushOp->GetBrushTargetType();
+
+		FVector3d ProjectedPosition;
+		FVector3d ProjectedNormal;
+		int32 HitTriangle = IndexConstants::InvalidID;
+
+		switch (TargetType)
+		{
+		case ESculptBrushOpTargetType::SculptMesh:
+			ProjectWorldRayOnSculptMesh(WorldRay, true, ProjectedPosition, ProjectedNormal, HitTriangle);
+			break;
+		case ESculptBrushOpTargetType::TargetMesh:
+			ProjectWorldRayOnTargetMesh(WorldRay, true, ProjectedPosition, ProjectedNormal, HitTriangle);
+			break;
+		case ESculptBrushOpTargetType::ActivePlane:
+			ProjectWorldRayOnActivePlane(WorldRay, ProjectedPosition, ProjectedNormal);
+			break;
+		}
+
+		AppendStrokePath(ProjectedPosition);
 	}
 }
 
@@ -746,80 +904,111 @@ void UMeshSculptToolBase::UpdateBrushTargetPlaneFromHit(const FRay& WorldRayIn,
 	ActiveBrushTargetPlaneWorld = FFrame3d(WorldPosWithBrushDepth, -WorldRay.Direction);
 }
 
-bool UMeshSculptToolBase::UpdateBrushPositionOnActivePlane(const FRay& WorldRayIn)
+bool UMeshSculptToolBase::ProjectWorldRayOnActivePlane(const FRay& WorldRayIn, FVector3d& ProjectedPosition, FVector3d& ProjectedNormal) const
 {
-	LastBrushTriangleID = IndexConstants::InvalidID;
-
-	FRay3d WorldRay(WorldRayIn);
-	FVector3d NewHitPosWorld;
-	ActiveBrushTargetPlaneWorld.RayPlaneIntersection(WorldRay.Origin, WorldRay.Direction, 2, NewHitPosWorld);
-	UpdateBrushFrameWorld(NewHitPosWorld, ActiveBrushTargetPlaneWorld.Z());
+	FRay3d WorldRay(WorldRayIn);	
+	ActiveBrushTargetPlaneWorld.RayPlaneIntersection(WorldRay.Origin, WorldRay.Direction, 2, ProjectedPosition);
+	ProjectedNormal = ActiveBrushTargetPlaneWorld.Z();
 	return true;
 }
 
-bool UMeshSculptToolBase::UpdateBrushPositionOnTargetMesh(const FRay& WorldRayIn, bool bFallbackToViewPlane)
+bool UMeshSculptToolBase::ProjectWorldRayOnTargetMesh(const FRay& WorldRayIn, bool bFallbackToViewPlane, FVector3d& ProjectedPosition, FVector3d& ProjectedNormal, int32& HitTriangle) const
 {
-	LastBrushTriangleID = IndexConstants::InvalidID;
-
 	FRay3d WorldRay(WorldRayIn);
 	FRay3d LocalRay = GetLocalRay(WorldRayIn);
-	int32 HitTID = FindHitTargetMeshTriangle(LocalRay);
+	int32 HitTID = FindHitTargetMeshTriangleConst(LocalRay);
 	if (HitTID != IndexConstants::InvalidID)
 	{
 		const FDynamicMesh3* BaseMesh = GetBaseMesh();
 		FIntrRay3Triangle3d Query = TMeshQueries<FDynamicMesh3>::TriangleIntersection(*BaseMesh, HitTID, LocalRay);
-		FVector3d WorldNormal = CurTargetTransform.TransformNormal(BaseMesh->GetTriNormal(HitTID));
-		FVector3d WorldPos = CurTargetTransform.TransformPosition(LocalRay.PointAt(Query.RayParameter));
-		LastBrushTriangleID = HitTID;
-		UpdateBrushFrameWorld(WorldPos, WorldNormal);
+		ProjectedNormal = CurTargetTransform.TransformNormal(BaseMesh->GetTriNormal(HitTID));
+		ProjectedPosition = CurTargetTransform.TransformPosition(LocalRay.PointAt(Query.RayParameter));
+		HitTriangle = HitTID;
 		return true;
 	}
 
 	if (bFallbackToViewPlane)
 	{
 		FFrame3d BrushPlane(GetBrushFrameWorld().Origin, (FVector3d)CameraState.Forward());
-		FVector3d NewHitPosWorld;
-		BrushPlane.RayPlaneIntersection(WorldRay.Origin, WorldRay.Direction, 2, NewHitPosWorld);
-		UpdateBrushFrameWorld(NewHitPosWorld, ActiveBrushTargetPlaneWorld.Z());
+		BrushPlane.RayPlaneIntersection(WorldRay.Origin, WorldRay.Direction, 2, ProjectedPosition);
+		ProjectedNormal = ActiveBrushTargetPlaneWorld.Z();
 		return true;
 	}
 
 	return false;
 }
 
-bool UMeshSculptToolBase::UpdateBrushPositionOnSculptMesh(const FRay& WorldRayIn, bool bFallbackToViewPlane)
+bool UMeshSculptToolBase::ProjectWorldRayOnSculptMesh(const FRay& WorldRayIn, bool bFallbackToViewPlane, FVector3d& ProjectedPosition, FVector3d& ProjectedNormal, int32& HitTriangle) const
 {
-	LastBrushTriangleID = IndexConstants::InvalidID;
+	HitTriangle = IndexConstants::InvalidID;
 
 	FRay3d WorldRay(WorldRayIn);
 	FRay3d LocalRay = GetLocalRay(WorldRayIn);
-	int32 HitTID = FindHitSculptMeshTriangle(LocalRay);
+	int32 HitTID = FindHitSculptMeshTriangleConst(LocalRay);
 	if (HitTID != IndexConstants::InvalidID)
 	{
 		const FDynamicMesh3* SculptMesh = GetSculptMesh();
 		FIntrRay3Triangle3d Query = TMeshQueries<FDynamicMesh3>::TriangleIntersection(*SculptMesh, HitTID, LocalRay);
-		FVector3d WorldNormal = CurTargetTransform.TransformNormal(SculptMesh->GetTriNormal(HitTID));
-		FVector3d WorldPos = CurTargetTransform.TransformPosition(LocalRay.PointAt(Query.RayParameter));
-		LastBrushTriangleID = HitTID;
-		UpdateBrushFrameWorld(WorldPos, WorldNormal);
+		ProjectedNormal = CurTargetTransform.TransformNormal(SculptMesh->GetTriNormal(HitTID));
+		ProjectedPosition = CurTargetTransform.TransformPosition(LocalRay.PointAt(Query.RayParameter));
+		HitTriangle = HitTID;
 		return true;
 	}
 
 	if (bFallbackToViewPlane)
 	{
 		FFrame3d BrushPlane(GetBrushFrameWorld().Origin, (FVector3d)CameraState.Forward());
-		FVector3d NewHitPosWorld;
-		BrushPlane.RayPlaneIntersection(WorldRay.Origin, WorldRay.Direction, 2, NewHitPosWorld);
-		UpdateBrushFrameWorld(NewHitPosWorld, ActiveBrushTargetPlaneWorld.Z());
+		BrushPlane.RayPlaneIntersection(WorldRay.Origin, WorldRay.Direction, 2, ProjectedPosition);
+		ProjectedNormal = ActiveBrushTargetPlaneWorld.Z();
 		return true;
 	}
 
 	return false;
 }
 
+bool UMeshSculptToolBase::UpdateBrushPositionOnActivePlane(const FRay& WorldRayIn)
+{
+	LastBrushTriangleID = IndexConstants::InvalidID;
+	FVector3d NewHitPosWorld;
+	FVector3d NewHitNormalWorld;
+	bool bHit = ProjectWorldRayOnActivePlane(WorldRayIn, NewHitPosWorld, NewHitNormalWorld);
+	UpdateBrushFrameWorld(NewHitPosWorld, ActiveBrushTargetPlaneWorld.Z());
+	return bHit;
+}
 
+bool UMeshSculptToolBase::UpdateBrushPositionOnTargetMesh(const FRay& WorldRayIn, bool bFallbackToViewPlane)
+{
+	LastBrushTriangleID = IndexConstants::InvalidID;
+	FRay3d LocalRay = GetLocalRay(WorldRayIn);
+	FVector3d NewHitPosWorld;
+	FVector3d NewHitNormalWorld;
+	int32 HitTID;
+	bool bHit = ProjectWorldRayOnTargetMesh(WorldRayIn, bFallbackToViewPlane, NewHitPosWorld, NewHitNormalWorld, HitTID);
+	UpdateBrushFrameWorld(NewHitPosWorld, NewHitNormalWorld);
+	UpdateHitTargetMeshTriangle(HitTID, LocalRay);
+	if (HitTID != IndexConstants::InvalidID)
+	{
+		LastBrushTriangleID = HitTID;
+	}
+	return bHit;
+}
 
-
+bool UMeshSculptToolBase::UpdateBrushPositionOnSculptMesh(const FRay& WorldRayIn, bool bFallbackToViewPlane)
+{
+	LastBrushTriangleID = IndexConstants::InvalidID;
+	FRay3d LocalRay = GetLocalRay(WorldRayIn);
+	FVector3d NewHitPosWorld;
+	FVector3d NewHitNormalWorld;
+	int32 HitTID;
+	bool bHit = ProjectWorldRayOnSculptMesh(WorldRayIn, bFallbackToViewPlane, NewHitPosWorld, NewHitNormalWorld, HitTID);
+	UpdateBrushFrameWorld(NewHitPosWorld, NewHitNormalWorld);
+	UpdateHitSculptMeshTriangle(HitTID, LocalRay);
+	if (HitTID != IndexConstants::InvalidID)
+	{
+		LastBrushTriangleID = HitTID;
+	}
+	return bHit;
+}
 
 void UMeshSculptToolBase::SaveActiveStrokeModifiers()
 {
@@ -840,55 +1029,217 @@ float UMeshSculptToolBase::GetStampTemporalFlowRate() const
 	return BrushProperties->FlowRate;
 }
 
+void UMeshSculptToolBase::ProcessPerTickStamps(TFunction<bool(const FRay& StampRay)> UpdateStampPosition, TFunction<void(const FRay& StampRay)> ExecuteStampOperation)
+{
+	if (InStroke())
+	{
+		TRACE_CPUPROFILER_EVENT_SCOPE(BaseSculptTool_Tick_ProcessPerTickStamps);
+		// We save the active pressure here, so we can then override and fake it during the stamp application to generate interpolated pressures along the stroke.
+		TGuardValue<double> CachedActivePressure(ActivePressure, ActivePressure);
+
+		// update brush position
+		if (UpdateStampPosition(GetPendingStampRayWorld()) == false)
+		{
+			return;
+		}
+
+		UpdateStampPendingState();
+		if (IsStampPending() == false)
+		{
+			return;
+		}
+
+		if (!PendingStampRays.IsEmpty())
+		{
+			for (const FStampRayData& StampRay : PendingStampRays)
+			{
+				ActivePressure = StampRay.Pressure;
+
+				UpdateStampPosition(StampRay.StampRay);
+				ExecuteStampOperation(StampRay.StampRay);
+				StampsDuringStroke++;
+			}
+		}
+
+	}
+}
+
+
 void UMeshSculptToolBase::UpdateStampPendingState()
 {
 	if (InStroke() == false) return;
+	PendingStampRays.Empty();
 
-	bool bFlowStampPending = false;
-	float UseStampFlowRate = GetStampTemporalFlowRate();
-	if (UseStampFlowRate >= 1.0)
+	// Fail out early with a stamp if we don't support the currently requested stroke type
+	if (!GetActiveBrushOp()->SupportsStrokeType(BrushProperties->StrokeType))
 	{
-		bFlowStampPending = true;
+		PendingStampRays.Add({ .StampRay = PendingStampRay, .Pressure = GetActivePressure() });
+		bIsStampPending = true;
+		return;
 	}
-	else if (UseStampFlowRate == 0.0)
+
+	switch (BrushProperties->StrokeType)
 	{
-		bFlowStampPending = (LastFlowTimeStamp++ == 0);
-	}
-	else
+	case EMeshSculptStrokeType::Airbrush:
 	{
-		double dt = (1.0 - UseStampFlowRate);
-		int FlowTimestamp = (int)(ActiveStrokeTime / dt);
-		if (FlowTimestamp > LastFlowTimeStamp)
+		bool bFlowStampPending = false;
+		float UseStampFlowRate = GetStampTemporalFlowRate();
+		if (UseStampFlowRate >= 1.0)
 		{
-			LastFlowTimeStamp = FlowTimestamp;
 			bFlowStampPending = true;
 		}
-	}
+		else if (UseStampFlowRate == 0.0)
+		{
+			bFlowStampPending = (LastFlowTimeStamp++ == 0);
+		}
+		else
+		{
+			double dt = (1.0 - UseStampFlowRate);
+			int FlowTimestamp = (int)(ActiveStrokeTime / dt);
+			if (FlowTimestamp > LastFlowTimeStamp)
+			{
+				LastFlowTimeStamp = FlowTimestamp;
+				bFlowStampPending = true;
+			}
+		}
 
-	// if brush does not support variable spacing, we only use flow
-	if (GetActiveBrushOp()->SupportsStrokeType(EMeshSculptStrokeType::Spacing) == false)
-	{
+		PendingStampRays.Add({ .StampRay = PendingStampRay, .Pressure = GetActivePressure() });
 		bIsStampPending = bFlowStampPending;
 		return;
 	}
-
-	bool bSpacingStampPending = false;
-	if (BrushProperties->Spacing == 0)
-	{
-		bSpacingStampPending = true;
-	}
-	else
+		break;
+	case EMeshSculptStrokeType::Spacing:
 	{
-		double ArcSpacing = BrushProperties->Spacing * (2 * GetCurrentBrushRadius());
-		int SpacingTimestamp = 1 + (int)(ActiveStrokePathArcLen / ArcSpacing);
-		if (SpacingTimestamp > LastSpacingTimestamp)
+		// Save our active pressure since we'll be manipulating it during the next block
+		TGuardValue<double> CachedActivePressure(ActivePressure, ActivePressure);
+
+		TArray<FTimedStrokePoint> UnProjectedStampPoints;
+		if (StrokePath.Num() > 1)
+		{
+			TArray<FTimedStrokePoint> UnconsumedPathPoints;
+
+			int32 SectionIndex = 0;
+			double SpacingSectionConsumedLength = 0.0;
+
+			UnProjectedStampPoints.Add(StrokePath[0]);
+			UnconsumedPathPoints.Add(StrokePath[0]);
+
+			for (int32 Index = 0; Index < StrokePath.Num() - 1; Index++)
+			{
+				// Compute average spacing over stroke section - needed for dynamic spacing as the pressure changes over the path.
+				double Spacing = 1.0;
+				{
+					ActivePressure = StrokePath[Index].Pressure;
+					double SpacingAtBeginning = FMath::Clamp(BrushProperties->Spacing * GetActiveBrushRadius() * 2, 0.1, 100000);
+					ActivePressure = StrokePath[Index + 1].Pressure;
+					double SpacingAtEnding = FMath::Clamp(BrushProperties->Spacing * GetActiveBrushRadius() * 2, 0.1, 100000);
+					Spacing = (SpacingAtBeginning + SpacingAtEnding) * 0.5;
+				}
+
+				FVector3d Span = StrokePath[Index + 1].Position - StrokePath[Index].Position; // Vector representing our stroke segment
+				FVector3d SpanDir = Span.GetSafeNormal(); // The direction of our stroke segment
+
+				double RemainingSpace = FMath::Max(0, Spacing - SpacingSectionConsumedLength); // How much space do we have left before a stamp happens
+				double OverflowSpace = Span.Length() - RemainingSpace; // If positive, our stroke section meets or exceeds the remaining space for a stamp to happen
+
+				if (OverflowSpace < 0)
+				{
+					SpacingSectionConsumedLength += Span.Length();
+					UnconsumedPathPoints.Add(StrokePath[Index + 1]);
+				}
+				else
+				{
+					double SpanAlpha = FMath::IsNearlyZero(Span.Length()) ? 0.0 : RemainingSpace / Span.Length(); // How far along, within the current stroke segment, are we?
+					double OverflowAlphaIncrement = FMath::IsNearlyZero(Span.Length()) ? 0.0 : Spacing / Span.Length(); // For our current stroke segment, what is the percentage of it's length in terms of a spacing interval?
+					
+					FVector3d TerminalPosition = FMath::LerpStable(StrokePath[Index].Position, StrokePath[Index + 1].Position, SpanAlpha);
+					SpacingSectionConsumedLength += RemainingSpace;
+
+					UnProjectedStampPoints.Add({ .Position = TerminalPosition, .Pressure = FMath::LerpStable(StrokePath[Index].Pressure, StrokePath[Index + 1].Pressure, SpanAlpha) });
+
+					while (OverflowSpace > Spacing) // Handle cases where we have so much line we inject multiple points...
+					{
+						SpanAlpha += OverflowAlphaIncrement;
+						TerminalPosition = FMath::LerpStable(StrokePath[Index].Position, StrokePath[Index + 1].Position, SpanAlpha);
+						UnProjectedStampPoints.Add({ .Position = TerminalPosition, .Pressure = FMath::LerpStable(StrokePath[Index].Pressure, StrokePath[Index + 1].Pressure, SpanAlpha) });
+						OverflowSpace -= Spacing;
+					}
+
+					ensure(SpanAlpha <= 1.0);
+
+					SpacingSectionConsumedLength = OverflowSpace;
+					UnconsumedPathPoints.Empty();
+					UnconsumedPathPoints.Add(
+						{
+								.Timestamp = FMath::LerpStable(StrokePath[Index].Timestamp, StrokePath[Index + 1].Timestamp, SpanAlpha),
+								.Position = TerminalPosition,
+								.Pressure = FMath::LerpStable(StrokePath[Index].Pressure, StrokePath[Index + 1].Pressure, SpanAlpha)
+						});
+					UnconsumedPathPoints.Add({ StrokePath[Index + 1] });
+				}
+			}
+
+			StrokePath = UnconsumedPathPoints;
+		}
+
+		// We want to ensure that when starting the stroke, we stamp right at the beginning of the path
+		int StartIndex = 1;
+		if (StampsDuringStroke == 0)
+		{
+			StartIndex = 0;
+		}
+
+		for (int32 UnProjectedStampPointsIndex = StartIndex; UnProjectedStampPointsIndex < UnProjectedStampPoints.Num(); ++UnProjectedStampPointsIndex)
+		{
+			FRay3d CameraToPathRay(CameraState.Position, (UnProjectedStampPoints[UnProjectedStampPointsIndex].Position - CameraState.Position).GetSafeNormal());
+
+			TUniquePtr<FMeshSculptBrushOp>& UseBrushOp = GetActiveBrushOp();
+			ESculptBrushOpTargetType TargetType = UseBrushOp->GetBrushTargetType();
+
+			FVector3d ProjectedPosition = FVector3d::ZeroVector;
+			FVector3d ProjectedNormal = FVector3d::ZeroVector;
+			int32 HitTriangle = IndexConstants::InvalidID;
+
+			switch (TargetType)
+			{
+			case ESculptBrushOpTargetType::SculptMesh:
+				ProjectWorldRayOnSculptMesh(CameraToPathRay, true, ProjectedPosition, ProjectedNormal, HitTriangle);
+				break;
+			case ESculptBrushOpTargetType::TargetMesh:
+				ProjectWorldRayOnTargetMesh(CameraToPathRay, true, ProjectedPosition, ProjectedNormal, HitTriangle);
+				break;
+			case ESculptBrushOpTargetType::ActivePlane:
+				ProjectWorldRayOnActivePlane(CameraToPathRay, ProjectedPosition, ProjectedNormal);
+				break;
+			}
+
+			PendingStampRays.Add({ .StampRay = FRay3d(CameraState.Position, (ProjectedPosition - CameraState.Position).GetSafeNormal()), .Pressure = UnProjectedStampPoints[UnProjectedStampPointsIndex].Pressure });
+		}
+
+		if (!PendingStampRays.IsEmpty())
 		{
-			LastSpacingTimestamp = SpacingTimestamp;
-			bSpacingStampPending = true;
+			bIsStampPending = true;
 		}
+
+		return;
 	}
+		break;
+
+	case EMeshSculptStrokeType::Dots:
+	{
+		PendingStampRays.Add({ .StampRay = PendingStampRay, .Pressure = GetActivePressure() });
+		bIsStampPending = true;
+		return;
+	}
+		break;
+
+	default:
+		break;
+	}
+
+
+
 
-	bIsStampPending = bFlowStampPending && bSpacingStampPending;
 }
 
 
@@ -898,6 +1249,8 @@ void UMeshSculptToolBase::ResetStrokeTime()
 	LastFlowTimeStamp = 0;
 	ActiveStrokePathArcLen = 0;
 	LastSpacingTimestamp = 0;
+	StampsDuringStroke = 0;
+	ClearStrokePath();
 }
 
 void UMeshSculptToolBase::AccumulateStrokeTime(float DeltaTime)
diff --git a/Engine/Plugins/Runtime/MeshModelingToolset/Source/MeshModelingTools/Public/MeshGroupPaintTool.h b/Engine/Plugins/Runtime/MeshModelingToolset/Source/MeshModelingTools/Public/MeshGroupPaintTool.h
--- a/Engine/Plugins/Runtime/MeshModelingToolset/Source/MeshModelingTools/Public/MeshGroupPaintTool.h
+++ b/Engine/Plugins/Runtime/MeshModelingToolset/Source/MeshModelingTools/Public/MeshGroupPaintTool.h
@@ -360,8 +360,8 @@ protected:
 	virtual FDynamicMesh3* GetBaseMesh() { check(false); return nullptr; }
 	virtual const FDynamicMesh3* GetBaseMesh() const { check(false); return nullptr; }
 
-	UE_API virtual int32 FindHitSculptMeshTriangle(const FRay3d& LocalRay) override;
-	UE_API virtual int32 FindHitTargetMeshTriangle(const FRay3d& LocalRay) override;
+	UE_API virtual int32 FindHitSculptMeshTriangleConst(const FRay3d& LocalRay) const override;
+	UE_API virtual int32 FindHitTargetMeshTriangleConst(const FRay3d& LocalRay) const override;
 
 	UE_API virtual void OnBeginStroke(const FRay& WorldRay) override;
 	UE_API virtual void OnEndStroke() override;
@@ -432,7 +432,7 @@ protected:
 	TSet<int32> AccumulatedTriangleROI;
 	bool bUndoUpdatePending = false;
 	TArray<int> NormalsBuffer;
-	UE_API void WaitForPendingUndoRedo();
+	UE_API void WaitForPendingUndoRedoUpdate();
 
 	TArray<int> TempROIBuffer;
 	TArray<int> VertexROI;
diff --git a/Engine/Plugins/Runtime/MeshModelingToolset/Source/MeshModelingTools/Public/MeshVertexPaintTool.h b/Engine/Plugins/Runtime/MeshModelingToolset/Source/MeshModelingTools/Public/MeshVertexPaintTool.h
--- a/Engine/Plugins/Runtime/MeshModelingToolset/Source/MeshModelingTools/Public/MeshVertexPaintTool.h
+++ b/Engine/Plugins/Runtime/MeshModelingToolset/Source/MeshModelingTools/Public/MeshVertexPaintTool.h
@@ -182,6 +182,10 @@ public:
 	UPROPERTY(EditAnywhere, Category = Settings, meta = (DisplayName = "Action"))
 	EMeshVertexPaintInteractionType SubToolType = EMeshVertexPaintInteractionType::Brush;
 
+	// The strength of the brush effect when stamped on the mesh
+	UPROPERTY(EditAnywhere, Category = Settings, meta = (UIMin = 0, UIMax = 1))
+	double Strength = 1.0;
+
 	/** The Color that will be assigned to painted triangle vertices */
 	UPROPERTY(EditAnywhere, Category = Settings, meta = ())
 	FLinearColor PaintColor = FLinearColor::Red;
@@ -558,8 +562,8 @@ protected:
 	virtual FDynamicMesh3* GetBaseMesh() { check(false); return nullptr; }
 	virtual const FDynamicMesh3* GetBaseMesh() const { check(false); return nullptr; }
 
-	UE_API virtual int32 FindHitSculptMeshTriangle(const FRay3d& LocalRay) override;
-	UE_API virtual int32 FindHitTargetMeshTriangle(const FRay3d& LocalRay) override;
+	UE_API virtual int32 FindHitSculptMeshTriangleConst(const FRay3d& LocalRay) const override;
+	UE_API virtual int32 FindHitTargetMeshTriangleConst(const FRay3d& LocalRay) const override;
 
 	UE_API virtual void OnBeginStroke(const FRay& WorldRay) override;
 	UE_API virtual void OnEndStroke() override;
@@ -651,7 +655,7 @@ protected:
 	TSet<int32> AccumulatedTriangleROI;
 	bool bUndoUpdatePending = false;
 	TArray<int> NormalsBuffer;
-	UE_API void WaitForPendingUndoRedo();
+	UE_API void WaitForPendingUndoRedoUpdate();
 
 	TArray<int> TempROIBuffer;
 	TArray<int> VertexROI;
diff --git a/Engine/Plugins/Runtime/MeshModelingToolset/Source/MeshModelingTools/Public/MeshVertexSculptTool.h b/Engine/Plugins/Runtime/MeshModelingToolset/Source/MeshModelingTools/Public/MeshVertexSculptTool.h
--- a/Engine/Plugins/Runtime/MeshModelingToolset/Source/MeshModelingTools/Public/MeshVertexSculptTool.h
+++ b/Engine/Plugins/Runtime/MeshModelingToolset/Source/MeshModelingTools/Public/MeshVertexSculptTool.h
@@ -320,9 +320,9 @@ protected:
 	virtual FDynamicMesh3* GetBaseMesh() override{ return &BaseMesh; }
 	virtual const FDynamicMesh3* GetBaseMesh() const override{ return &BaseMesh; }
 
-	UE_API virtual int32 FindHitSculptMeshTriangle(const FRay3d& LocalRay) override;
-	UE_API virtual int32 FindHitTargetMeshTriangle(const FRay3d& LocalRay) override;
-	UE_API bool IsHitTriangleBackFacing(int32 TriangleID, const FDynamicMesh3* QueryMesh);
+	UE_API virtual int32 FindHitSculptMeshTriangleConst(const FRay3d& LocalRay) const override;
+	UE_API virtual int32 FindHitTargetMeshTriangleConst(const FRay3d& LocalRay) const override;
+	UE_API bool IsHitTriangleBackFacing(int32 TriangleID, const FDynamicMesh3* QueryMesh) const;
 
 	UE_API virtual void UpdateHoverStamp(const FFrame3d& StampFrameWorld) override;
 
@@ -355,18 +355,26 @@ protected:
 	int32 InitialStrokeTriangleID = -1;
 
 	TSet<int32> AccumulatedTriangleROI;
+	
 	bool bUndoUpdatePending = false;
 	TFuture<bool> UndoNormalsFuture;
 	TFuture<bool> UndoUpdateOctreeFuture;
 	TFuture<bool> UndoUpdateBaseMeshFuture;
+	TFuture<void> UndoUpdateFuture;
 	TArray<int> NormalsBuffer;
-	UE_API void WaitForPendingUndoRedo();
+
+	UE_DEPRECATED(5.7, "Vertex sculpt now uses a const version of WaitForPendingUndoRedo")
+	UE_API void WaitForPendingUndoRedo() {};
+	UE_API void WaitForPendingUndoRedoUpdate() const;
 
 	TArray<uint32> OctreeUpdateTempBuffer;
 	TArray<bool> OctreeUpdateTempFlagBuffer;
 	TFuture<void> StampUpdateOctreeFuture;
 	bool bStampUpdatePending = false;
-	UE_API void WaitForPendingStampUpdate();
+
+	UE_DEPRECATED(5.7, "Vertex sculpt now uses a const version of WaitForPendingStampUpdate")
+	UE_API void WaitForPendingStampUpdate() {};
+	UE_API void WaitForPendingStampUpdateConst() const;
 
 	TArray<int> RangeQueryTriBuffer;
 	UE::Geometry::FUniqueIndexSet VertexROIBuilder;
diff --git a/Engine/Plugins/Runtime/MeshModelingToolset/Source/MeshModelingTools/Public/Sculpting/MeshBrushOpBase.h b/Engine/Plugins/Runtime/MeshModelingToolset/Source/MeshModelingTools/Public/Sculpting/MeshBrushOpBase.h
--- a/Engine/Plugins/Runtime/MeshModelingToolset/Source/MeshModelingTools/Public/Sculpting/MeshBrushOpBase.h
+++ b/Engine/Plugins/Runtime/MeshModelingToolset/Source/MeshModelingTools/Public/Sculpting/MeshBrushOpBase.h
@@ -105,12 +105,15 @@ public:
 	bool bIsStrengthPressureEnabled = true;
 };
 
-/** Mesh Sculpting Brush Stroke Behaviors */
+/** Mesh Sculpting Brush Stroke Types */
 UENUM()
 enum class EMeshSculptStrokeType : uint8
 {
+	// Brushes stamp at regular distance intervals as the stroke is drawn
 	Spacing = 0,
+	// Brushes stamp at regular time intervals as the stroke is drawn
 	Airbrush = 1,
+	// Brushes stamp once per frame as the stroke is drawn.
 	Dots = 2,
 
 	LastValue UMETA(Hidden)
@@ -272,9 +275,9 @@ public:
 		return false;
 	}
 
-	virtual bool SupportsStrokeType(EMeshSculptStrokeType StrokeBehavior) const
+	virtual bool SupportsStrokeType(EMeshSculptStrokeType StrokeType) const
 	{
-		switch (StrokeBehavior)
+		switch (StrokeType)
 		{
 		case EMeshSculptStrokeType::Airbrush:
 		case EMeshSculptStrokeType::Dots:
diff --git a/Engine/Plugins/Runtime/MeshModelingToolset/Source/MeshModelingTools/Public/Sculpting/MeshSculptBrushOps.h b/Engine/Plugins/Runtime/MeshModelingToolset/Source/MeshModelingTools/Public/Sculpting/MeshSculptBrushOps.h
--- a/Engine/Plugins/Runtime/MeshModelingToolset/Source/MeshModelingTools/Public/Sculpting/MeshSculptBrushOps.h
+++ b/Engine/Plugins/Runtime/MeshModelingToolset/Source/MeshModelingTools/Public/Sculpting/MeshSculptBrushOps.h
@@ -52,9 +52,9 @@ public:
 		return true;
 	}
 
-	virtual bool SupportsStrokeType(EMeshSculptStrokeType StrokeBehavior) const override
+	virtual bool SupportsStrokeType(EMeshSculptStrokeType StrokeType) const override
 	{
-		switch (StrokeBehavior)
+		switch (StrokeType)
 		{
 		case EMeshSculptStrokeType::Airbrush:
 		case EMeshSculptStrokeType::Dots:
@@ -153,9 +153,9 @@ public:
 	{
 		return true;
 	}
-	virtual bool SupportsStrokeType(EMeshSculptStrokeType StrokeBehavior) const override
+	virtual bool SupportsStrokeType(EMeshSculptStrokeType StrokeType) const override
 	{
-		switch (StrokeBehavior)
+		switch (StrokeType)
 		{
 		case EMeshSculptStrokeType::Airbrush:
 		case EMeshSculptStrokeType::Dots:
@@ -263,9 +263,9 @@ public:
 	{
 		return true;
 	}
-	virtual bool SupportsStrokeType(EMeshSculptStrokeType StrokeBehavior) const override
+	virtual bool SupportsStrokeType(EMeshSculptStrokeType StrokeType) const override
 	{
-		switch (StrokeBehavior)
+		switch (StrokeType)
 		{
 		case EMeshSculptStrokeType::Airbrush:
 		case EMeshSculptStrokeType::Dots:
diff --git a/Engine/Plugins/Runtime/MeshModelingToolset/Source/MeshModelingTools/Public/Sculpting/MeshSculptToolBase.h b/Engine/Plugins/Runtime/MeshModelingToolset/Source/MeshModelingTools/Public/Sculpting/MeshSculptToolBase.h
--- a/Engine/Plugins/Runtime/MeshModelingToolset/Source/MeshModelingTools/Public/Sculpting/MeshSculptToolBase.h
+++ b/Engine/Plugins/Runtime/MeshModelingToolset/Source/MeshModelingTools/Public/Sculpting/MeshSculptToolBase.h
@@ -21,6 +21,7 @@ class UMaterialInstanceDynamic;
 class UCombinedTransformGizmo;
 class UTransformProxy;
 class UPreviewMesh;
+class UPreviewGeometry;
 
 /**
  * Type of Brush Size currently active in FBrushToolRadius
@@ -100,15 +101,15 @@ enum class EMeshSculptFalloffType : uint8
 };
 
 
-
-
-
 UCLASS(MinimalAPI)
 class USculptBrushProperties : public UInteractiveToolPropertySet
 {
 	GENERATED_BODY()
 public:
 
+	UPROPERTY(EditAnywhere, Category = Stroke)
+	EMeshSculptStrokeType StrokeType = EMeshSculptStrokeType::Spacing;
+
 	UPROPERTY(EditAnywhere, Category = Brush, meta = (DisplayPriority = 1, ModelingQuickEdit, ModelingQuickSettings = 100))
 	FBrushToolRadius BrushSize;
 
@@ -129,12 +130,12 @@ public:
 	bool bHitBackFaces = true;
 
 	/** Brush stamps are applied at this time interval. 0 for a single stamp, 1 for continuous stamps, 0.5 is a stamp every half-second */
-	UPROPERTY(EditAnywhere, Category = Brush, meta = (DisplayName = "Flow", UIMin = "0.0", UIMax = "1.0", ClampMin = "0.0", ClampMax = "1.0", DisplayPriority = 7, HideEditConditionToggle, EditConditionHides, EditCondition = "bShowFlowRate"))
+	UPROPERTY(EditAnywhere, Category = Brush, meta = (DisplayName = "Flow", UIMin = "0.0", UIMax = "1.0", ClampMin = "0.0", ClampMax = "1.0", DisplayPriority = 7, HideEditConditionToggle, EditConditionHides, EditCondition = "bShowFlowRate && StrokeType == EMeshSculptStrokeType::Airbrush"))
 	float FlowRate = 1.0f;
 
 	/** Space out stamp centers at distances Spacing*BrushDiameter along the stroke (so Spacing of 1 means that stamps will be adjacent but non-overlapping). Zero spacing means continuous stamps. */
-	UPROPERTY(EditAnywhere, Category = Brush, meta = (UIMin = "0.0", UIMax = "4.0", ClampMin = "0.0", ClampMax = "1000.0", DisplayPriority = 8, HideEditConditionToggle, EditConditionHides, EditCondition = "bShowSpacing"))
-	float Spacing = 0.0f;
+	UPROPERTY(EditAnywhere, Category = Brush, meta = (UIMin = "0.0", UIMax = "4.0", ClampMin = "0.0", ClampMax = "1000.0", DisplayPriority = 8, HideEditConditionToggle, EditConditionHides, EditCondition = "bShowSpacing && StrokeType == EMeshSculptStrokeType::Spacing"))
+	float Spacing = .05f;
 
 	/** Lazy brush smooths out the brush path by averaging the cursor positions */
 	UPROPERTY(EditAnywhere, Category = Brush, meta = (DisplayName = "Lazy", UIMin = "0", UIMax = "1.0", ClampMin = "0", ClampMax = "1.0", DisplayPriority = 9, HideEditConditionToggle, EditConditionHides, EditCondition = "bShowLazyness"))
@@ -464,16 +465,44 @@ protected:
 
 	bool GetBrushCanHitBackFaces() const { return BrushProperties->bHitBackFaces; }
 
-	/** @return hit triangle at ray position - subclass must implement this */
-	virtual int32 FindHitSculptMeshTriangle(const FRay3d& LocalRay) { check(false); return IndexConstants::InvalidID; }
-	/** @return hit triangle at ray position - subclass should implement this for most brushes */
-	virtual int32 FindHitTargetMeshTriangle(const FRay3d& LocalRay) { check(false); return IndexConstants::InvalidID;	}
+	// These routines handle the actual ray intersection queries againt the sculpt and target meshes. There 
+	// are two versions for each - a const version and a non-const version. The versions have different names
+	// to avoid running afoul of deprecation rules. Subclasses may leave the non-const versions unimplemented,
+	// as the default implementation simply calls the const version. Subclasses should implement the const versions,
+	// to provide correct ray intersection logic for thier specific sculpt/target situation.
+
+	/** @return hit triangle at ray position - subclass *may* implement this */
+	virtual int32 FindHitSculptMeshTriangle(const FRay3d& LocalRay) { return FindHitSculptMeshTriangleConst(LocalRay); }
+	/** @return hit triangle at ray position - subclass *must* implement this */
+	virtual int32 FindHitSculptMeshTriangleConst(const FRay3d& LocalRay) const { check(false); return IndexConstants::InvalidID; }
+	/** @return hit triangle at ray position - subclass *may* implement this */
+	virtual int32 FindHitTargetMeshTriangle(const FRay3d& LocalRay) { return FindHitTargetMeshTriangleConst(LocalRay); }
+	/** @return hit triangle at ray position - subclass *should* implement this for most brushes  */
+	virtual int32 FindHitTargetMeshTriangleConst(const FRay3d& LocalRay) const { check(false); return IndexConstants::InvalidID;	}
+
+	bool ProjectWorldRayOnActivePlane(const FRay& WorldRay, FVector3d& ProjectedPosition, FVector3d& ProjectedNormal) const;
+	bool ProjectWorldRayOnTargetMesh(const FRay& WorldRay, bool bFallbackToViewPlane, FVector3d& ProjectedPosition, FVector3d& ProjectedNormal, int32& HitTriangle) const;
+	bool ProjectWorldRayOnSculptMesh(const FRay& WorldRay, bool bFallbackToViewPlane, FVector3d& ProjectedPosition, FVector3d& ProjectedNormal, int32& HitTriangle) const;
 
 	UE_API virtual bool UpdateBrushPositionOnActivePlane(const FRay& WorldRay);
 	UE_API virtual bool UpdateBrushPositionOnTargetMesh(const FRay& WorldRay, bool bFallbackToViewPlane);
 	UE_API virtual bool UpdateBrushPositionOnSculptMesh(const FRay& WorldRay, bool bFallbackToViewPlane);
 
+	/**
+	* Provides a hook for informing subclasses which triangle was last hit during a brush update.
+	* Subclass should implement this if it needs to know which triangle was hit
+	* @param TriangleID the triangle last hit during brush updates on the Sculpt Mesh
+	* @param LocalRay the ray in local space which caused the triangle to be hit
+	*/
+	virtual void UpdateHitSculptMeshTriangle(int32 TriangleID, const FRay3d& LocalRay) {};
 
+	/**
+	* Provides a hook for informing subclasses which triangle was last hit during a brush update.
+	* Subclass should implement this if it needs to know which triangle was hit
+	* @param TriangleID the triangle last hit during brush updates on the Target Mesh
+	* @param LocalRay the ray in local space which caused the triangle to be hit
+	*/
+	virtual void UpdateHitTargetMeshTriangle(int32 TriangleID, const FRay3d& LocalRay) {};
 
 	//
 	// Brush Target Plane is plane that some brushes move on
@@ -495,6 +524,20 @@ protected:
 	virtual bool GetInSmoothingStroke() const { return bSmoothing; }
 	virtual bool GetInInvertStroke() const { return bInvert; }
 
+	struct FTimedStrokePoint
+	{
+		double Timestamp;
+		FVector3d Position;		
+		double Pressure;
+	};
+	TArray<FTimedStrokePoint> StrokePath;
+
+
+	UPROPERTY()
+	TObjectPtr<UPreviewGeometry> StrokeGeometry;
+	UE_API virtual void ClearStrokePath();
+	UE_API virtual void AppendStrokePath(const FVector3d& NextPoint);
+
 	// when in a stroke, this function determines when a new stamp should be emitted, based on spacing and flow rate settings
 	UE_API virtual void UpdateStampPendingState();
 
@@ -505,6 +548,7 @@ protected:
 	int LastSpacingTimestamp = 0;
 	UE_API virtual void ResetStrokeTime();
 	UE_API virtual void AccumulateStrokeTime(float DeltaTime);
+	int StampsDuringStroke = 0;
 
 	//
 	// Stamps
@@ -512,6 +556,14 @@ protected:
 protected:
 	bool bIsStampPending = false;
 	FRay PendingStampRay;
+
+	struct FStampRayData
+	{
+		FRay StampRay;
+		double Pressure;
+	};
+
+	TArray<FStampRayData> PendingStampRays;
 	FSculptBrushStamp HoverStamp;
 	FSculptBrushStamp CurrentStamp;
 	FSculptBrushStamp LastStamp;
@@ -521,6 +573,7 @@ protected:
 	// Temporal Flow Rate defines the frequency of stamp placement. 1 is max rate, 0 is no stamps. Defaults to BrushProperties->FlowRate, but subclasses can re-use that setting for other things
 	UE_API virtual float GetStampTemporalFlowRate() const;
 
+	UE_API void ProcessPerTickStamps(TFunction<bool(const FRay& StampRay)> UpdateStampPosition, TFunction<void(const FRay& StampRay)> ExecuteStampOperation);
 
 	//
 	// Stamp ROI Plane is a plane used by some brush ops
diff --git a/Engine/Plugins/Runtime/MeshModelingToolset/Source/MeshModelingTools/Public/Sculpting/MeshVertexPaintBrushOps.h b/Engine/Plugins/Runtime/MeshModelingToolset/Source/MeshModelingTools/Public/Sculpting/MeshVertexPaintBrushOps.h
--- a/Engine/Plugins/Runtime/MeshModelingToolset/Source/MeshModelingTools/Public/Sculpting/MeshVertexPaintBrushOps.h
+++ b/Engine/Plugins/Runtime/MeshModelingToolset/Source/MeshModelingTools/Public/Sculpting/MeshVertexPaintBrushOps.h
@@ -273,9 +273,9 @@ public:
 		return true;
 	}
 
-	virtual bool SupportsStrokeType(EMeshSculptStrokeType StrokeBehavior) const override
+	virtual bool SupportsStrokeType(EMeshSculptStrokeType StrokeType) const override
 	{
-		switch (StrokeBehavior)
+		switch (StrokeType)
 		{
 		case EMeshSculptStrokeType::Airbrush:
 		case EMeshSculptStrokeType::Dots:
