ModelingTools: EnableVolumeLocking by default again until we have sorted out the mesh element selection <-> tool interactions for volumes.

diff --git a/Engine/Plugins/Runtime/MeshModelingToolset/Source/ModelingComponentsEditorOnly/Private/Selection/VolumeSelector.cpp b/Engine/Plugins/Runtime/MeshModelingToolset/Source/ModelingComponentsEditorOnly/Private/Selection/VolumeSelector.cpp
--- a/Engine/Plugins/Runtime/MeshModelingToolset/Source/ModelingComponentsEditorOnly/Private/Selection/VolumeSelector.cpp
+++ b/Engine/Plugins/Runtime/MeshModelingToolset/Source/ModelingComponentsEditorOnly/Private/Selection/VolumeSelector.cpp
@@ -26,7 +26,7 @@ using namespace UE::Geometry;
 
 static TAutoConsoleVariable<bool> CVarEnableModelingSelectionVolumeLocking(
 	TEXT("modeling.Selection.EnableVolumeLocking"),
-	false,
+	true,
 	TEXT("Control whether Selection Locking is enabled by default for Volumes"));
 
 
