ModelingMode: Improved performance on the VertexSculpt tool.

Also added support for Octree testing and visualization, with two CVARS - one for disabling octree updates during sculpting for testing purposes and another for visualizing which triangles belong to which Octree root cell.

diff --git a/Engine/Plugins/ChaosClothAssetEditor/Source/ChaosClothAssetEditorTools/Private/ChaosClothAsset/ClothWeightMapPaintTool.cpp b/Engine/Plugins/ChaosClothAssetEditor/Source/ChaosClothAssetEditorTools/Private/ChaosClothAsset/ClothWeightMapPaintTool.cpp
--- a/Engine/Plugins/ChaosClothAssetEditor/Source/ChaosClothAssetEditorTools/Private/ChaosClothAsset/ClothWeightMapPaintTool.cpp
+++ b/Engine/Plugins/ChaosClothAssetEditor/Source/ChaosClothAssetEditorTools/Private/ChaosClothAsset/ClothWeightMapPaintTool.cpp
@@ -1953,7 +1953,7 @@ void UClothEditorWeightMapPaintTool::OnTick(float DeltaTime)
 		bPendingPickWeight = false;
 	}
 
-	auto ExecuteStampOperation = [this](const FRay& StampRay)
+	auto ExecuteStampOperation = [this](int StampIndex, const FRay& StampRay)
 		{
 			SCOPE_CYCLE_COUNTER(WeightMapPaintTool_Tick_ApplyStampBlock);
 
diff --git a/Engine/Plugins/Experimental/Dataflow/Source/DataflowEditor/Private/DataflowEditorTools/DataflowEditorWeightMapPaintTool.cpp b/Engine/Plugins/Experimental/Dataflow/Source/DataflowEditor/Private/DataflowEditorTools/DataflowEditorWeightMapPaintTool.cpp
--- a/Engine/Plugins/Experimental/Dataflow/Source/DataflowEditor/Private/DataflowEditorTools/DataflowEditorWeightMapPaintTool.cpp
+++ b/Engine/Plugins/Experimental/Dataflow/Source/DataflowEditor/Private/DataflowEditorTools/DataflowEditorWeightMapPaintTool.cpp
@@ -1564,7 +1564,7 @@ void UDataflowEditorWeightMapPaintTool::OnTick(float DeltaTime)
 		bPendingPickWeight = false;
 	}
 
-	auto ExecuteStampOperation = [this](const FRay& StampRay)
+	auto ExecuteStampOperation = [this](int StampIndex, const FRay& StampRay)
 		{
 			SCOPE_CYCLE_COUNTER(WeightMapPaintTool_Tick_ApplyStampBlock);
 
diff --git a/Engine/Plugins/Runtime/MeshModelingToolset/Source/MeshModelingTools/Private/MeshGroupPaintTool.cpp b/Engine/Plugins/Runtime/MeshModelingToolset/Source/MeshModelingTools/Private/MeshGroupPaintTool.cpp
--- a/Engine/Plugins/Runtime/MeshModelingToolset/Source/MeshModelingTools/Private/MeshGroupPaintTool.cpp
+++ b/Engine/Plugins/Runtime/MeshModelingToolset/Source/MeshModelingTools/Private/MeshGroupPaintTool.cpp
@@ -1045,7 +1045,7 @@ void UMeshGroupPaintTool::OnTick(float DeltaTime)
 		bPendingPickGroup = bPendingToggleFreezeGroup = false;
 	}
 
-	auto ExecuteStampOperation = [this](const FRay& StampRay)
+	auto ExecuteStampOperation = [this](int StampIndex, const FRay& StampRay)
 		{
 			SCOPE_CYCLE_COUNTER(GroupPaintTool_Tick_ApplyStampBlock);
 
diff --git a/Engine/Plugins/Runtime/MeshModelingToolset/Source/MeshModelingTools/Private/MeshVertexPaintTool.cpp b/Engine/Plugins/Runtime/MeshModelingToolset/Source/MeshModelingTools/Private/MeshVertexPaintTool.cpp
--- a/Engine/Plugins/Runtime/MeshModelingToolset/Source/MeshModelingTools/Private/MeshVertexPaintTool.cpp
+++ b/Engine/Plugins/Runtime/MeshModelingToolset/Source/MeshModelingTools/Private/MeshVertexPaintTool.cpp
@@ -1367,7 +1367,7 @@ void UMeshVertexPaintTool::OnTick(float DeltaTime)
 		bPendingPickColor = bPendingPickEraseColor = false;
 	}
 
-	auto ExecuteStampOperation = [this](const FRay& StampRay)
+	auto ExecuteStampOperation = [this](int StampIndex, const FRay& StampRay)
 		{
 			TRACE_CPUPROFILER_EVENT_SCOPE(MeshVPaint_OnTick_InStroke);
 
diff --git a/Engine/Plugins/Runtime/MeshModelingToolset/Source/MeshModelingTools/Private/MeshVertexSculptTool.cpp b/Engine/Plugins/Runtime/MeshModelingToolset/Source/MeshModelingTools/Private/MeshVertexSculptTool.cpp
--- a/Engine/Plugins/Runtime/MeshModelingToolset/Source/MeshModelingTools/Private/MeshVertexSculptTool.cpp
+++ b/Engine/Plugins/Runtime/MeshModelingToolset/Source/MeshModelingTools/Private/MeshVertexSculptTool.cpp
@@ -10,14 +10,17 @@
 #include "Async/ParallelFor.h"
 #include "Async/Async.h"
 #include "Selections/MeshConnectedComponents.h"
+#include "Algo/Unique.h"
 
 #include "MeshWeights.h"
 #include "DynamicMesh/MeshNormals.h"
 #include "DynamicMesh/MeshIndexUtil.h"
 #include "Parameterization/MeshPlanarSymmetry.h"
 #include "Util/BufferUtil.h"
+#include "Util/UniqueIndexSet.h"
 #include "AssetUtils/Texture2DUtil.h"
 #include "ToolSetupUtil.h"
+#include "Drawing/PreviewGeometryActor.h"
 
 #include "BaseGizmos/BrushStampIndicator.h"
 #include "PreviewMesh.h"
@@ -62,6 +65,16 @@ namespace
 #endif
 }
 
+namespace MeshVertexSculptToolLocals
+{
+	const FString& OctreePointSetID(TEXT("OctreePointSet"));
+	const FString& OctreeLineSetID(TEXT("OctreeLineSet"));
+
+	static FAutoConsoleVariable EnableOctreeVisuals(TEXT("modeling.Sculpting.EnableOctreeVisuals"), false, TEXT("Enable visualizing the octree used for determining ROIs."));
+	static FAutoConsoleVariable DisableOctreeUpdates(TEXT("modeling.Sculpting.DisableOctreeUpdates"), false, TEXT("Disable updating the octree during sculpting"));
+	static FAutoConsoleVariable OctreeRootCellSizeOverride(TEXT("modeling.Sculpting.OctreeRootCellSizeOverride"), 0.0f, TEXT("If greater than 0.0, set octree root cell size to this value instead of auto-computing."));
+	static FAutoConsoleVariable OctreeTreeDepthOverride(TEXT("modeling.Sculpting.OctreeTreeDepthOverride"), 0, TEXT("If greater than 0, set octree tree depth to this value instead of auto-computing."));
+}
 
 /*
  * ToolBuilder
@@ -132,7 +145,9 @@ void UMeshVertexSculptTool::Setup()
 	InitialBoundsMaxDim = Bounds.MaxDim();
 
 	// initialize dynamic octree
-	TFuture<void> InitializeOctree = Async(VertexSculptToolAsyncExecTarget, [SculptMesh, Bounds, this]()
+	float RootCellSizeOverride = FMath::Abs(MeshVertexSculptToolLocals::OctreeRootCellSizeOverride->GetFloat());	
+	int TreeDepthOverride = FMath::Abs(MeshVertexSculptToolLocals::OctreeTreeDepthOverride->GetInt());
+	TFuture<void> InitializeOctree = Async(VertexSculptToolAsyncExecTarget, [SculptMesh, Bounds, RootCellSizeOverride, TreeDepthOverride, this]()
 	{
 		if (SculptMesh->TriangleCount() > 100000)
 		{
@@ -144,6 +159,16 @@ void UMeshVertexSculptTool::Setup()
 			Octree.RootDimension = InitialBoundsMaxDim / 2.0;
 			Octree.SetMaxTreeDepth(8);
 		}
+
+		if (!FMath::IsNearlyZero(RootCellSizeOverride))
+		{
+			Octree.RootDimension = RootCellSizeOverride;
+		}
+		if (TreeDepthOverride > 0)
+		{
+			Octree.SetMaxTreeDepth(TreeDepthOverride);
+		}
+
 		Octree.Initialize(SculptMesh);
 		//Octree.CheckValidity(EValidityCheckFailMode::Check, true, true);
 		//FDynamicMeshOctree3::FStatistics Stats;
@@ -407,6 +432,12 @@ FString UMeshVertexSculptTool::GetPropertyCacheIdentifier() const
 
 void UMeshVertexSculptTool::Shutdown(EToolShutdownType ShutdownType)
 {
+	if (OctreeGeometry)
+	{
+		OctreeGeometry->Disconnect();
+		OctreeGeometry = nullptr;
+	}
+
 	if (DynamicMeshComponent != nullptr)
 	{
 		DynamicMeshComponent->OnMeshChanged.Remove(OnDynamicMeshComponentChangedHandle);
@@ -1375,6 +1406,16 @@ void UMeshVertexSculptTool::DrawHUD(FCanvas* Canvas, IToolsContextRenderAPI* Ren
 	}
 }
 
+namespace VertexSculptAsyncHelpers
+{
+	template<typename CallableType, typename TGateFutureType>
+	auto ChainedAsync(TFuture<TGateFutureType>&& GateFuture,  EAsyncExecution Execution, CallableType&& Callable) -> TFuture<decltype(Forward<CallableType>(Callable)())>
+	{
+		GateFuture.Wait();
+		return Async(Execution, Callable);
+	}
+
+}
 
 void UMeshVertexSculptTool::OnTick(float DeltaTime)
 {
@@ -1404,103 +1445,127 @@ void UMeshVertexSculptTool::OnTick(float DeltaTime)
 		bTargetDirty = false;
 	}
 
-
 	FDynamicMesh3* Mesh = GetSculptMesh();
-	auto ExecuteStampOperation = [this, Mesh](const FRay& StampRay)
-	{
-		TRACE_CPUPROFILER_EVENT_SCOPE(VtxSculptTool_Tick_StrokeUpdate);
+	const FDynamicMeshNormalOverlay* Normals = Mesh->HasAttributes() ? Mesh->Attributes()->PrimaryNormals() : nullptr;
+	const bool bUsingOverlayNormalsOut = Normals != nullptr;
+	TFuture<void> AccumulateROI;
+	TFuture<void> NormalsROI;
+	TArray<TArray<int>> TriangleROIArrays;
+	FUniqueIndexSet TriangleROISet;
 
+	auto PreExecuteStampOperation = [this, Mesh, &TriangleROISet](int StampCount)
+	{
+		TRACE_CPUPROFILER_EVENT_SCOPE(VtxSculptTool_Tick_PreStrokeUpdate);
 		// need to make sure previous stamp finished
 		WaitForPendingStampUpdateConst();
 
-		// update sculpt ROI
-		UpdateROI(CurrentStamp.LocalFrame);
+		TriangleROISet.Initialize(Mesh->TriangleCount(), (Mesh->TriangleCount() * 0.1) );
+	};
 
-		// Append updated ROI to modified region (async). For some reason this is very expensive,
-		// maybe because of TSet? but we have a lot of time to do it.
-		TFuture<void> AccumulateROI = Async(VertexSculptToolAsyncExecTarget, [this]()
-			{
-				TRACE_CPUPROFILER_EVENT_SCOPE(VtxSculptTool_Tick_AccumROI);
-				for (int32 tid : TriangleROIArray)
+
+	auto ExecuteStampOperation = [this, Mesh, bUsingOverlayNormalsOut, &AccumulateROI, &NormalsROI, &TriangleROISet](int StampIndex, const FRay& StampRay)
+	{
+		TRACE_CPUPROFILER_EVENT_SCOPE(VtxSculptTool_Tick_StrokeUpdate);
+		{
+			TRACE_CPUPROFILER_EVENT_SCOPE(VtxSculptTool_Tick_UpdateROI);
+			// update sculpt ROI
+			UpdateROI(CurrentStamp.LocalFrame);
+		}
+
+		// Instead of figuring out which triangles are unique here, we're going to simply accumulate the stamp ROIs in place
+		{
+			TRACE_CPUPROFILER_EVENT_SCOPE(VtxSculptTool_Tick_AccumROIWrapper);
+			AccumulateROI = VertexSculptAsyncHelpers::ChainedAsync(MoveTemp(AccumulateROI), VertexSculptToolAsyncExecTarget, [this, StampIndex, &TriangleROISet, TriangleROIArrayCopy = TriangleROIArray]()
 				{
-					AccumulatedTriangleROI.Add(tid);
-				}
-			});
+					TRACE_CPUPROFILER_EVENT_SCOPE(VtxSculptTool_Tick_AccumROI);
+					for (const int& Tid : TriangleROIArrayCopy)
+					{
+						TriangleROISet.Add(Tid);
+					}
+				});
+		}
 
-		// Start precomputing the normals ROI. This is currently the most expensive single thing we do next
-		// to Octree re-insertion, despite it being almost trivial. Why?!?
-		bool bUsingOverlayNormalsOut = false;
-		TFuture<void> NormalsROI = Async(VertexSculptToolAsyncExecTarget, [Mesh, &bUsingOverlayNormalsOut, this]()
-			{
-				TRACE_CPUPROFILER_EVENT_SCOPE(VtxSculptTool_Tick_NormalsROI);
+		// need to make sure previous stamp finished
+		WaitForPendingStampUpdateConst();
 
-				//UE::SculptUtil::PrecalculateNormalsROI(Mesh, TriangleROIArray, NormalsROIBuilder, bUsingOverlayNormals, false);
-				UE::SculptUtil::PrecalculateNormalsROI(Mesh, TriangleROIArray, NormalsFlags, bUsingOverlayNormalsOut, false);
-			});
+		// Nathan - Hypothetically we can apply more than one stamp at a time, but this would require
+		// understanding which stamps do and don't overlap. Any non-overlapping stamps technically
+		// could be applied in parallel.
+		TFuture<void> UpdateChangeFuture;
+		{
+			TRACE_CPUPROFILER_EVENT_SCOPE(VtxSculptTool_Tick_ApplyStamp);
+			UpdateChangeFuture = ApplyStamp();
+		}
+
+		{
+			TRACE_CPUPROFILER_EVENT_SCOPE(VtxSculptTool_Tick_ApplyStampWait);
+			UpdateChangeFuture.Wait();
+		}
+	};
+
+	auto PostExecuteStampOperation = [this, Mesh, bUsingOverlayNormalsOut, &AccumulateROI, &NormalsROI,  &TriangleROISet]()
+	{
+		TRACE_CPUPROFILER_EVENT_SCOPE(VtxSculptTool_Tick_PostStrokeUpdate);
 
 		// NOTE: you might try to speculatively do the octree remove here, to save doing it later on Reinsert().
 		// This will not improve things, as Reinsert() checks if it needs to actually re-insert, which avoids many
 		// removes, and does much of the work of Remove anyway.
-
-		// Apply the stamp. This will return a future that is updating the vertex-change record, 
-		// which can run until the end of the frame, as it is using cached information
-		TFuture<void> UpdateChangeFuture = ApplyStamp();
-
+		// 
 		// begin octree rebuild calculation
-		StampUpdateOctreeFuture = Async(VertexSculptToolAsyncExecTarget, [this]()
-			{
-				TRACE_CPUPROFILER_EVENT_SCOPE(VtxSculptTool_Tick_OctreeReinsert);
-				Octree.ReinsertTrianglesParallel(TriangleROIArray, OctreeUpdateTempBuffer, OctreeUpdateTempFlagBuffer);
-			});
-		bStampUpdatePending = true;
-		StampUpdateOctreeFuture.Then([this](TFuture<void> Future)
-			{
-				bStampUpdatePending = false;
-			});
-
-		//TFuture<void> OctreeRebuild = Async(VertexSculptToolAsyncExecTarget, [&]()
-		//{
-		//	TRACE_CPUPROFILER_EVENT_SCOPE(VtxSculptTool_Tick_OctreeReinsert);
-		//	Octree.ReinsertTriangles(TriangleROIArray);
-		//});
+		if (!MeshVertexSculptToolLocals::DisableOctreeUpdates->GetBool())
+		{
+			StampUpdateOctreeFuture = Async(VertexSculptToolAsyncExecTarget, [this]()
+				{
+					TRACE_CPUPROFILER_EVENT_SCOPE(VtxSculptTool_Tick_OctreeReinsert);
+					Octree.ReinsertTrianglesParallel(TriangleROIArray, OctreeUpdateTempBuffer, OctreeUpdateTempFlagBuffer);
+				});
+			bStampUpdatePending = true;
+			StampUpdateOctreeFuture.Then([this](TFuture<void> Future)
+				{
+					bStampUpdatePending = false;
+					bOctreeUpdated = true;
+				});
+		}
 
-		// TODO: first step of RecalculateROINormals() is to convert TriangleROI into vertex or element ROI.
-		// We can do this while we are computing stamp!
+		// Prepare list of triangles to process after all stamps are applied
+		AccumulateROI.Wait();
+		TArray<int> AccumulatedTriangleROIArray;
+		TriangleROISet.SwapValuesWith(AccumulatedTriangleROIArray);
 
 		// precompute dynamic mesh update info
 		TArray<int32> RenderUpdateSets; FAxisAlignedBox3d RenderUpdateBounds;
-		TFuture<bool> RenderUpdatePrecompute = DynamicMeshComponent->FastNotifyTriangleVerticesUpdated_TryPrecompute(
-			TriangleROIArray, RenderUpdateSets, RenderUpdateBounds);
+		TFuture<bool> RenderUpdatePrecompute;
+		{
+			TRACE_CPUPROFILER_EVENT_SCOPE(VtxSculptTool_Tick_PrecomputeUpdateMesh);
+			RenderUpdatePrecompute = DynamicMeshComponent->FastNotifyTriangleVerticesUpdated_TryPrecompute(
+				AccumulatedTriangleROIArray, RenderUpdateSets, RenderUpdateBounds);
+		}
 
 		// recalculate normals. This has to complete before we can update component
 		// (in fact we could do it per-chunk...)
 		{
 			TRACE_CPUPROFILER_EVENT_SCOPE(VtxSculptTool_Tick_RecalcNormals);
 			NormalsROI.Wait();
-			UE::SculptUtil::RecalculateROINormals(Mesh, NormalsFlags, bUsingOverlayNormalsOut);
-			//UE::SculptUtil::RecalculateROINormals(Mesh, NormalsROIBuilder.Indices(), bUsingOverlayNormals);
+			UE::SculptUtil::RecalculateROINormalForTriangles(Mesh, AccumulatedTriangleROIArray, bUsingOverlayNormalsOut);
 		}
 
 		{
-			TRACE_CPUPROFILER_EVENT_SCOPE(VtxSculptTool_Tick_UpdateMesh);
+			TRACE_CPUPROFILER_EVENT_SCOPE(VtxSculptTool_Tick_UpdateMesh);	
 			RenderUpdatePrecompute.Wait();
-			DynamicMeshComponent->FastNotifyTriangleVerticesUpdated_ApplyPrecompute(TriangleROIArray,
+			DynamicMeshComponent->FastNotifyTriangleVerticesUpdated_ApplyPrecompute(AccumulatedTriangleROIArray,
 				EMeshRenderAttributeFlags::Positions | EMeshRenderAttributeFlags::VertexNormals,
 				RenderUpdatePrecompute, RenderUpdateSets, RenderUpdateBounds);
 
 			GetToolManager()->PostInvalidation();
 		}
 
-		// we don't really need to wait for these to happen to end Tick()...
-		UpdateChangeFuture.Wait();
-		AccumulateROI.Wait();
-	};
 
+	};
 	
 	ProcessPerTickStamps(
 		[this](const FRay& StampRay) -> bool {
 			return UpdateStampPosition(StampRay);
-		}, ExecuteStampOperation);
+		}, PreExecuteStampOperation, ExecuteStampOperation, PostExecuteStampOperation);
 
 
 	if (InStroke() == false && bTargetDirty)
@@ -1515,6 +1580,64 @@ void UMeshVertexSculptTool::OnTick(float DeltaTime)
 		bTargetDirty = false;
 	}
 
+
+	if (MeshVertexSculptToolLocals::EnableOctreeVisuals->GetBool())
+	{
+		if (!OctreeGeometry)
+		{
+			// Set up all the components we need to visualize things.
+			OctreeGeometry = NewObject<UPreviewGeometry>();
+			OctreeGeometry->CreateInWorld(TargetWorld, FTransform());
+
+			// These visualize the current spline edges that would be extracted
+			OctreeGeometry->AddPointSet(MeshVertexSculptToolLocals::OctreePointSetID);
+			OctreeGeometry->AddLineSet(MeshVertexSculptToolLocals::OctreeLineSetID);
+		}
+
+		ULineSetComponent* OctreeLineSet = OctreeGeometry->FindLineSet(MeshVertexSculptToolLocals::OctreeLineSetID);
+
+		if (bOctreeUpdated)
+		{
+			OctreeLineSet->Clear();
+
+			if (!CutTree)
+			{
+				CutTree = MakeShared<FDynamicMeshOctree3::FTreeCutSet>(Octree.BuildLevelCutSet(0));
+			}
+			else
+			{
+				TArray<FDynamicMeshOctree3::FCellReference> NewCells;
+				Octree.UpdateLevelCutSet(*CutTree, NewCells);
+			}
+
+			int ColorSeed = 0;
+			for (const FDynamicMeshOctree3::FCellReference& CellRef : CutTree->CutCells)
+			{
+				FColor CutColor = FColor::MakeRandomSeededColor(ColorSeed++);
+				Octree.CollectTriangles(CellRef, [this, OctreeLineSet, Mesh, &CutColor](int32 Tid) {
+					FVector A, B, C;
+					Mesh->GetTriVertices(Tid, A, B, C);
+					OctreeLineSet->AddLine(A, B, CutColor, 2.0);
+					OctreeLineSet->AddLine(B, C, CutColor, 2.0);
+					OctreeLineSet->AddLine(C, A, CutColor, 2.0);
+					});
+			}
+
+			bOctreeUpdated = false;
+		}
+	}
+	else
+	{
+		if (OctreeGeometry)
+		{
+			ULineSetComponent* OctreeLineSet = OctreeGeometry->FindLineSet(MeshVertexSculptToolLocals::OctreeLineSetID);
+			if (bOctreeUpdated)
+			{
+				OctreeLineSet->Clear();
+			}
+		}
+	}
+
 }
 
 
@@ -1757,6 +1880,7 @@ void UMeshVertexSculptTool::OnDynamicMeshComponentChanged(UDynamicMeshComponent*
 	UndoUpdateOctreeFuture = Async(VertexSculptToolAsyncExecTarget, [this, Mesh]()
 	{
 		Octree.ReinsertTriangles(AccumulatedTriangleROI);
+		bOctreeUpdated = true;
 		return true;
 	});
 
diff --git a/Engine/Plugins/Runtime/MeshModelingToolset/Source/MeshModelingTools/Private/Sculpting/MeshSculptToolBase.cpp b/Engine/Plugins/Runtime/MeshModelingToolset/Source/MeshModelingTools/Private/Sculpting/MeshSculptToolBase.cpp
--- a/Engine/Plugins/Runtime/MeshModelingToolset/Source/MeshModelingTools/Private/Sculpting/MeshSculptToolBase.cpp
+++ b/Engine/Plugins/Runtime/MeshModelingToolset/Source/MeshModelingTools/Private/Sculpting/MeshSculptToolBase.cpp
@@ -39,7 +39,7 @@ namespace MeshSculptBaseLocals
 	const FString& StrokePointSetID(TEXT("StrokePointSet"));
 	const FString& StrokeLineSetID(TEXT("StrokeLineSet"));
 
-	static FAutoConsoleVariable EnableMouseTrail(TEXT("modeling.Sculpting.EnableMouseTrail"), 0.0f, TEXT("Enable visualizing the path the mouse takes while sculpting."));
+	static FAutoConsoleVariable EnableMouseTrail(TEXT("modeling.Sculpting.EnableMouseTrail"), false, TEXT("Enable visualizing the path the mouse takes while sculpting."));
 }
 
 void FBrushToolRadius::InitializeWorldSizeRange(TInterval<float> Range, bool bValidateWorldRadius)
@@ -164,14 +164,6 @@ void UMeshSculptToolBase::Setup()
 		[this](bool bOn) { UpdateTwoSidedSetting(bOn); });
 	ViewProperties->WatchProperty(ViewProperties->Image,
 		[this](UTexture2D* NewImage) { UpdateImageSetting(NewImage); });
-	
-	// Set up all the components we need to visualize things.
-	StrokeGeometry = NewObject<UPreviewGeometry>();
-	StrokeGeometry->CreateInWorld(TargetWorld, FTransform());
-
-	// These visualize the current spline edges that would be extracted
-	StrokeGeometry->AddPointSet(MeshSculptBaseLocals::StrokePointSetID);
-	StrokeGeometry->AddLineSet(MeshSculptBaseLocals::StrokeLineSetID);
 
 	// add input behavior to click-drag while holding hotkey to adjust brush size and strength
     {
@@ -379,6 +371,17 @@ void UMeshSculptToolBase::OnTick(float DeltaTime)
 			}
 		}
 
+		if (!StrokeGeometry)
+		{
+			// Set up all the components we need to visualize things.
+			StrokeGeometry = NewObject<UPreviewGeometry>();
+			StrokeGeometry->CreateInWorld(TargetWorld, FTransform());
+
+			// These visualize the current spline edges that would be extracted
+			StrokeGeometry->AddPointSet(MeshSculptBaseLocals::StrokePointSetID);
+			StrokeGeometry->AddLineSet(MeshSculptBaseLocals::StrokeLineSetID);
+		}
+
 		ULineSetComponent* LineSet = StrokeGeometry->FindLineSet(MeshSculptBaseLocals::StrokeLineSetID);
 		if (LineSet)
 		{
@@ -429,6 +432,22 @@ void UMeshSculptToolBase::OnTick(float DeltaTime)
 			}
 		}
 	}
+	else
+	{
+		if (StrokeGeometry)
+		{
+			ULineSetComponent* LineSet = StrokeGeometry->FindLineSet(MeshSculptBaseLocals::StrokeLineSetID);
+			if (LineSet)
+			{
+				LineSet->Clear();
+			}
+			UPointSetComponent* PointSet = StrokeGeometry->FindPointSet(MeshSculptBaseLocals::StrokePointSetID);
+			if (PointSet)
+			{
+				PointSet->Clear();
+			}
+		}
+	}
 
 
 	// always using HoverStamp here because it's position should always be up-to-date for the current apply-stamp...
@@ -1029,11 +1048,18 @@ float UMeshSculptToolBase::GetStampTemporalFlowRate() const
 	return BrushProperties->FlowRate;
 }
 
-void UMeshSculptToolBase::ProcessPerTickStamps(TFunction<bool(const FRay& StampRay)> UpdateStampPosition, TFunction<void(const FRay& StampRay)> ExecuteStampOperation)
+void UMeshSculptToolBase::ProcessPerTickStamps(TFunction<bool(const FRay& StampRay)> UpdateStampPosition, TFunction<void(int StampIndex, const FRay& StampRay)> ExecuteStampOperation)
+{
+	ProcessPerTickStamps(UpdateStampPosition, [](int StampCount) {}, ExecuteStampOperation, [] {});
+}
+
+void UMeshSculptToolBase::ProcessPerTickStamps(TFunction<bool(const FRay& StampRay)> UpdateStampPosition, TFunction<void(int StampCount)> PreExecuteStampsOperation, TFunction<void(int StampIndex, const FRay& StampRay)> ExecuteStampOperation, TFunction<void()> PostExecuteStampsOperation)
 {
 	if (InStroke())
 	{
 		TRACE_CPUPROFILER_EVENT_SCOPE(BaseSculptTool_Tick_ProcessPerTickStamps);
+		FDynamicMesh3* Mesh = GetSculptMesh();
+
 		// We save the active pressure here, so we can then override and fake it during the stamp application to generate interpolated pressures along the stroke.
 		TGuardValue<double> CachedActivePressure(ActivePressure, ActivePressure);
 
@@ -1051,20 +1077,25 @@ void UMeshSculptToolBase::ProcessPerTickStamps(TFunction<bool(const FRay& StampR
 
 		if (!PendingStampRays.IsEmpty())
 		{
-			for (const FStampRayData& StampRay : PendingStampRays)
+			PreExecuteStampsOperation(PendingStampRays.Num());
+			for (int Index = 0; Index < PendingStampRays.Num(); ++Index)
 			{
+				const FStampRayData& StampRay = PendingStampRays[Index];
+
 				ActivePressure = StampRay.Pressure;
 
 				UpdateStampPosition(StampRay.StampRay);
-				ExecuteStampOperation(StampRay.StampRay);
+				ExecuteStampOperation(Index, StampRay.StampRay);
 				StampsDuringStroke++;
 			}
+			PostExecuteStampsOperation();
 		}
 
 	}
 }
 
 
+
 void UMeshSculptToolBase::UpdateStampPendingState()
 {
 	if (InStroke() == false) return;
diff --git a/Engine/Plugins/Runtime/MeshModelingToolset/Source/MeshModelingTools/Private/Sculpting/MeshSculptUtil.cpp b/Engine/Plugins/Runtime/MeshModelingToolset/Source/MeshModelingTools/Private/Sculpting/MeshSculptUtil.cpp
--- a/Engine/Plugins/Runtime/MeshModelingToolset/Source/MeshModelingTools/Private/Sculpting/MeshSculptUtil.cpp
+++ b/Engine/Plugins/Runtime/MeshModelingToolset/Source/MeshModelingTools/Private/Sculpting/MeshSculptUtil.cpp
@@ -430,3 +430,14 @@ void UE::SculptUtil::RecalculateROINormals(FDynamicMesh3* Mesh, TArray<std::atom
 }
 
 
+void UE::SculptUtil::RecalculateROINormalForTriangles(FDynamicMesh3* Mesh, TArray<int> Triangles, bool bIsOverlayElements)
+{
+	if (bIsOverlayElements)
+	{
+		FMeshNormals::RecomputeOverlayTriNormals(*Mesh, Triangles);
+	}
+	else
+	{
+		FMeshNormals::QuickComputeVertexNormalsForTriangles(*Mesh, Triangles);
+	}
+}
diff --git a/Engine/Plugins/Runtime/MeshModelingToolset/Source/MeshModelingTools/Public/MeshVertexSculptTool.h b/Engine/Plugins/Runtime/MeshModelingToolset/Source/MeshModelingTools/Public/MeshVertexSculptTool.h
--- a/Engine/Plugins/Runtime/MeshModelingToolset/Source/MeshModelingTools/Public/MeshVertexSculptTool.h
+++ b/Engine/Plugins/Runtime/MeshModelingToolset/Source/MeshModelingTools/Public/MeshVertexSculptTool.h
@@ -428,6 +428,10 @@ protected:
 	TFunction<bool(int32, const FVector3d&, double MaxDist, FVector3d&, FVector3d&)> BaseMeshQueryFunc;
 
 	UE::Geometry::FDynamicMeshOctree3 Octree;
+	bool bOctreeUpdated = false;
+	TSharedPtr<UE::Geometry::FDynamicMeshOctree3::FTreeCutSet> CutTree;
+	UPROPERTY()
+	TObjectPtr<UPreviewGeometry> OctreeGeometry;
 
 	UE_API bool UpdateBrushPosition(const FRay& WorldRay);
 
diff --git a/Engine/Plugins/Runtime/MeshModelingToolset/Source/MeshModelingTools/Public/Sculpting/MeshSculptToolBase.h b/Engine/Plugins/Runtime/MeshModelingToolset/Source/MeshModelingTools/Public/Sculpting/MeshSculptToolBase.h
--- a/Engine/Plugins/Runtime/MeshModelingToolset/Source/MeshModelingTools/Public/Sculpting/MeshSculptToolBase.h
+++ b/Engine/Plugins/Runtime/MeshModelingToolset/Source/MeshModelingTools/Public/Sculpting/MeshSculptToolBase.h
@@ -546,9 +546,9 @@ protected:
 	double ActiveStrokePathArcLen = 0.0;
 	int LastFlowTimeStamp = 0;
 	int LastSpacingTimestamp = 0;
+	int StampsDuringStroke = 0;
 	UE_API virtual void ResetStrokeTime();
 	UE_API virtual void AccumulateStrokeTime(float DeltaTime);
-	int StampsDuringStroke = 0;
 
 	//
 	// Stamps
@@ -573,7 +573,8 @@ protected:
 	// Temporal Flow Rate defines the frequency of stamp placement. 1 is max rate, 0 is no stamps. Defaults to BrushProperties->FlowRate, but subclasses can re-use that setting for other things
 	UE_API virtual float GetStampTemporalFlowRate() const;
 
-	UE_API void ProcessPerTickStamps(TFunction<bool(const FRay& StampRay)> UpdateStampPosition, TFunction<void(const FRay& StampRay)> ExecuteStampOperation);
+	UE_API void ProcessPerTickStamps(TFunction<bool(const FRay& StampRay)> UpdateStampPosition, TFunction<void(int StampIndex, const FRay& StampRay)> ExecuteStampOperation);
+	UE_API void ProcessPerTickStamps(TFunction<bool(const FRay& StampRay)> UpdateStampPosition, TFunction<void(int StampCount )> PreExecuteStampsOperation, TFunction<void(int StampIndex, const FRay& StampRay)> ExecuteStampOperation, TFunction<void()> PostExecuteStampsOperation );
 
 	//
 	// Stamp ROI Plane is a plane used by some brush ops
diff --git a/Engine/Plugins/Runtime/MeshModelingToolset/Source/MeshModelingTools/Public/Sculpting/MeshSculptUtil.h b/Engine/Plugins/Runtime/MeshModelingToolset/Source/MeshModelingTools/Public/Sculpting/MeshSculptUtil.h
--- a/Engine/Plugins/Runtime/MeshModelingToolset/Source/MeshModelingTools/Public/Sculpting/MeshSculptUtil.h
+++ b/Engine/Plugins/Runtime/MeshModelingToolset/Source/MeshModelingTools/Public/Sculpting/MeshSculptUtil.h
@@ -84,6 +84,10 @@ namespace SculptUtil
 	 */
 	void RecalculateROINormals(FDynamicMesh3* Mesh, TArray<std::atomic<bool>>& ROIFlags, bool bIsOverlayElements);
 
-
+	/**
+	 * Recalculate either overlay normals or vertex normals for the given indices defined by a set of triangles
+	 * @param bIsOverlayElements if true, normals are computed on the overlay. Otherwise, they are computed as vertex normals
+	 */
+	void RecalculateROINormalForTriangles(FDynamicMesh3* Mesh, TArray<int> Triangles, bool bIsOverlayElements);
 
 /* end namespace UE::SculptUtil */  } }
\ No newline at end of file
