make the mesh attrib paint tool also support targets w/out mesh descriptions

diff --git a/Engine/Plugins/Runtime/MeshModelingToolset/Source/MeshModelingTools/Private/MeshAttributePaintTool.cpp b/Engine/Plugins/Runtime/MeshModelingToolset/Source/MeshModelingTools/Private/MeshAttributePaintTool.cpp
--- a/Engine/Plugins/Runtime/MeshModelingToolset/Source/MeshModelingTools/Private/MeshAttributePaintTool.cpp
+++ b/Engine/Plugins/Runtime/MeshModelingToolset/Source/MeshModelingTools/Private/MeshAttributePaintTool.cpp
@@ -11,10 +11,13 @@
 
 #include "MeshDescription.h"
 
+#include "TargetInterfaces/DynamicMeshProvider.h"
+#include "TargetInterfaces/DynamicMeshCommitter.h"
 #include "TargetInterfaces/MeshDescriptionProvider.h"
 #include "TargetInterfaces/MeshDescriptionCommitter.h"
-#include "TargetInterfaces/PrimitiveComponentBackedTarget.h"
+#include "TargetInterfaces/SceneComponentBackedTarget.h"
 #include "ModelingToolTargetUtil.h"
+#include "DynamicMesh/DynamicMeshAttributeSet.h"
 #include "DynamicMesh/NonManifoldMappingSupport.h"
 #include "ToolTargetManager.h"
 
@@ -57,8 +60,90 @@ public:
 	}
 };
 
+class FDynamicMeshVertexAttributeAdapter : public IMeshVertexAttributeAdapter
+{
+public:
+	FDynamicMesh3* Mesh;
+	FDynamicMeshWeightAttribute* WeightAttribute;
+
+	FDynamicMeshVertexAttributeAdapter(FDynamicMesh3* InMesh, FDynamicMeshWeightAttribute* InWeightAttribute)
+		: Mesh(InMesh), WeightAttribute(InWeightAttribute)
+	{
+	}
+
+	virtual int32 ElementNum() const override
+	{
+		return Mesh->MaxVertexID();
+	}
+
+	virtual float GetValue(int32 Index) const override
+	{
+		float Wt;
+		WeightAttribute->GetValue(Index, &Wt);
+		return Wt;
+	}
+
+	virtual void SetValue(int32 Index, float Value) override
+	{
+		WeightAttribute->SetScalarValue(Index, Value);
+	}
 
+	virtual FInterval1f GetValueRange() override
+	{
+		return FInterval1f(0.0f, 1.0f);
+	}
+};
 
+class FDynamicMeshVertexAttributeSource : public IMeshVertexAttributeSource
+{
+public:
+	FDynamicMesh3* Mesh = nullptr;
+
+	FDynamicMeshVertexAttributeSource(FDynamicMesh3* MeshIn)
+	{
+		Mesh = MeshIn;
+	}
+
+	virtual int32 GetAttributeElementNum() override
+	{
+		return Mesh->MaxVertexID();
+	}
+
+	virtual TArray<FName> GetAttributeList() override
+	{
+		TArray<FName> Result;
+		
+		if (FDynamicMeshAttributeSet* Attributes = Mesh->Attributes())
+		{
+			const int32 NumLayers = Attributes->NumWeightLayers();
+			Result.Reserve(NumLayers);
+			for (int32 LayerIdx = 0; LayerIdx < NumLayers; ++LayerIdx)
+			{
+				Result.Add(Attributes->GetWeightLayer(LayerIdx)->GetName());
+			}
+		}
+		return Result;
+	}
+
+
+	virtual TUniquePtr<IMeshVertexAttributeAdapter> GetAttribute(FName AttributeName) override
+	{
+		if (FDynamicMeshAttributeSet* Attributes = Mesh->Attributes())
+		{
+			const int32 NumLayers = Attributes->NumWeightLayers();
+			for (int32 LayerIdx = 0; LayerIdx < NumLayers; ++LayerIdx)
+			{
+				FDynamicMeshWeightAttribute* WeightLayer = Attributes->GetWeightLayer(LayerIdx);
+				if (WeightLayer->GetName() == AttributeName)
+				{
+					return MakeUnique<FDynamicMeshVertexAttributeAdapter>(Mesh, WeightLayer);
+				}
+			}
+		}
+		return nullptr;
+	}
+
+};
 
 
 class FMeshDescriptionVertexAttributeSource : public IMeshVertexAttributeSource
@@ -198,9 +283,9 @@ const FToolTargetTypeRequirements& UMeshAttributePaintToolBuilder::GetTargetRequ
 {
 	static FToolTargetTypeRequirements TypeRequirements({
 		UMaterialProvider::StaticClass(),
-		UMeshDescriptionProvider::StaticClass(),
-		UMeshDescriptionCommitter::StaticClass(),
-		UPrimitiveComponentBackedTarget::StaticClass()
+		UDynamicMeshProvider::StaticClass(),
+		UDynamicMeshCommitter::StaticClass(),
+		USceneComponentBackedTarget::StaticClass()
 		});
 	return TypeRequirements;
 }
@@ -297,10 +382,21 @@ void UMeshAttributePaintTool::Setup()
 
 	ColorMapper = MakeUnique<FFloatAttributeColorMapper>();
 
-	EditedMesh = MakeUnique<FMeshDescription>();
-	*EditedMesh = *UE::ToolTarget::GetMeshDescription(Target);
+	if (Cast<IMeshDescriptionProvider>(Target) && Cast<IMeshDescriptionCommitter>(Target))
+	{
+		EditedMesh = MakeUnique<FMeshDescription>();
+		*EditedMesh = *UE::ToolTarget::GetMeshDescription(Target);
+		AttributeSource = MakeUnique<FMeshDescriptionVertexAttributeSource>(EditedMesh.Get());
+	}
+	else
+	{
+		EditedDynamicMesh = MakeUnique<FDynamicMesh3>();
+		FGetMeshParameters Params;
+		Params.bWantMeshTangents = true;
+		*EditedDynamicMesh = UE::ToolTarget::GetDynamicMeshCopy(Target, Params);
+		AttributeSource = MakeUnique< FDynamicMeshVertexAttributeSource>(EditedDynamicMesh.Get());
+	}
 
-	AttributeSource = MakeUnique<FMeshDescriptionVertexAttributeSource>(EditedMesh.Get());
 	AttribProps->Initialize(AttributeSource->GetAttributeList(), true);
 
 	if (AttribProps->Attributes.Num() == 0)
@@ -460,7 +556,7 @@ bool UMeshAttributePaintTool::HitTest(const FRay& Ray, FHitResult& OutHit)
 
 void UMeshAttributePaintTool::CalculateVertexROI(const FBrushStampData& Stamp, TArray<int>& VertexROI)
 {
-	FTransform3d Transform(Cast<IPrimitiveComponentBackedTarget>(Target)->GetWorldTransform());
+	FTransform3d Transform(Cast<ISceneComponentBackedTarget>(Target)->GetWorldTransform());
 	FVector3d StampPosLocal = Transform.InverseTransformPosition((FVector3d)Stamp.WorldPosition);
 
 	float Radius = GetCurrentBrushRadiusLocal();
@@ -657,7 +753,7 @@ void UMeshAttributePaintTool::ApplyStamp(const FBrushStampData& Stamp)
 
 void UMeshAttributePaintTool::ApplyStamp_Paint(const FBrushStampData& Stamp, FStampActionData& ActionData)
 {
-	FTransform3d Transform(Cast<IPrimitiveComponentBackedTarget>(Target)->GetWorldTransform());
+	FTransform3d Transform(Cast<ISceneComponentBackedTarget>(Target)->GetWorldTransform());
 	FVector3d StampPosLocal = Transform.InverseTransformPosition((FVector3d)Stamp.WorldPosition);
 
 	int32 NumVertices = ActionData.ROIVertices.Num();
@@ -820,7 +916,14 @@ void UMeshAttributePaintTool::OnShutdown(EToolShutdownType ShutdownType)
 		// this block bakes the modified DynamicMeshComponent back into the StaticMeshComponent inside an undo transaction
 		GetToolManager()->BeginUndoTransaction(LOCTEXT("MeshAttributePaintTool", "Edit Attributes"));
 
-		UE::ToolTarget::CommitMeshDescriptionUpdate(Target, EditedMesh.Get());
+		if (EditedMesh.IsValid())
+		{
+			UE::ToolTarget::CommitMeshDescriptionUpdate(Target, EditedMesh.Get());
+		}
+		else
+		{
+			UE::ToolTarget::CommitDynamicMeshUpdate(Target, *EditedDynamicMesh, false);
+		}
 
 		GetToolManager()->EndUndoTransaction();
 	}
diff --git a/Engine/Plugins/Runtime/MeshModelingToolset/Source/MeshModelingTools/Public/MeshAttributePaintTool.h b/Engine/Plugins/Runtime/MeshModelingToolset/Source/MeshModelingTools/Public/MeshAttributePaintTool.h
--- a/Engine/Plugins/Runtime/MeshModelingToolset/Source/MeshModelingTools/Public/MeshAttributePaintTool.h
+++ b/Engine/Plugins/Runtime/MeshModelingToolset/Source/MeshModelingTools/Public/MeshAttributePaintTool.h
@@ -216,7 +216,7 @@ public:
 
 
 /**
- * UMeshAttributePaintTool paints single-channel float attributes on a MeshDescription.
+ * UMeshAttributePaintTool paints single-channel float attributes on a MeshDescription or DynamicMesh
  * 
  */
 UCLASS(MinimalAPI)
@@ -296,6 +296,7 @@ protected:
 	FBrushStampData LastStamp;
 	bool bStampPending;
 
+	// Note: This is used if the target supports mesh description read/write; otherwise EditedDynamicMesh is used
 	TUniquePtr<FMeshDescription> EditedMesh;
 
 	UE_API double CalculateBrushFalloff(double Distance);
@@ -342,6 +343,9 @@ private:
 	TSet<int32> SelectionTids;
 	TSet<int32> SelectionVids;
 	UE_API bool ShouldFilterTriangles() const;
+
+	// Note: This is used if the target doesn't support mesh description read/write; otherwise EditedMesh is used
+	TUniquePtr<FDynamicMesh3> EditedDynamicMesh;
 };
 
 
