Skin Weight Paint Tool: unbind the property set on shutdown to avoid useless notifications on reload

#rnx

diff --git a/Engine/Plugins/Runtime/MeshModelingToolset/Source/MeshModelingToolsEditorOnly/Private/SkeletalMesh/SkinWeightsPaintTool.cpp b/Engine/Plugins/Runtime/MeshModelingToolset/Source/MeshModelingToolsEditorOnly/Private/SkeletalMesh/SkinWeightsPaintTool.cpp
--- a/Engine/Plugins/Runtime/MeshModelingToolset/Source/MeshModelingToolsEditorOnly/Private/SkeletalMesh/SkinWeightsPaintTool.cpp
+++ b/Engine/Plugins/Runtime/MeshModelingToolset/Source/MeshModelingToolsEditorOnly/Private/SkeletalMesh/SkinWeightsPaintTool.cpp
@@ -3216,6 +3216,7 @@ void USkinWeightsPaintTool::OnShutdown(EToolShutdownType ShutdownType)
 	
 	// save tool properties
 	WeightToolProperties->SaveProperties(this);
+	RemoveToolPropertySource(WeightToolProperties);
 
 	// shutdown polygon selection mechanic
 	if (MeshSelector)
@@ -5068,44 +5069,47 @@ void USkinWeightsPaintTool::OnPropertyModified(UObject* ModifiedObject, FPropert
 {
 	Super::OnPropertyModified(ModifiedObject, ModifiedProperty);
 
-	if (ModifiedProperty->GetName() == GET_MEMBER_NAME_STRING_CHECKED(USkinWeightsPaintToolProperties, BrushStrength))
+	if (ModifiedProperty)
 	{
-		WeightToolProperties->GetBrushConfig().Strength = WeightToolProperties->BrushStrength;
-	}
-	if (ModifiedProperty->GetName() == GET_MEMBER_NAME_STRING_CHECKED(USkinWeightsPaintToolProperties, BrushRadius))
-	{
-		WeightToolProperties->GetBrushConfig().Radius = WeightToolProperties->BrushRadius;
-	}
-	if (ModifiedProperty->GetName() == GET_MEMBER_NAME_STRING_CHECKED(USkinWeightsPaintToolProperties, BrushFalloffAmount))
-	{
-		WeightToolProperties->GetBrushConfig().Falloff = WeightToolProperties->BrushFalloffAmount;
-	}
+		if (ModifiedProperty->GetName() == GET_MEMBER_NAME_STRING_CHECKED(USkinWeightsPaintToolProperties, BrushStrength))
+		{
+			WeightToolProperties->GetBrushConfig().Strength = WeightToolProperties->BrushStrength;
+		}
+		if (ModifiedProperty->GetName() == GET_MEMBER_NAME_STRING_CHECKED(USkinWeightsPaintToolProperties, BrushRadius))
+		{
+			WeightToolProperties->GetBrushConfig().Radius = WeightToolProperties->BrushRadius;
+		}
+		if (ModifiedProperty->GetName() == GET_MEMBER_NAME_STRING_CHECKED(USkinWeightsPaintToolProperties, BrushFalloffAmount))
+		{
+			WeightToolProperties->GetBrushConfig().Falloff = WeightToolProperties->BrushFalloffAmount;
+		}
 	
-	const FString NameOfModifiedProperty = ModifiedProperty->GetNameCPP();
+		const FString NameOfModifiedProperty = ModifiedProperty->GetNameCPP();
 
-	// invalidate vertex color cache when any weight color properties are modified
-	const TArray<FString> ColorPropertyNames = {
-		GET_MEMBER_NAME_STRING_CHECKED(USkinWeightsPaintToolProperties, ColorMode),
-		GET_MEMBER_NAME_STRING_CHECKED(USkinWeightsPaintToolProperties, ColorRamp),
-		GET_MEMBER_NAME_STRING_CHECKED(FLinearColor, R),
-		GET_MEMBER_NAME_STRING_CHECKED(FLinearColor, G),
-		GET_MEMBER_NAME_STRING_CHECKED(FLinearColor, B),
-		GET_MEMBER_NAME_STRING_CHECKED(FLinearColor, A)};
-	if (ColorPropertyNames.Contains(NameOfModifiedProperty))
-	{
-		bVertexColorsNeedUpdated = true;
-
-		// force all colors to have Alpha = 1
-		for (FLinearColor& Color : WeightToolProperties->ColorRamp)
+		// invalidate vertex color cache when any weight color properties are modified
+		const TArray<FString> ColorPropertyNames = {
+			GET_MEMBER_NAME_STRING_CHECKED(USkinWeightsPaintToolProperties, ColorMode),
+			GET_MEMBER_NAME_STRING_CHECKED(USkinWeightsPaintToolProperties, ColorRamp),
+			GET_MEMBER_NAME_STRING_CHECKED(FLinearColor, R),
+			GET_MEMBER_NAME_STRING_CHECKED(FLinearColor, G),
+			GET_MEMBER_NAME_STRING_CHECKED(FLinearColor, B),
+			GET_MEMBER_NAME_STRING_CHECKED(FLinearColor, A)};
+		if (ColorPropertyNames.Contains(NameOfModifiedProperty))
 		{
-			Color.A = 1.f;
+			bVertexColorsNeedUpdated = true;
+
+			// force all colors to have Alpha = 1
+			for (FLinearColor& Color : WeightToolProperties->ColorRamp)
+			{
+				Color.A = 1.f;
+			}
 		}
-	}
 
-	// let the mesh transfer system react to properties being set
-	TransferManager->OnPropertyModified(WeightToolProperties, ModifiedProperty);
+		// let the mesh transfer system react to properties being set
+		TransferManager->OnPropertyModified(WeightToolProperties, ModifiedProperty);
 	
-	SetFocusInViewport();
+		SetFocusInViewport();
+	}
 }
 
 #undef LOCTEXT_NAMESPACE
\ No newline at end of file
