Lattice Deformer Op: fix a check that fires when operating on non-compact meshes

diff --git a/Engine/Plugins/Runtime/MeshModelingToolset/Source/ModelingOperators/Private/DeformationOps/LatticeDeformerOp.cpp b/Engine/Plugins/Runtime/MeshModelingToolset/Source/ModelingOperators/Private/DeformationOps/LatticeDeformerOp.cpp
--- a/Engine/Plugins/Runtime/MeshModelingToolset/Source/ModelingOperators/Private/DeformationOps/LatticeDeformerOp.cpp
+++ b/Engine/Plugins/Runtime/MeshModelingToolset/Source/ModelingOperators/Private/DeformationOps/LatticeDeformerOp.cpp
@@ -31,11 +31,11 @@ void FLatticeDeformerOp::CalculateResult(FProgressCancel* Progress)
 	
 	if (bUsingSubmesh)
 	{
-		check(Submesh->GetSubmesh().VertexCount() == DeformedPositions.Num())
+		check(Submesh->GetSubmesh().MaxVertexID() == DeformedPositions.Num())
 	}
 	else
 	{
-		check(ResultMesh->VertexCount() == DeformedPositions.Num());
+		check(ResultMesh->MaxVertexID() == DeformedPositions.Num());
 	}
 
 	if (Progress && Progress->Cancelled())
