Skin Weight Paint Tool: fix crashes related to wrong ids being used when transfering weights

as InVertexSubset and Weights.PreChangeWeights/CurrentWeights reflect the current state of the mesh (partial or not) being edited, ApplyTranferredWeightsAsTransaction must map vertex indices correctly :

- the transfer algorithm now transfers all vertex weights, not just the selected ones, to ensure that non-manifold vertices are handled correctly
- the original vertex id is retrieved from the cleaned mesh to get the correct weight from the transferred skin weights attribute
- whean dealing with a isolated mesh, we make sure that the id is mapped form partial to clean to original to get the correct weight from the transferred skin weights attribute

#rnx

diff --git a/Engine/Plugins/Runtime/MeshModelingToolset/Source/MeshModelingToolsEditorOnly/Private/SkeletalMesh/SkinWeightsPaintTool.cpp b/Engine/Plugins/Runtime/MeshModelingToolset/Source/MeshModelingToolsEditorOnly/Private/SkeletalMesh/SkinWeightsPaintTool.cpp
--- a/Engine/Plugins/Runtime/MeshModelingToolset/Source/MeshModelingToolsEditorOnly/Private/SkeletalMesh/SkinWeightsPaintTool.cpp
+++ b/Engine/Plugins/Runtime/MeshModelingToolset/Source/MeshModelingToolsEditorOnly/Private/SkeletalMesh/SkinWeightsPaintTool.cpp
@@ -1022,11 +1022,6 @@ void UWeightToolTransferManager::TransferWeightsFromOtherMeshOrSubset()
 		return;
 	}
 
-	if (ToolProperties->EditingMode == EWeightEditMode::Mesh)
-	{
-		TransferBoneWeights.TargetVerticesSubset = MoveTemp(TargetSelectedVertices);
-	}
-
 	const FName TargetProfile = ToolProperties->GetActiveSkinWeightProfile();
 	if (!TransferBoneWeights.TransferWeightsToMesh(TargetMesh, TargetProfile))
 	{
@@ -1037,11 +1032,9 @@ void UWeightToolTransferManager::TransferWeightsFromOtherMeshOrSubset()
 
 	FDynamicMeshVertexSkinWeightsAttribute* TransferredSkinWeights = nullptr;
 	TransferredSkinWeights = TargetAttributes->GetSkinWeightsAttribute(TargetProfile);
-	
-	TArray<int32> VertexSubset = MoveTemp(TransferBoneWeights.TargetVerticesSubset);
 
 	// apply the weight changes as a transaction
-	ApplyTranferredWeightsAsTransaction(TransferredSkinWeights, VertexSubset, TargetMesh);
+	ApplyTranferredWeightsAsTransaction(TransferredSkinWeights, TargetSelectedVertices, TargetMesh);
 }
 
 void UWeightToolTransferManager::TransferWeightsFromSameMeshAndLOD()
@@ -1158,27 +1151,32 @@ void UWeightToolTransferManager::ApplyTranferredWeightsAsTransaction(
 	static constexpr float ZeroWeight = 0.f;
 	const bool bUseSubset = !InVertexSubset.IsEmpty();
 
+	const UWeightToolSelectionIsolator* Isolator = WeightTool->GetSelectionIsolator();
+	const bool bIsSelectionIsolated = Isolator->IsSelectionIsolated();
+	
 	// make sure the target has been cleaned as FSkinToolWeights data is based on a cleaned mesh description
 	FMeshDescription DummyCleanedMeshDescription;
 	FCleanedEditMesh CleanedTarget(InTargetMesh, DummyCleanedMeshDescription);
-	const FDynamicMesh3& TargetMesh = CleanedTarget.GetEditableMesh();
+
+	// InVertexSubset and Weights.PreChangeWeights/CurrentWeights reflect the current state of the mesh (partial or not) being edited
+	// so the target is one or the other.
+	// the full mesh is still needed to be able to get the original VertexID in order to get the transferred weights.  
+	const FDynamicMesh3& TargetMesh = bIsSelectionIsolated ? Isolator->GetPartialMesh() : CleanedTarget.GetEditableMesh();
 	
-	const int32 NumVertices = bUseSubset ? InVertexSubset.Num() : InTargetMesh.VertexCount();
-	const UE::Geometry::FNonManifoldMappingSupport NonManifoldMappingSupport(TargetMesh);
+	const int32 NumVertices = bUseSubset ? InVertexSubset.Num() : TargetMesh.VertexCount();
 	for (int32 VertexIndex = 0; VertexIndex < NumVertices; ++VertexIndex)
 	{
 		const int32 VertexID = bUseSubset ? InVertexSubset[VertexIndex] : VertexIndex;
-		const int32 SrcVertexID = NonManifoldMappingSupport.GetOriginalNonManifoldVertexID(VertexID);
 		
 		// remove all weight on vertex
-		const VertexWeights& VertexBoneWeights = Weights.PreChangeWeights[SrcVertexID];
+		const VertexWeights& VertexBoneWeights = Weights.PreChangeWeights[VertexID];
 		if (!VertexBoneWeights.IsEmpty())
 		{
 			for (const FVertexBoneWeight& BoneWeight : VertexBoneWeights)
 			{
 				// when transferring weights we do prune influences because we would prefer the results to be identical to the source
 				constexpr bool bPruneInfluence = true;
-				WeightEdits.MergeSingleEdit(BoneWeight.BoneID, SrcVertexID, ZeroWeight, bPruneInfluence, Weights.PreChangeWeights);
+				WeightEdits.MergeSingleEdit(BoneWeight.BoneID, VertexID, ZeroWeight, bPruneInfluence, Weights.PreChangeWeights);
 			}
 		}
 		else
@@ -1186,18 +1184,23 @@ void UWeightToolTransferManager::ApplyTranferredWeightsAsTransaction(
 			// in the unlikely event that the target vertex has no weight, "fake" remove it from the root so that undo will put it back there
 			constexpr bool bPruneInfluence = false;
 			constexpr BoneIndex RootBoneIndex = 0;
-			WeightEdits.MergeSingleEdit(RootBoneIndex, SrcVertexID, ZeroWeight, bPruneInfluence, Weights.PreChangeWeights);
+			WeightEdits.MergeSingleEdit(RootBoneIndex, VertexID, ZeroWeight, bPruneInfluence, Weights.PreChangeWeights);
 		}
 
+		// map from partial to cleaned mesh (if isolated) 
+		const int32 CleanedVertexID = bIsSelectionIsolated ? Isolator->PartialToFullMeshVertexIndex(VertexID) : VertexID;
+		// map from cleaned mesh to the full dynamesh
+		const int32 FullVertexID = CleanedTarget.GetOriginalVertexID(CleanedVertexID);
+
 		// update with new weight
-		UE::AnimationCore::FBoneWeights TransferedBoneWeights;
-		InTransferredSkinWeights->GetValue(VertexID, TransferedBoneWeights);
-		for (const UE::AnimationCore::FBoneWeight& BoneWeight: TransferedBoneWeights)
+		UE::AnimationCore::FBoneWeights TransferredBoneWeights;
+		InTransferredSkinWeights->GetValue(FullVertexID, TransferredBoneWeights);
+		for (const UE::AnimationCore::FBoneWeight& BoneWeight: TransferredBoneWeights)
 		{
 			const int32 BoneIndex = BoneWeight.GetBoneIndex();					
 			const float NewWeight = BoneWeight.GetWeight();
 			constexpr bool bPruneInfluence = false;
-			WeightEdits.MergeSingleEdit(BoneIndex, SrcVertexID, NewWeight, bPruneInfluence, Weights.PreChangeWeights);
+			WeightEdits.MergeSingleEdit(BoneIndex, VertexID, NewWeight, bPruneInfluence, Weights.PreChangeWeights);
 		}
 	}
 
@@ -4635,6 +4638,12 @@ int32 UWeightToolSelectionIsolator::FullToPartialMeshVertexIndex(int32 FullMeshV
 	return PartialSubMesh.MapVertexToSubmesh(FullMeshVertexIndex);
 }
 
+const FDynamicMesh3& UWeightToolSelectionIsolator::GetPartialMesh() const
+{
+	static const FDynamicMesh3 Dummy;
+	return PartialMeshDescription.IsValid() ? PartialSubMesh.GetSubmesh() : Dummy;
+}
+
 void UWeightToolSelectionIsolator::CreatePartialMesh()
 {
 	UPolygonSelectionMechanic* SelectionMechanic = WeightTool->GetMainMeshSelector()->GetSelectionMechanic();
@@ -4733,9 +4742,9 @@ void UWeightToolSelectionIsolator::CreatePartialMesh()
 	WeightTool->UpdateCurrentlyEditedMesh(SkeletalMeshComponent, PartialSubMesh.GetSubmesh(), *PartialMeshDescription.Get());
 }
 
-FCleanedEditMesh::FCleanedEditMesh(
-	const FDynamicMesh3& InDynamicMesh,
-	const FMeshDescription& InMeshDescription) : OriginalDynamicMesh(InDynamicMesh), OriginalMeshDescription(InMeshDescription)
+FCleanedEditMesh::FCleanedEditMesh(const FDynamicMesh3& InDynamicMesh, const FMeshDescription& InMeshDescription)
+	: OriginalDynamicMesh(InDynamicMesh)
+	, OriginalMeshDescription(InMeshDescription)
 {
 	using namespace UE::Geometry;
 	
@@ -4873,6 +4882,11 @@ FMeshDescription& FCleanedEditMesh::GetOriginalMeshDescription()
 	return OriginalMeshDescription;
 }
 
+int32 FCleanedEditMesh::GetOriginalVertexID(const int32 InCleanedVertexId) const
+{
+	return CleanedSubMeshDescription ? CleanedSubMesh.MapVertexToBaseMesh(InCleanedVertexId) : InCleanedVertexId;
+}
+
 bool USkinWeightsPaintTool::HasActiveSelectionOnMainMesh()
 {
 	if (!WeightToolProperties)
diff --git a/Engine/Plugins/Runtime/MeshModelingToolset/Source/MeshModelingToolsEditorOnly/Public/SkeletalMesh/SkinWeightsPaintTool.h b/Engine/Plugins/Runtime/MeshModelingToolset/Source/MeshModelingToolsEditorOnly/Public/SkeletalMesh/SkinWeightsPaintTool.h
--- a/Engine/Plugins/Runtime/MeshModelingToolset/Source/MeshModelingToolsEditorOnly/Public/SkeletalMesh/SkinWeightsPaintTool.h
+++ b/Engine/Plugins/Runtime/MeshModelingToolset/Source/MeshModelingToolsEditorOnly/Public/SkeletalMesh/SkinWeightsPaintTool.h
@@ -643,6 +643,9 @@ public:
 	int32 PartialToFullMeshVertexIndex(int32 PartialMeshVertexIndex) const;
 	int32 FullToPartialMeshVertexIndex(int32 FullMeshVertexIndex) const;
 
+	// returns the isolated partial mesh (if PartialMeshDescription is not null, this will return an empty DynamicMesh3)
+	const FDynamicMesh3& GetPartialMesh() const;
+
 private:
 
 	void CreatePartialMesh();
@@ -689,6 +692,9 @@ struct MESHMODELINGTOOLSEDITORONLY_API FCleanedEditMesh
 	FDynamicMesh3& GetOriginalMesh();
 	FMeshDescription& GetOriginalMeshDescription();
 
+	// returns the VertexID on the original mesh 
+	int32 GetOriginalVertexID(const int32 InCleanedVertexId) const;
+	
 private:
 	FDynamicMesh3 OriginalDynamicMesh;
 	FMeshDescription OriginalMeshDescription;
