Use ILatticeStateStorage to notify when the Lattice Tool begins and ends

diff --git a/Engine/Plugins/Runtime/MeshModelingToolset/Source/MeshModelingTools/Private/LatticeDeformerTool.cpp b/Engine/Plugins/Runtime/MeshModelingToolset/Source/MeshModelingTools/Private/LatticeDeformerTool.cpp
--- a/Engine/Plugins/Runtime/MeshModelingToolset/Source/MeshModelingTools/Private/LatticeDeformerTool.cpp
+++ b/Engine/Plugins/Runtime/MeshModelingToolset/Source/MeshModelingTools/Private/LatticeDeformerTool.cpp
@@ -395,6 +395,8 @@ void ULatticeDeformerTool::Setup()
 		}
 
 		Settings->InterpolationType = LatticeStorage->ReadInterpolationType();
+
+		LatticeStorage->InteractiveToolStarted();
 	}
 
 	if (Settings->bSoftDeformation)
@@ -534,6 +536,11 @@ void ULatticeDeformerTool::OnShutdown(EToolShutdownType ShutdownType)
 			GetToolManager()->EndUndoTransaction();
 		}
 	}
+
+	if (LatticeStorage)
+	{
+		LatticeStorage->InteractiveToolShutDown();
+	}
 }
 
 
diff --git a/Engine/Plugins/Runtime/MeshModelingToolset/Source/ModelingComponents/Public/LatticeManager.h b/Engine/Plugins/Runtime/MeshModelingToolset/Source/ModelingComponents/Public/LatticeManager.h
--- a/Engine/Plugins/Runtime/MeshModelingToolset/Source/ModelingComponents/Public/LatticeManager.h
+++ b/Engine/Plugins/Runtime/MeshModelingToolset/Source/ModelingComponents/Public/LatticeManager.h
@@ -39,4 +39,7 @@ public:
 	virtual UE::Geometry::FAxisAlignedBox3d GetInitialBounds() const = 0;
 	virtual UE::Geometry::FTransformSRT3d GetTransform() const = 0;
 	virtual UE::Geometry::FVector3i GetResolution() const = 0;
+
+	virtual void InteractiveToolStarted() = 0;
+	virtual void InteractiveToolShutDown() = 0;
 };
