Make the triangulate curves op properly handle the duplicate start/end point generated on looping splines

diff --git a/Engine/Plugins/Runtime/MeshModelingToolset/Source/ModelingOperators/Private/CurveOps/TriangulateCurvesOp.cpp b/Engine/Plugins/Runtime/MeshModelingToolset/Source/ModelingOperators/Private/CurveOps/TriangulateCurvesOp.cpp
--- a/Engine/Plugins/Runtime/MeshModelingToolset/Source/ModelingOperators/Private/CurveOps/TriangulateCurvesOp.cpp
+++ b/Engine/Plugins/Runtime/MeshModelingToolset/Source/ModelingOperators/Private/CurveOps/TriangulateCurvesOp.cpp
@@ -23,6 +23,11 @@ void FTriangulateCurvesOp::AddSpline(USplineComponent* Spline, double ErrorToler
 	FCurvePath& Path = Paths.Emplace_GetRef();
 	Spline->ConvertSplineToPolyLine(ESplineCoordinateSpace::Type::World, ErrorTolerance * ErrorTolerance, Path.Vertices);
 	Path.bClosed = Spline->IsClosedLoop();
+	// Spline conversion creates a duplicate end/start point, which we don't need for triangulation
+	if (Path.bClosed && Path.Vertices.Num() > 1)
+	{
+		Path.Vertices.Pop();
+	}
 	if (Paths.Num() == 1)
 	{
 		FirstPathTransform = Spline->GetComponentTransform();
@@ -32,8 +37,16 @@ void FTriangulateCurvesOp::AddSpline(USplineComponent* Spline, double ErrorToler
 void FTriangulateCurvesOp::AddWorldCurve(TArrayView<const FVector3d> WorldSpaceVertices, bool bClosed, const FTransform& ReferenceTransform)
 {
 	FCurvePath& Path = Paths.Emplace_GetRef();
-	Path.Vertices.Append(WorldSpaceVertices);
 	Path.bClosed = bClosed;
+	// Make sure we don't add a duplicate end/start point on closed curves
+	if (bClosed && WorldSpaceVertices.Num() > 1 && WorldSpaceVertices[0] == WorldSpaceVertices.Last())
+	{
+		Path.Vertices.Append(WorldSpaceVertices.Slice(0, WorldSpaceVertices.Num() - 1));
+	}
+	else
+	{
+		Path.Vertices.Append(WorldSpaceVertices);
+	}
 	if (Paths.Num() == 1)
 	{
 		FirstPathTransform = ReferenceTransform;
