Make curve triangulation op use the 'join method' for closed curves, rather than always using 'square' joins

diff --git a/Engine/Plugins/Runtime/MeshModelingToolset/Source/ModelingOperators/Private/CurveOps/TriangulateCurvesOp.cpp b/Engine/Plugins/Runtime/MeshModelingToolset/Source/ModelingOperators/Private/CurveOps/TriangulateCurvesOp.cpp
--- a/Engine/Plugins/Runtime/MeshModelingToolset/Source/ModelingOperators/Private/CurveOps/TriangulateCurvesOp.cpp
+++ b/Engine/Plugins/Runtime/MeshModelingToolset/Source/ModelingOperators/Private/CurveOps/TriangulateCurvesOp.cpp
@@ -273,7 +273,7 @@ void FTriangulateCurvesOp::CalculateResult(FProgressCancel* Progress)
 					PolygonsToAdd, PolygonsToAdd,
 					true,
 					MiterLimit,
-					EPolygonOffsetJoinType::Square, // TODO: convert this from inputs
+					ConvertOffsetJoinType(OffsetJoinMethod),
 					OffsetClosedMethod == EOffsetClosedCurvesMethod::OffsetOuterSide ? EPolygonOffsetEndType::Polygon : EPolygonOffsetEndType::Joined);
 			}
 			GetPolygonsGroup(CurveIdx).Append(PolygonsToAdd);
