Updated headers using UnrealCodeFixup to make sure dllstorage is on methods/staticvars instead of types. Part 4/n

diff --git a/Engine/Plugins/Runtime/HairStrands/Source/HairStrandsEditor/Public/HairStrandsEditor.h b/Engine/Plugins/Runtime/HairStrands/Source/HairStrandsEditor/Public/HairStrandsEditor.h
--- a/Engine/Plugins/Runtime/HairStrands/Source/HairStrandsEditor/Public/HairStrandsEditor.h
+++ b/Engine/Plugins/Runtime/HairStrands/Source/HairStrandsEditor/Public/HairStrandsEditor.h
@@ -9,6 +9,8 @@
 #include "HAL/LowLevelMemTracker.h"
 #include "RHIFwd.h"
 
+#define UE_API HAIRSTRANDSEDITOR_API
+
 #define HAIRSTRANDSEDITOR_MODULE_NAME TEXT("HairStrandsEditor")
 
 class IGroomTranslator;
@@ -16,20 +18,20 @@ class IGroomTranslator;
 LLM_DECLARE_TAG_API(GroomEditor, HAIRSTRANDSEDITOR_API);
 
 /** Implements the HairStrands module  */
-class HAIRSTRANDSEDITOR_API FGroomEditor : public IModuleInterface
+class FGroomEditor : public IModuleInterface
 {
 public:
 
 	//~ IModuleInterface interface
 
-	virtual void StartupModule() override;
-	virtual void ShutdownModule() override;
+	UE_API virtual void StartupModule() override;
+	UE_API virtual void ShutdownModule() override;
 	virtual bool SupportsDynamicReloading() override { return false; }
 
-	void OnPostEngineInit();
-	void OnPreviewPlatformChanged();
-	void OnPreExit();
-	void OnPreviewFeatureLevelChanged(ERHIFeatureLevel::Type InPreviewFeatureLevel);
+	UE_API void OnPostEngineInit();
+	UE_API void OnPreviewPlatformChanged();
+	UE_API void OnPreExit();
+	UE_API void OnPreviewFeatureLevelChanged(ERHIFeatureLevel::Type InPreviewFeatureLevel);
 
 	static inline FGroomEditor& Get()
 	{
@@ -48,10 +50,10 @@ public:
 	}
 
 	/** Get new instances of HairStrandsTranslators */
-	TArray<TSharedPtr<IGroomTranslator>> GetHairTranslators();
+	UE_API TArray<TSharedPtr<IGroomTranslator>> GetHairTranslators();
 
 private:
-	void RegisterMenus();
+	UE_API void RegisterMenus();
 
 	TArray<TFunction<TSharedPtr<IGroomTranslator>()>> TranslatorSpawners;
 
@@ -63,5 +65,7 @@ private:
 
 public:
 
-	static FName GroomEditorAppIdentifier;
+	static UE_API FName GroomEditorAppIdentifier;
 };
+
+#undef UE_API
diff --git a/Engine/Plugins/Runtime/HairStrands/Source/HairStrandsEditor/Public/HairStrandsImporter.h b/Engine/Plugins/Runtime/HairStrands/Source/HairStrandsEditor/Public/HairStrandsImporter.h
--- a/Engine/Plugins/Runtime/HairStrands/Source/HairStrandsEditor/Public/HairStrandsImporter.h
+++ b/Engine/Plugins/Runtime/HairStrands/Source/HairStrandsEditor/Public/HairStrandsImporter.h
@@ -5,14 +5,16 @@
 #include "CoreMinimal.h"
 #include "UObject/ObjectMacros.h"
 
+#define UE_API HAIRSTRANDSEDITOR_API
+
 class FHairDescription;
 class UGroomAsset;
 class UGroomImportOptions;
 class UGroomHairGroupsMapping;
 
-struct HAIRSTRANDSEDITOR_API FHairImportContext
+struct FHairImportContext
 {
-	FHairImportContext(UGroomImportOptions* InImportOptions, UObject* InParent = nullptr, UClass* InClass = nullptr, FName InName = FName(), EObjectFlags InFlags = EObjectFlags::RF_NoFlags);
+	UE_API FHairImportContext(UGroomImportOptions* InImportOptions, UObject* InParent = nullptr, UClass* InClass = nullptr, FName InName = FName(), EObjectFlags InFlags = EObjectFlags::RF_NoFlags);
 
 	UGroomImportOptions* ImportOptions;
 	UObject* Parent;
@@ -21,7 +23,9 @@ struct HAIRSTRANDSEDITOR_API FHairImportContext
 	EObjectFlags Flags;
 };
 
-struct HAIRSTRANDSEDITOR_API FHairStrandsImporter
+struct FHairStrandsImporter
 {
-	static UGroomAsset* ImportHair(const FHairImportContext& ImportContext, FHairDescription& HairDescription, UGroomAsset* ExistingHair = nullptr, const UGroomHairGroupsMapping* GroupsMapping=nullptr);
+	static UE_API UGroomAsset* ImportHair(const FHairImportContext& ImportContext, FHairDescription& HairDescription, UGroomAsset* ExistingHair = nullptr, const UGroomHairGroupsMapping* GroupsMapping=nullptr);
 };
+
+#undef UE_API
diff --git a/Engine/Plugins/Runtime/HairStrands/Source/HairStrandsEditor/Public/HairStrandsTranslator.h b/Engine/Plugins/Runtime/HairStrands/Source/HairStrandsEditor/Public/HairStrandsTranslator.h
--- a/Engine/Plugins/Runtime/HairStrands/Source/HairStrandsEditor/Public/HairStrandsTranslator.h
+++ b/Engine/Plugins/Runtime/HairStrands/Source/HairStrandsEditor/Public/HairStrandsTranslator.h
@@ -4,7 +4,7 @@
 
 #include "HairDescription.h"
 
-class HAIRSTRANDSEDITOR_API IGroomTranslator
+class IGroomTranslator
 {
 public:
 	virtual ~IGroomTranslator() {}
diff --git a/Engine/Plugins/Runtime/Harmonix/Source/HarmonixDspEditor/Private/FusionPatchAssetFactory.h b/Engine/Plugins/Runtime/Harmonix/Source/HarmonixDspEditor/Private/FusionPatchAssetFactory.h
--- a/Engine/Plugins/Runtime/Harmonix/Source/HarmonixDspEditor/Private/FusionPatchAssetFactory.h
+++ b/Engine/Plugins/Runtime/Harmonix/Source/HarmonixDspEditor/Private/FusionPatchAssetFactory.h
@@ -9,34 +9,36 @@
 
 #include "FusionPatchAssetFactory.generated.h"
 
+#define UE_API HARMONIXDSPEDITOR_API
+
 class UFusionPatchImportOptions;
 class UFusionPatchCreateOptions;
 
 DECLARE_LOG_CATEGORY_EXTERN(LogFusionPatchAssetFactory, Log, All);
 
-UCLASS()
-class HARMONIXDSPEDITOR_API UFusionPatchAssetFactory : public UFactory, public FReimportHandler
+UCLASS(MinimalAPI)
+class UFusionPatchAssetFactory : public UFactory, public FReimportHandler
 {
 	GENERATED_BODY()
 
 public:
 
-	UFusionPatchAssetFactory();
+	UE_API UFusionPatchAssetFactory();
 
-	virtual bool FactoryCanImport(const FString& Filename) override;
-	virtual bool CanReimport(UObject* Obj, TArray<FString>& OutFilenames) override;
-	virtual void SetReimportPaths(UObject* Obj, const TArray<FString>& NewReimportPaths) override;;
-	virtual EReimportResult::Type Reimport(UObject* Obj) override;
-	virtual void CleanUp() override;
+	UE_API virtual bool FactoryCanImport(const FString& Filename) override;
+	UE_API virtual bool CanReimport(UObject* Obj, TArray<FString>& OutFilenames) override;
+	UE_API virtual void SetReimportPaths(UObject* Obj, const TArray<FString>& NewReimportPaths) override;;
+	UE_API virtual EReimportResult::Type Reimport(UObject* Obj) override;
+	UE_API virtual void CleanUp() override;
 
-	virtual UObject* FactoryCreateText(UClass* InClass, UObject* InParent, FName InName, EObjectFlags Flags, UObject* Context, const TCHAR* Type, const TCHAR*& Buffer, const TCHAR* BufferEnd, FFeedbackContext* Warn) override;
-	virtual UObject* FactoryCreateNew(UClass* InClass, UObject* InParent, FName InName, EObjectFlags Flags, UObject* Context, FFeedbackContext* Warn) override;
+	UE_API virtual UObject* FactoryCreateText(UClass* InClass, UObject* InParent, FName InName, EObjectFlags Flags, UObject* Context, const TCHAR* Type, const TCHAR*& Buffer, const TCHAR* BufferEnd, FFeedbackContext* Warn) override;
+	UE_API virtual UObject* FactoryCreateNew(UClass* InClass, UObject* InParent, FName InName, EObjectFlags Flags, UObject* Context, FFeedbackContext* Warn) override;
 
 	const UFusionPatchCreateOptions* CreateOptions = nullptr;
 private:
 
-	static bool GetReplaceExistingSamplesResponse(const FString& InName);
-	static bool GetApplyOptionsToAllImportResponse();
+	static UE_API bool GetReplaceExistingSamplesResponse(const FString& InName);
+	static UE_API bool GetApplyOptionsToAllImportResponse();
 
 	enum class EApplyAllOption : uint8
 	{
@@ -50,5 +52,7 @@ private:
 	bool ReplaceExistingSamples = false;
 	TArray<UObject*> ImportedObjects;
 
-	void UpdateFusionPatchImportNotificationItem(TSharedPtr<SNotificationItem> InItem, bool bImportSuccessful, FName InName);
-};
\ No newline at end of file
+	UE_API void UpdateFusionPatchImportNotificationItem(TSharedPtr<SNotificationItem> InItem, bool bImportSuccessful, FName InName);
+};
+
+#undef UE_API
diff --git a/Engine/Plugins/Runtime/InstancedActors/Source/InstancedActorsEditor/Public/InstancedActorsEditorModule.h b/Engine/Plugins/Runtime/InstancedActors/Source/InstancedActorsEditor/Public/InstancedActorsEditorModule.h
--- a/Engine/Plugins/Runtime/InstancedActors/Source/InstancedActorsEditor/Public/InstancedActorsEditorModule.h
+++ b/Engine/Plugins/Runtime/InstancedActors/Source/InstancedActorsEditor/Public/InstancedActorsEditorModule.h
@@ -9,6 +9,8 @@
 #include "Modules/ModuleInterface.h"
 #include "Templates/SubclassOf.h"
 
+#define UE_API INSTANCEDACTORSEDITOR_API
+
 class AActor;
 class FExtender;
 class FUICommandList;
@@ -20,45 +22,45 @@ class UInstancedActorsSubsystem;
 * This module contains logic for converting between AActors and Instanced Actors. 
 * @todo move to a dedicated class
 */
-class INSTANCEDACTORSEDITOR_API FInstancedActorsEditorModule : public IModuleInterface
+class FInstancedActorsEditorModule : public IModuleInterface
 {
 public:
 	DECLARE_DELEGATE_OneParam(FOnConvert, TConstArrayView<AActor*>);
 
 	/** Convert selected actors to Instance Actors (IAs) while using the IASubsystemClass */
-	void CustomizedConvertActorsToIAsUIAction(TConstArrayView<AActor*> InActors, TSubclassOf<UInstancedActorsSubsystem> IASubsystemClass) const;
+	UE_API void CustomizedConvertActorsToIAsUIAction(TConstArrayView<AActor*> InActors, TSubclassOf<UInstancedActorsSubsystem> IASubsystemClass) const;
 	
 	/** 
 	 * Sets InDelegate as the delegate that will be executed as "Convert Instanced Actors to regular Actors" action.
 	 * @param ActionFormatLabelOverride if supplied will be used to override the default action label. Note that 
 	 *	this is a format string and needs to require one parameter, the Actor's name
 	 */
-	void SetIAToActorDelegate(const FOnConvert& InDelegate, const FTextFormat& ActionFormatLabelOverride = FTextFormat());
+	UE_API void SetIAToActorDelegate(const FOnConvert& InDelegate, const FTextFormat& ActionFormatLabelOverride = FTextFormat());
 
 	/**
 	 * Sets InDelegate as the delegate that will be executed as "Convert Actors to Instanced Actors" action.
 	 * @param ActionFormatLabelOverride if supplied will be used to override the default action label. Note that
 	 *	this is a format string and needs to require one parameter, the Actor's name
 	 */
-	void SetActorToIADelegate(const FOnConvert& InDelegate, const FTextFormat& ActionFormatLabelOverride = FTextFormat());
+	UE_API void SetActorToIADelegate(const FOnConvert& InDelegate, const FTextFormat& ActionFormatLabelOverride = FTextFormat());
 
 	/** Resets FOnConvert delegates and action labels to their defaults */
-	void ResetConversionDelegates();
+	UE_API void ResetConversionDelegates();
 
 protected:
 	// Begin IModuleInterface
-	virtual void StartupModule() override;
-	virtual void ShutdownModule() override;
+	UE_API virtual void StartupModule() override;
+	UE_API virtual void ShutdownModule() override;
 
 	/** Convert selected actors to Instance Actors (IAs). */
-	void ConvertActorsToIAsUIAction(TConstArrayView<AActor*> InActors) const;
+	UE_API void ConvertActorsToIAsUIAction(TConstArrayView<AActor*> InActors) const;
 
 	/** Convert selected Instance Actors to Actors. */
-	void ConvertIAsToActorsUIAction(TConstArrayView<AActor*> InActors) const;
+	UE_API void ConvertIAsToActorsUIAction(TConstArrayView<AActor*> InActors) const;
 
-	TSharedRef<FExtender> CreateLevelViewportContextMenuExtender(const TSharedRef<FUICommandList> CommandList, const TArray<AActor*> InActors);
-	void AddLevelViewportMenuExtender();
-	void RemoveLevelViewportMenuExtender();
+	UE_API TSharedRef<FExtender> CreateLevelViewportContextMenuExtender(const TSharedRef<FUICommandList> CommandList, const TArray<AActor*> InActors);
+	UE_API void AddLevelViewportMenuExtender();
+	UE_API void RemoveLevelViewportMenuExtender();
 
 	FDelegateHandle LevelViewportExtenderHandle;
 	FOnConvert ActorToIADelegate;
@@ -67,3 +69,5 @@ protected:
 	FTextFormat ActorToIAFormatLabel;
 	FTextFormat IAToActorFormatLabel;
 };
+
+#undef UE_API
diff --git a/Engine/Plugins/Runtime/MassGameplay/Source/MassGameplayEditor/Public/IMassGameplayEditor.h b/Engine/Plugins/Runtime/MassGameplay/Source/MassGameplayEditor/Public/IMassGameplayEditor.h
--- a/Engine/Plugins/Runtime/MassGameplay/Source/MassGameplayEditor/Public/IMassGameplayEditor.h
+++ b/Engine/Plugins/Runtime/MassGameplay/Source/MassGameplayEditor/Public/IMassGameplayEditor.h
@@ -7,7 +7,7 @@
 
 
 /** Mass Editor public interface */
-class MASSGAMEPLAYEDITOR_API IMassGameplayEditor : public FAssetEditorToolkit
+class IMassGameplayEditor : public FAssetEditorToolkit
 {
 public:
 };
diff --git a/Engine/Plugins/Runtime/MassGameplay/Source/MassGameplayEditor/Public/MassActorEditorSubsystem.h b/Engine/Plugins/Runtime/MassGameplay/Source/MassGameplayEditor/Public/MassActorEditorSubsystem.h
--- a/Engine/Plugins/Runtime/MassGameplay/Source/MassGameplayEditor/Public/MassActorEditorSubsystem.h
+++ b/Engine/Plugins/Runtime/MassGameplay/Source/MassGameplayEditor/Public/MassActorEditorSubsystem.h
@@ -5,11 +5,13 @@
 #include "EditorSubsystem.h"
 #include "MassActorEditorSubsystem.generated.h"
 
+#define UE_API MASSGAMEPLAYEDITOR_API
+
 
 struct FMassActorManager;
 
-UCLASS()
-class MASSGAMEPLAYEDITOR_API UMassActorEditorSubsystem : public UEditorSubsystem
+UCLASS(MinimalAPI)
+class UMassActorEditorSubsystem : public UEditorSubsystem
 {
 	GENERATED_BODY()
 
@@ -17,7 +19,9 @@ public:
 	FMassActorManager& GetMutableActorManager() { return *ActorManager.Get(); }
 
 protected:
-	virtual void Initialize(FSubsystemCollectionBase& Collection) override;
+	UE_API virtual void Initialize(FSubsystemCollectionBase& Collection) override;
 
 	TSharedPtr<FMassActorManager> ActorManager;
 };
+
+#undef UE_API
diff --git a/Engine/Plugins/Runtime/MassGameplay/Source/MassGameplayEditor/Public/MassGameplayEditorModule.h b/Engine/Plugins/Runtime/MassGameplay/Source/MassGameplayEditor/Public/MassGameplayEditorModule.h
--- a/Engine/Plugins/Runtime/MassGameplay/Source/MassGameplayEditor/Public/MassGameplayEditorModule.h
+++ b/Engine/Plugins/Runtime/MassGameplay/Source/MassGameplayEditor/Public/MassGameplayEditorModule.h
@@ -4,15 +4,19 @@
 
 #include "Modules/ModuleInterface.h"
 
+#define UE_API MASSGAMEPLAYEDITOR_API
+
 /**
 * The public interface to this module
 */
-class MASSGAMEPLAYEDITOR_API FMassGameplayEditorModule : public IModuleInterface
+class FMassGameplayEditorModule : public IModuleInterface
 {
 public:
-	virtual void StartupModule() override;
-	virtual void ShutdownModule() override;
+	UE_API virtual void StartupModule() override;
+	UE_API virtual void ShutdownModule() override;
 
 protected:
-	void RegisterSectionMappings();
+	UE_API void RegisterSectionMappings();
 };
+
+#undef UE_API
diff --git a/Engine/Plugins/Runtime/MassGameplay/Source/MassGameplayEditor/Public/MassTraitRepository.h b/Engine/Plugins/Runtime/MassGameplay/Source/MassGameplayEditor/Public/MassTraitRepository.h
--- a/Engine/Plugins/Runtime/MassGameplay/Source/MassGameplayEditor/Public/MassTraitRepository.h
+++ b/Engine/Plugins/Runtime/MassGameplay/Source/MassGameplayEditor/Public/MassTraitRepository.h
@@ -39,36 +39,36 @@ struct FMassTraitInspectionContext
  * traits supply a fragment they need.
  * It also serves to post trait validation information and fix options to the MessageLog.
  */
-UCLASS()
-class MASSGAMEPLAYEDITOR_API UMassTraitRepository : public UEditorSubsystem
+UCLASS(MinimalAPI)
+class UMassTraitRepository : public UEditorSubsystem
 {
 	GENERATED_BODY()
 	
 public:
-	static UWorld* GetInvestigationWorld();
+	static MASSGAMEPLAYEDITOR_API UWorld* GetInvestigationWorld();
 
-	TConstArrayView<FName> GetTraitsNameAddingElements(const FName ElementName);
-	TWeakObjectPtr<UClass> GetTraitClass(const FName TraitClassName) const;
+	MASSGAMEPLAYEDITOR_API TConstArrayView<FName> GetTraitsNameAddingElements(const FName ElementName);
+	MASSGAMEPLAYEDITOR_API TWeakObjectPtr<UClass> GetTraitClass(const FName TraitClassName) const;
 
 protected:
-	virtual void Initialize(FSubsystemCollectionBase& Collection) override;
-	virtual void Deinitialize() override;
+	MASSGAMEPLAYEDITOR_API virtual void Initialize(FSubsystemCollectionBase& Collection) override;
+	MASSGAMEPLAYEDITOR_API virtual void Deinitialize() override;
 
 #if WITH_MASSENTITY_DEBUG
 	/** 
 	 * Note that the function does nothing until InitRepository is called. This is done to avoid collecting trait
 	 * data until it's actually necessary.
 	 */
-	void OnNewTraitType(UMassEntityTraitBase& Trait);
+	MASSGAMEPLAYEDITOR_API void OnNewTraitType(UMassEntityTraitBase& Trait);
 
-	void OnDebugEvent(const FName EventName, FConstStructView Payload, EMassDebugMessageSeverity SeverityOverride);
+	MASSGAMEPLAYEDITOR_API void OnDebugEvent(const FName EventName, FConstStructView Payload, EMassDebugMessageSeverity SeverityOverride);
 #endif // WITH_MASSENTITY_DEBUG
 	/** 
 	 * The method gathers all existing trait classes and processed them. We don't do that on subsystem's init since
 	 * very often the data won't be needed during the given editor run. We call the function lazily the first time 
 	 * the data is needed.
 	 */
-	void InitRepository();
+	MASSGAMEPLAYEDITOR_API void InitRepository();
 
 	FDelegateHandle OnNewTraitTypeHandle;
 	struct FTraitAndElements
@@ -87,7 +87,7 @@ protected:
 	/** We initialize the repository's data lazily and this property indicates whether it has been already done. */
 	bool bIsRepositoryInitialized;
 
-	static TWeakObjectPtr<UWorld> GlobalInvestigationWorld;
+	static MASSGAMEPLAYEDITOR_API TWeakObjectPtr<UWorld> GlobalInvestigationWorld;
 };
 
 
diff --git a/Engine/Plugins/Runtime/MassGameplay/Source/MassMovementEditor/Public/IMassMovementEditor.h b/Engine/Plugins/Runtime/MassGameplay/Source/MassMovementEditor/Public/IMassMovementEditor.h
--- a/Engine/Plugins/Runtime/MassGameplay/Source/MassMovementEditor/Public/IMassMovementEditor.h
+++ b/Engine/Plugins/Runtime/MassGameplay/Source/MassMovementEditor/Public/IMassMovementEditor.h
@@ -6,7 +6,7 @@
 #include "Toolkits/AssetEditorToolkit.h"
 
 /** MassMovement Editor public interface */
-class MASSMOVEMENTEDITOR_API IMassMovementEditor : public FAssetEditorToolkit
+class IMassMovementEditor : public FAssetEditorToolkit
 {
 public:
 
diff --git a/Engine/Plugins/Runtime/MassGameplay/Source/MassMovementEditor/Public/MassMovementEditorModule.h b/Engine/Plugins/Runtime/MassGameplay/Source/MassMovementEditor/Public/MassMovementEditorModule.h
--- a/Engine/Plugins/Runtime/MassGameplay/Source/MassMovementEditor/Public/MassMovementEditorModule.h
+++ b/Engine/Plugins/Runtime/MassGameplay/Source/MassMovementEditor/Public/MassMovementEditorModule.h
@@ -16,10 +16,10 @@ class IMassMovementEditor;
 /**
 * The public interface to this module
 */
-class MASSMOVEMENTEDITOR_API FMassMovementEditorModule : public IModuleInterface//, public IHasMenuExtensibility, public IHasToolBarExtensibility
+class FMassMovementEditorModule : public IModuleInterface//, public IHasMenuExtensibility, public IHasToolBarExtensibility
 {
 public:
 	// Begin IModuleInterface
-	virtual void StartupModule() override;
-	virtual void ShutdownModule() override;
+	MASSMOVEMENTEDITOR_API virtual void StartupModule() override;
+	MASSMOVEMENTEDITOR_API virtual void ShutdownModule() override;
 };
diff --git a/Engine/Plugins/Runtime/MeshModelingToolset/Source/MeshModelingToolsEditorOnly/Public/AttributeEditorTool.h b/Engine/Plugins/Runtime/MeshModelingToolset/Source/MeshModelingToolsEditorOnly/Public/AttributeEditorTool.h
--- a/Engine/Plugins/Runtime/MeshModelingToolset/Source/MeshModelingToolsEditorOnly/Public/AttributeEditorTool.h
+++ b/Engine/Plugins/Runtime/MeshModelingToolset/Source/MeshModelingToolsEditorOnly/Public/AttributeEditorTool.h
@@ -9,6 +9,8 @@
 #include "InteractiveToolChange.h"
 #include "AttributeEditorTool.generated.h"
 
+#define UE_API MESHMODELINGTOOLSEDITORONLY_API
+
 
 class IPrimitiveComponentBackedTarget;
 class UAttributeEditorTool;
@@ -17,17 +19,17 @@ class UAttributeEditorTool;
 /**
  *
  */
-UCLASS()
-class MESHMODELINGTOOLSEDITORONLY_API UAttributeEditorToolBuilder : public UMultiSelectionMeshEditingToolBuilder
+UCLASS(MinimalAPI)
+class UAttributeEditorToolBuilder : public UMultiSelectionMeshEditingToolBuilder
 {
 	GENERATED_BODY()
 
 public:
-	virtual UMultiSelectionMeshEditingTool* CreateNewTool(const FToolBuilderState& SceneState) const override;
-	virtual bool CanBuildTool(const FToolBuilderState& SceneState) const override;
+	UE_API virtual UMultiSelectionMeshEditingTool* CreateNewTool(const FToolBuilderState& SceneState) const override;
+	UE_API virtual bool CanBuildTool(const FToolBuilderState& SceneState) const override;
 
 protected:
-	virtual const FToolTargetTypeRequirements& GetTargetRequirements() const override;
+	UE_API virtual const FToolTargetTypeRequirements& GetTargetRequirements() const override;
 };
 
 
@@ -61,8 +63,8 @@ enum class EAttributeEditorAttribType : uint8
 };
 
 
-UCLASS()
-class MESHMODELINGTOOLSEDITORONLY_API UAttributeEditorAttribProperties : public UInteractiveToolPropertySet
+UCLASS(MinimalAPI)
+class UAttributeEditorAttribProperties : public UInteractiveToolPropertySet
 {
 	GENERATED_BODY()
 public:
@@ -114,8 +116,8 @@ enum class EAttributeEditorToolActions
 
 
 
-UCLASS()
-class MESHMODELINGTOOLSEDITORONLY_API UAttributeEditorActionPropertySet : public UInteractiveToolPropertySet
+UCLASS(MinimalAPI)
+class UAttributeEditorActionPropertySet : public UInteractiveToolPropertySet
 {
 	GENERATED_BODY()
 
@@ -123,13 +125,13 @@ public:
 	TWeakObjectPtr<UAttributeEditorTool> ParentTool;
 
 	void Initialize(UAttributeEditorTool* ParentToolIn) { ParentTool = ParentToolIn; }
-	void PostAction(EAttributeEditorToolActions Action);
+	UE_API void PostAction(EAttributeEditorToolActions Action);
 
 };
 
 
-UCLASS()
-class MESHMODELINGTOOLSEDITORONLY_API UAttributeEditorNormalsActions : public UAttributeEditorActionPropertySet
+UCLASS(MinimalAPI)
+class UAttributeEditorNormalsActions : public UAttributeEditorActionPropertySet
 {
 	GENERATED_BODY()
 public:
@@ -150,8 +152,8 @@ public:
 
 
 
-UCLASS()
-class MESHMODELINGTOOLSEDITORONLY_API UAttributeEditorUVActions : public UAttributeEditorActionPropertySet
+UCLASS(MinimalAPI)
+class UAttributeEditorUVActions : public UAttributeEditorActionPropertySet
 {
 	GENERATED_BODY()
 public:
@@ -160,7 +162,7 @@ public:
 	FString UVLayer;
 
 	UFUNCTION()
-	TArray<FString> GetUVLayerNamesFunc();
+	UE_API TArray<FString> GetUVLayerNamesFunc();
 
 	UPROPERTY()
 	TArray<FString> UVLayerNamesList;
@@ -197,8 +199,8 @@ public:
 
 
 
-UCLASS()
-class MESHMODELINGTOOLSEDITORONLY_API UAttributeEditorLightmapUVActions : public UAttributeEditorActionPropertySet
+UCLASS(MinimalAPI)
+class UAttributeEditorLightmapUVActions : public UAttributeEditorActionPropertySet
 {
 	GENERATED_BODY()
 public:
@@ -235,8 +237,8 @@ public:
 };
 
 
-UCLASS()
-class MESHMODELINGTOOLSEDITORONLY_API UAttributeEditorNewAttributeActions : public UAttributeEditorActionPropertySet
+UCLASS(MinimalAPI)
+class UAttributeEditorNewAttributeActions : public UAttributeEditorActionPropertySet
 {
 	GENERATED_BODY()
 public:
@@ -274,8 +276,8 @@ public:
 };
 
 
-UCLASS()
-class MESHMODELINGTOOLSEDITORONLY_API UAttributeEditorModifyAttributeActions : public UAttributeEditorActionPropertySet
+UCLASS(MinimalAPI)
+class UAttributeEditorModifyAttributeActions : public UAttributeEditorActionPropertySet
 {
 	GENERATED_BODY()
 public:
@@ -283,7 +285,7 @@ public:
 	FString Attribute;
 
 	UFUNCTION()
-	TArray<FString> GetAttributeNamesFunc();
+	UE_API TArray<FString> GetAttributeNamesFunc();
 
 	UPROPERTY()
 	TArray<FString> AttributeNamesList;
@@ -304,8 +306,8 @@ public:
 };
 
 
-UCLASS()
-class MESHMODELINGTOOLSEDITORONLY_API UAttributeEditorCopyAttributeActions : public UAttributeEditorActionPropertySet
+UCLASS(MinimalAPI)
+class UAttributeEditorCopyAttributeActions : public UAttributeEditorActionPropertySet
 {
 	GENERATED_BODY()
 public:
@@ -330,24 +332,24 @@ public:
 /**
  * Mesh Attribute Editor Tool
  */
-UCLASS()
-class MESHMODELINGTOOLSEDITORONLY_API UAttributeEditorTool : public UMultiSelectionMeshEditingTool, public IInteractiveToolManageGeometrySelectionAPI
+UCLASS(MinimalAPI)
+class UAttributeEditorTool : public UMultiSelectionMeshEditingTool, public IInteractiveToolManageGeometrySelectionAPI
 {
 	GENERATED_BODY()
 
 public:
 
-	UAttributeEditorTool();
+	UE_API UAttributeEditorTool();
 
-	virtual void Setup() override;
-	virtual void OnShutdown(EToolShutdownType ShutdownType) override;
-	virtual void OnTick(float DeltaTime) override;
+	UE_API virtual void Setup() override;
+	UE_API virtual void OnShutdown(EToolShutdownType ShutdownType) override;
+	UE_API virtual void OnTick(float DeltaTime) override;
 
 	virtual bool HasCancel() const override { return false; }
 	virtual bool HasAccept() const override { return false; }
 	virtual bool CanAccept() const override { return false; }
 
-	virtual void RequestAction(EAttributeEditorToolActions ActionType);
+	UE_API virtual void RequestAction(EAttributeEditorToolActions ActionType);
 
 	// IInteractiveToolManageGeometrySelectionAPI -- this tool won't update external geometry selection or change selection-relevant mesh IDs
 	virtual bool IsInputSelectionValidOnOutput() override
@@ -384,36 +386,38 @@ protected:
 
 	bool bAttributeListsValid = false;
 	bool bHaveAutoGeneratedLightmapUVSet = false;
-	void InitializeAttributeLists();
+	UE_API void InitializeAttributeLists();
 
-	void UpdateAutoGeneratedLightmapUVChannel(UPrimitiveComponent* TargetComponent, int32 NewMaxUVChannels);
+	UE_API void UpdateAutoGeneratedLightmapUVChannel(UPrimitiveComponent* TargetComponent, int32 NewMaxUVChannels);
 
 	EAttributeEditorToolActions PendingAction = EAttributeEditorToolActions::NoAction;
-	void ClearNormals();
-	void DiscardTangents();
-	void ClearUVs();
-	void DeleteSelectedUVSet();
-	void DuplicateSelectedUVSet();
-	void AddUVSet();
-	void AddNewAttribute();
-	void AddNewWeightMap();
-	void AddNewGroupsLayer();
-	void DeleteAttribute();
-	void ClearAttribute();
-	void SetLightmapUVsEnabled(bool bEnabled);
-	void ResetLightmapUVsChannels();
-
-	void AddNewAttribute(EAttributeEditorElementType ElemType, EAttributeEditorAttribType AttribType, FName AttributeName);
+	UE_API void ClearNormals();
+	UE_API void DiscardTangents();
+	UE_API void ClearUVs();
+	UE_API void DeleteSelectedUVSet();
+	UE_API void DuplicateSelectedUVSet();
+	UE_API void AddUVSet();
+	UE_API void AddNewAttribute();
+	UE_API void AddNewWeightMap();
+	UE_API void AddNewGroupsLayer();
+	UE_API void DeleteAttribute();
+	UE_API void ClearAttribute();
+	UE_API void SetLightmapUVsEnabled(bool bEnabled);
+	UE_API void ResetLightmapUVsChannels();
+
+	UE_API void AddNewAttribute(EAttributeEditorElementType ElemType, EAttributeEditorAttribType AttribType, FName AttributeName);
 
 	friend class FAttributeEditor_AttributeListsChange;
-	void EmitAttributesChange();
+	UE_API void EmitAttributesChange();
 };
 
 
-class MESHMODELINGTOOLSEDITORONLY_API FAttributeEditor_AttributeListsChange : public FToolCommandChange
+class FAttributeEditor_AttributeListsChange : public FToolCommandChange
 {
 public:
-	virtual void Apply(UObject* Object) override;
-	virtual void Revert(UObject* Object) override;
-	virtual FString ToString() const override;
+	UE_API virtual void Apply(UObject* Object) override;
+	UE_API virtual void Revert(UObject* Object) override;
+	UE_API virtual FString ToString() const override;
 };
+
+#undef UE_API
diff --git a/Engine/Plugins/Runtime/MeshModelingToolset/Source/MeshModelingToolsEditorOnly/Public/ParameterizeMeshTool.h b/Engine/Plugins/Runtime/MeshModelingToolset/Source/MeshModelingToolsEditorOnly/Public/ParameterizeMeshTool.h
--- a/Engine/Plugins/Runtime/MeshModelingToolset/Source/MeshModelingToolsEditorOnly/Public/ParameterizeMeshTool.h
+++ b/Engine/Plugins/Runtime/MeshModelingToolset/Source/MeshModelingToolsEditorOnly/Public/ParameterizeMeshTool.h
@@ -13,6 +13,8 @@
 
 #include "ParameterizeMeshTool.generated.h"
 
+#define UE_API MESHMODELINGTOOLSEDITORONLY_API
+
 // Forward declarations
 class UDynamicMeshComponent;
 class UMaterialInterface;
@@ -24,13 +26,13 @@ class UParameterizeMeshToolPatchBuilderProperties;
 class UParameterizeMeshOperatorFactory;
 
 
-UCLASS()
-class MESHMODELINGTOOLSEDITORONLY_API UParameterizeMeshToolBuilder : public USingleSelectionMeshEditingToolBuilder
+UCLASS(MinimalAPI)
+class UParameterizeMeshToolBuilder : public USingleSelectionMeshEditingToolBuilder
 {
 	GENERATED_BODY()
 public:
-	virtual USingleSelectionMeshEditingTool* CreateNewTool(const FToolBuilderState& SceneState) const override;
-	virtual bool CanBuildTool(const FToolBuilderState& SceneState) const override;
+	UE_API virtual USingleSelectionMeshEditingTool* CreateNewTool(const FToolBuilderState& SceneState) const override;
+	UE_API virtual bool CanBuildTool(const FToolBuilderState& SceneState) const override;
 };
 
 
@@ -38,23 +40,23 @@ public:
  * UParameterizeMeshTool automatically decomposes the input mesh into charts, solves for UVs,
  * and then packs the resulting charts
  */
-UCLASS()
-class MESHMODELINGTOOLSEDITORONLY_API UParameterizeMeshTool : public USingleSelectionMeshEditingTool, public IInteractiveToolManageGeometrySelectionAPI
+UCLASS(MinimalAPI)
+class UParameterizeMeshTool : public USingleSelectionMeshEditingTool, public IInteractiveToolManageGeometrySelectionAPI
 {
 	GENERATED_BODY()
 
 public:
-	virtual void Setup() override;
-	virtual void OnShutdown(EToolShutdownType ShutdownType) override;
+	UE_API virtual void Setup() override;
+	UE_API virtual void OnShutdown(EToolShutdownType ShutdownType) override;
 
-	virtual void Render(IToolsContextRenderAPI* RenderAPI) override;
-	virtual void OnTick(float DeltaTime) override;
+	UE_API virtual void Render(IToolsContextRenderAPI* RenderAPI) override;
+	UE_API virtual void OnTick(float DeltaTime) override;
 
 	virtual bool HasCancel() const override { return true; }
 	virtual bool HasAccept() const override { return true; }
-	virtual bool CanAccept() const override;
+	UE_API virtual bool CanAccept() const override;
 
-	virtual void OnPropertyModified(UObject* PropertySet, FProperty* Property) override;
+	UE_API virtual void OnPropertyModified(UObject* PropertySet, FProperty* Property) override;
 
 	// IInteractiveToolManageGeometrySelectionAPI -- this tool won't update external geometry selection or change selection-relevant mesh IDs
 	virtual bool IsInputSelectionValidOnOutput() override
@@ -98,7 +100,9 @@ protected:
 
 	TSharedPtr<UE::Geometry::FDynamicMesh3, ESPMode::ThreadSafe> InputMesh;
 
-	void OnMethodTypeChanged();
+	UE_API void OnMethodTypeChanged();
 
-	void OnPreviewMeshUpdated();
+	UE_API void OnPreviewMeshUpdated();
 };
+
+#undef UE_API
diff --git a/Engine/Plugins/Runtime/MeshModelingToolset/Source/MeshModelingToolsEditorOnly/Public/PolygonOnMeshTool.h b/Engine/Plugins/Runtime/MeshModelingToolset/Source/MeshModelingToolsEditorOnly/Public/PolygonOnMeshTool.h
--- a/Engine/Plugins/Runtime/MeshModelingToolset/Source/MeshModelingToolsEditorOnly/Public/PolygonOnMeshTool.h
+++ b/Engine/Plugins/Runtime/MeshModelingToolset/Source/MeshModelingToolsEditorOnly/Public/PolygonOnMeshTool.h
@@ -17,6 +17,8 @@
 #include "BaseTools/SingleSelectionMeshEditingTool.h"
 #include "PolygonOnMeshTool.generated.h"
 
+#define UE_API MESHMODELINGTOOLSEDITORONLY_API
+
 
 // predeclarations
 struct FMeshDescription;
@@ -30,13 +32,13 @@ class ULineSetComponent;
 /**
  *
  */
-UCLASS()
-class MESHMODELINGTOOLSEDITORONLY_API UPolygonOnMeshToolBuilder : public USingleSelectionMeshEditingToolBuilder
+UCLASS(MinimalAPI)
+class UPolygonOnMeshToolBuilder : public USingleSelectionMeshEditingToolBuilder
 {
 	GENERATED_BODY()
 
 public:
-	virtual USingleSelectionMeshEditingTool* CreateNewTool(const FToolBuilderState& SceneState) const override;
+	UE_API virtual USingleSelectionMeshEditingTool* CreateNewTool(const FToolBuilderState& SceneState) const override;
 };
 
 
@@ -57,8 +59,8 @@ enum class EPolygonType
 /**
  * Standard properties of the polygon-on-mesh operations
  */
-UCLASS()
-class MESHMODELINGTOOLSEDITORONLY_API UPolygonOnMeshToolProperties : public UInteractiveToolPropertySet
+UCLASS(MinimalAPI)
+class UPolygonOnMeshToolProperties : public UInteractiveToolPropertySet
 {
 	GENERATED_BODY()
 
@@ -132,8 +134,8 @@ enum class EPolygonOnMeshToolActions
 };
 
 
-UCLASS()
-class MESHMODELINGTOOLSEDITORONLY_API UPolygonOnMeshToolActionPropertySet : public UInteractiveToolPropertySet
+UCLASS(MinimalAPI)
+class UPolygonOnMeshToolActionPropertySet : public UInteractiveToolPropertySet
 {
 	GENERATED_BODY()
 
@@ -142,7 +144,7 @@ public:
 
 	void Initialize(UPolygonOnMeshTool* ParentToolIn) { ParentTool = ParentToolIn; }
 
-	void PostAction(EPolygonOnMeshToolActions Action);
+	UE_API void PostAction(EPolygonOnMeshToolActions Action);
 
 	/** Extrude the current set of selected faces. Click in viewport to confirm extrude height. */
 	UFUNCTION(CallInEditor, Category = Actions, meta = (DisplayName = "Draw Polygon", DisplayPriority = 1))
@@ -158,34 +160,34 @@ public:
 /**
  * Simple Mesh Plane Cutting Tool
  */
-UCLASS()
-class MESHMODELINGTOOLSEDITORONLY_API UPolygonOnMeshTool : public USingleSelectionMeshEditingTool, public UE::Geometry::IDynamicMeshOperatorFactory, public IClickBehaviorTarget, public IHoverBehaviorTarget
+UCLASS(MinimalAPI)
+class UPolygonOnMeshTool : public USingleSelectionMeshEditingTool, public UE::Geometry::IDynamicMeshOperatorFactory, public IClickBehaviorTarget, public IHoverBehaviorTarget
 {
 	GENERATED_BODY()
 
 public:
 
-	UPolygonOnMeshTool();
+	UE_API UPolygonOnMeshTool();
 
-	virtual void Setup() override;
-	virtual void OnShutdown(EToolShutdownType ShutdownType) override;
+	UE_API virtual void Setup() override;
+	UE_API virtual void OnShutdown(EToolShutdownType ShutdownType) override;
 
-	virtual void SetWorld(UWorld* World);
+	UE_API virtual void SetWorld(UWorld* World);
 
-	virtual void OnTick(float DeltaTime) override;
-	virtual void Render(IToolsContextRenderAPI* RenderAPI) override;
+	UE_API virtual void OnTick(float DeltaTime) override;
+	UE_API virtual void Render(IToolsContextRenderAPI* RenderAPI) override;
 
 	virtual bool HasCancel() const override { return true; }
 	virtual bool HasAccept() const override { return true; }
-	virtual bool CanAccept() const override;
+	UE_API virtual bool CanAccept() const override;
 
-	virtual void OnPropertyModified(UObject* PropertySet, FProperty* Property) override;
+	UE_API virtual void OnPropertyModified(UObject* PropertySet, FProperty* Property) override;
 
 	// IDynamicMeshOperatorFactory API
-	virtual TUniquePtr<UE::Geometry::FDynamicMeshOperator> MakeNewOperator() override;
+	UE_API virtual TUniquePtr<UE::Geometry::FDynamicMeshOperator> MakeNewOperator() override;
 
 
-	virtual void RequestAction(EPolygonOnMeshToolActions ActionType);
+	UE_API virtual void RequestAction(EPolygonOnMeshToolActions ActionType);
 
 
 
@@ -193,16 +195,16 @@ public:
 	// support for hover and click, for drawing polygon. This should be in the UCollectSurfacePathMechanic,
 	// but we don't support dynamically changing input behavior set yet
 
-	virtual bool HitTest(const FRay& Ray, FHitResult& OutHit);
+	UE_API virtual bool HitTest(const FRay& Ray, FHitResult& OutHit);
 
 	// IClickBehaviorTarget API
-	virtual FInputRayHit IsHitByClick(const FInputDeviceRay& ClickPos) override;
-	virtual void OnClicked(const FInputDeviceRay& ClickPos) override;
+	UE_API virtual FInputRayHit IsHitByClick(const FInputDeviceRay& ClickPos) override;
+	UE_API virtual void OnClicked(const FInputDeviceRay& ClickPos) override;
 
 	// IHoverBehaviorTarget API
-	virtual FInputRayHit BeginHoverSequenceHitTest(const FInputDeviceRay& PressPos) override;
+	UE_API virtual FInputRayHit BeginHoverSequenceHitTest(const FInputDeviceRay& PressPos) override;
 	virtual void OnBeginHover(const FInputDeviceRay& DevicePos) override {}
-	virtual bool OnUpdateHover(const FInputDeviceRay& DevicePos) override;
+	UE_API virtual bool OnUpdateHover(const FInputDeviceRay& DevicePos) override;
 	virtual void OnEndHover() override {}
 
 
@@ -244,15 +246,17 @@ protected:
 
 	UE::Geometry::FPolygon2d LastDrawnPolygon;
 	UE::Geometry::FPolygon2d ActivePolygon;
-	void UpdatePolygonType();
+	UE_API void UpdatePolygonType();
 
-	void SetupPreview();
-	void UpdateDrawPlane();
+	UE_API void SetupPreview();
+	UE_API void UpdateDrawPlane();
 
-	void BeginDrawPolygon();
-	void CompleteDrawPolygon();
+	UE_API void BeginDrawPolygon();
+	UE_API void CompleteDrawPolygon();
 
-	void UpdateVisualization();
+	UE_API void UpdateVisualization();
 
-	void GenerateAsset(const TArray<FDynamicMeshOpResult>& Results);
+	UE_API void GenerateAsset(const TArray<FDynamicMeshOpResult>& Results);
 };
+
+#undef UE_API
diff --git a/Engine/Plugins/Runtime/MeshModelingToolset/Source/MeshModelingToolsEditorOnly/Public/SimplifyMeshTool.h b/Engine/Plugins/Runtime/MeshModelingToolset/Source/MeshModelingToolsEditorOnly/Public/SimplifyMeshTool.h
--- a/Engine/Plugins/Runtime/MeshModelingToolset/Source/MeshModelingToolsEditorOnly/Public/SimplifyMeshTool.h
+++ b/Engine/Plugins/Runtime/MeshModelingToolset/Source/MeshModelingToolsEditorOnly/Public/SimplifyMeshTool.h
@@ -10,6 +10,8 @@
 #include "CleaningOps/SimplifyMeshOp.h"		// required in header for enum types
 #include "SimplifyMeshTool.generated.h"
 
+#define UE_API MESHMODELINGTOOLSEDITORONLY_API
+
 class UMeshStatisticsProperties;
 class UMeshElementsVisualizer;
 PREDECLARE_GEOMETRY(class FDynamicMesh3);
@@ -18,13 +20,13 @@ PREDECLARE_GEOMETRY(typedef TMeshAABBTree3<FDynamicMesh3> FDynamicMeshAABBTree3)
 /**
  *
  */
-UCLASS()
-class MESHMODELINGTOOLSEDITORONLY_API USimplifyMeshToolBuilder : public USingleSelectionMeshEditingToolBuilder
+UCLASS(MinimalAPI)
+class USimplifyMeshToolBuilder : public USingleSelectionMeshEditingToolBuilder
 {
 	GENERATED_BODY()
 
 public:
-	virtual USingleSelectionMeshEditingTool* CreateNewTool(const FToolBuilderState& SceneState) const override;
+	UE_API virtual USingleSelectionMeshEditingTool* CreateNewTool(const FToolBuilderState& SceneState) const override;
 };
 
 
@@ -32,12 +34,12 @@ public:
 /**
  * Standard properties of the Simplify operation
  */
-UCLASS()
-class MESHMODELINGTOOLSEDITORONLY_API USimplifyMeshToolProperties : public UMeshConstraintProperties
+UCLASS(MinimalAPI)
+class USimplifyMeshToolProperties : public UMeshConstraintProperties
 {
 	GENERATED_BODY()
 public:
-	USimplifyMeshToolProperties();
+	UE_API USimplifyMeshToolProperties();
 
 	/** Simplification Scheme  */
 	UPROPERTY(EditAnywhere, Category = Options)
@@ -127,25 +129,25 @@ protected:
 /**
  * Simple Mesh Simplifying Tool
  */
-UCLASS()
-class MESHMODELINGTOOLSEDITORONLY_API USimplifyMeshTool : public USingleSelectionMeshEditingTool, public UE::Geometry::IDynamicMeshOperatorFactory
+UCLASS(MinimalAPI)
+class USimplifyMeshTool : public USingleSelectionMeshEditingTool, public UE::Geometry::IDynamicMeshOperatorFactory
 {
 	GENERATED_BODY()
 
 public:
-	virtual void Setup() override;
-	virtual void OnShutdown(EToolShutdownType ShutdownType) override;
+	UE_API virtual void Setup() override;
+	UE_API virtual void OnShutdown(EToolShutdownType ShutdownType) override;
 
-	virtual void OnTick(float DeltaTime) override;
+	UE_API virtual void OnTick(float DeltaTime) override;
 
 	virtual bool HasCancel() const override { return true; }
 	virtual bool HasAccept() const override { return true; }
-	virtual bool CanAccept() const override;
+	UE_API virtual bool CanAccept() const override;
 
-	virtual void OnPropertyModified(UObject* PropertySet, FProperty* Property) override;
+	UE_API virtual void OnPropertyModified(UObject* PropertySet, FProperty* Property) override;
 
 	// IDynamicMeshOperatorFactory API
-	virtual TUniquePtr<UE::Geometry::FDynamicMeshOperator> MakeNewOperator() override;
+	UE_API virtual TUniquePtr<UE::Geometry::FDynamicMeshOperator> MakeNewOperator() override;
 
 private:
 	UPROPERTY()
@@ -165,5 +167,7 @@ private:
 	TSharedPtr<UE::Geometry::FDynamicMesh3, ESPMode::ThreadSafe> OriginalMesh;
 	TSharedPtr<UE::Geometry::FDynamicMeshAABBTree3, ESPMode::ThreadSafe> OriginalMeshSpatial;
 
-	void UpdateVisualization();
+	UE_API void UpdateVisualization();
 };
+
+#undef UE_API
diff --git a/Engine/Plugins/Runtime/MeshModelingToolset/Source/MeshModelingToolsEditorOnly/Public/SkeletalMesh/SkeletalMeshEditionInterface.h b/Engine/Plugins/Runtime/MeshModelingToolset/Source/MeshModelingToolsEditorOnly/Public/SkeletalMesh/SkeletalMeshEditionInterface.h
--- a/Engine/Plugins/Runtime/MeshModelingToolset/Source/MeshModelingToolsEditorOnly/Public/SkeletalMesh/SkeletalMeshEditionInterface.h
+++ b/Engine/Plugins/Runtime/MeshModelingToolset/Source/MeshModelingToolsEditorOnly/Public/SkeletalMesh/SkeletalMeshEditionInterface.h
@@ -12,6 +12,8 @@
 
 #include "SkeletalMeshEditionInterface.generated.h"
 
+#define UE_API MESHMODELINGTOOLSEDITORONLY_API
+
 class FSkeletalMeshToolNotifier;
 class UTransformProxy;
 class UInteractiveToolManager;
@@ -23,8 +25,8 @@ struct FInputDeviceRay;
  * USkeletalMeshEditionInterface
  */
 
-UINTERFACE()
-class MESHMODELINGTOOLSEDITORONLY_API USkeletalMeshEditingInterface : public UInterface
+UINTERFACE(MinimalAPI)
+class USkeletalMeshEditingInterface : public UInterface
 {
 	GENERATED_BODY()
 };
@@ -33,25 +35,25 @@ class MESHMODELINGTOOLSEDITORONLY_API USkeletalMeshEditingInterface : public UIn
  * ISkeletalMeshEditionInterface
  */
 
-class MESHMODELINGTOOLSEDITORONLY_API ISkeletalMeshEditingInterface
+class ISkeletalMeshEditingInterface
 {
 	GENERATED_BODY()
 
 public:
-	ISkeletalMeshNotifier& GetNotifier();
-	bool NeedsNotification() const;
+	UE_API ISkeletalMeshNotifier& GetNotifier();
+	UE_API bool NeedsNotification() const;
 
-	virtual TArray<FName> GetSelectedBones() const;
-	virtual TArray<FName> GetSelectedMorphTargets() const;
+	UE_API virtual TArray<FName> GetSelectedBones() const;
+	UE_API virtual TArray<FName> GetSelectedMorphTargets() const;
 	
-	void BindTo(TSharedPtr<ISkeletalMeshEditorBinding> InBinding);
-	void Unbind();
+	UE_API void BindTo(TSharedPtr<ISkeletalMeshEditorBinding> InBinding);
+	UE_API void Unbind();
 
-	virtual TWeakObjectPtr<USkeletonModifier> GetModifier() const;
+	UE_API virtual TWeakObjectPtr<USkeletonModifier> GetModifier() const;
 	
 protected:
 	virtual void HandleSkeletalMeshModified(const TArray<FName>& BoneNames, const ESkeletalMeshNotifyType InNotifyType) = 0;
-	TOptional<FName> GetBoneName(HHitProxy* InHitProxy) const;
+	UE_API TOptional<FName> GetBoneName(HHitProxy* InHitProxy) const;
 	
 	TWeakPtr<ISkeletalMeshEditorBinding> Binding;
 
@@ -81,7 +83,7 @@ protected:
 
 struct HBoneHitProxy : public HHitProxy
 {
-	DECLARE_HIT_PROXY()
+	DECLARE_HIT_PROXY( UE_API )
 
 	int32 BoneIndex;
 	FName BoneName;
@@ -99,8 +101,8 @@ struct HBoneHitProxy : public HHitProxy
  * USkeletalMeshGizmoContextObjectBase
  */
 
-UCLASS(Abstract)
-class MESHMODELINGTOOLSEDITORONLY_API USkeletalMeshGizmoContextObjectBase : public UObject
+UCLASS(MinimalAPI, Abstract)
+class USkeletalMeshGizmoContextObjectBase : public UObject
 {
 	GENERATED_BODY()
 	
@@ -113,8 +115,8 @@ public:
  * USkeletalMeshGizmoWrapperBase
  */
 
-UCLASS(Abstract)
-class MESHMODELINGTOOLSEDITORONLY_API USkeletalMeshGizmoWrapperBase : public UObject
+UCLASS(MinimalAPI, Abstract)
+class USkeletalMeshGizmoWrapperBase : public UObject
 {
 	GENERATED_BODY()
 
@@ -144,8 +146,8 @@ public:
  * USkeletalMeshEditorContextObjectBase
  */
 
-UCLASS(Abstract)
-class MESHMODELINGTOOLSEDITORONLY_API USkeletalMeshEditorContextObjectBase : public UObject
+UCLASS(MinimalAPI, Abstract)
+class USkeletalMeshEditorContextObjectBase : public UObject
 {
 	GENERATED_BODY()
 
@@ -155,4 +157,6 @@ public:
 	
 	virtual void BindTo(ISkeletalMeshEditingInterface* InEditingInterface) PURE_VIRTUAL(USkeletalMeshEditorContextObjectBase::BindTo, return;)
 	virtual void UnbindFrom(ISkeletalMeshEditingInterface* InEditingInterface) PURE_VIRTUAL(USkeletalMeshEditorContextObjectBase::UnbindFrom, return;)
-};
\ No newline at end of file
+};
+
+#undef UE_API
diff --git a/Engine/Plugins/Runtime/MeshModelingToolset/Source/MeshModelingToolsEditorOnly/Public/SkeletalMesh/SkeletonEditingTool.h b/Engine/Plugins/Runtime/MeshModelingToolset/Source/MeshModelingToolsEditorOnly/Public/SkeletalMesh/SkeletonEditingTool.h
--- a/Engine/Plugins/Runtime/MeshModelingToolset/Source/MeshModelingToolsEditorOnly/Public/SkeletalMesh/SkeletonEditingTool.h
+++ b/Engine/Plugins/Runtime/MeshModelingToolset/Source/MeshModelingToolsEditorOnly/Public/SkeletalMesh/SkeletonEditingTool.h
@@ -17,6 +17,8 @@
 
 #include "SkeletonEditingTool.generated.h"
 
+#define UE_API MESHMODELINGTOOLSEDITORONLY_API
+
 class USingleClickInputBehavior;
 class UGizmoViewContext;
 class UTransformGizmo;
@@ -32,20 +34,20 @@ namespace SkeletonEditingTool
 /**
  * A wrapper change class that stores a reference skeleton and the bones' indexes trackers to be used for undo/redo.
  */
-class MESHMODELINGTOOLSEDITORONLY_API FRefSkeletonChange : public FToolCommandChange
+class FRefSkeletonChange : public FToolCommandChange
 {
 public:
-	FRefSkeletonChange(const USkeletonEditingTool* InTool);
+	UE_API FRefSkeletonChange(const USkeletonEditingTool* InTool);
 
 	virtual FString ToString() const override
 	{
 		return FString(TEXT("Edit Skeleton"));
 	}
 
-	virtual void Apply(UObject* Object) override;
-	virtual void Revert(UObject* Object) override;
+	UE_API virtual void Apply(UObject* Object) override;
+	UE_API virtual void Revert(UObject* Object) override;
 
-	void StoreSkeleton(const USkeletonEditingTool* InTool);
+	UE_API void StoreSkeleton(const USkeletonEditingTool* InTool);
 
 private:
 	FReferenceSkeleton PreChangeSkeleton;
@@ -60,19 +62,19 @@ private:
  * USkeletonEditingToolBuilder
  */
 
-UCLASS()
-class MESHMODELINGTOOLSEDITORONLY_API USkeletonEditingToolBuilder : public UInteractiveToolWithToolTargetsBuilder
+UCLASS(MinimalAPI)
+class USkeletonEditingToolBuilder : public UInteractiveToolWithToolTargetsBuilder
 {
 	GENERATED_BODY()
 
 public:
 	// UInteractiveToolBuilder overrides
-	virtual bool CanBuildTool(const FToolBuilderState& SceneState) const override;
-	virtual UInteractiveTool* BuildTool(const FToolBuilderState& SceneState) const override;
+	UE_API virtual bool CanBuildTool(const FToolBuilderState& SceneState) const override;
+	UE_API virtual UInteractiveTool* BuildTool(const FToolBuilderState& SceneState) const override;
 	
 protected:
 	// UInteractiveToolWithToolTargetsBuilder overrides
-	virtual const FToolTargetTypeRequirements& GetTargetRequirements() const override;
+	UE_API virtual const FToolTargetTypeRequirements& GetTargetRequirements() const override;
 };
 
 /**
@@ -96,8 +98,8 @@ enum class EEditingOperation : uint8
  * Changed are actually commit to the SkeletalMesh and it's mesh description on Accept.
  */
 
-UCLASS()
-class MESHMODELINGTOOLSEDITORONLY_API USkeletonEditingTool :
+UCLASS(MinimalAPI)
+class USkeletonEditingTool :
 	public USingleSelectionTool,
 	public IClickDragBehaviorTarget,
 	public ISkeletalMeshEditingInterface
@@ -106,73 +108,73 @@ class MESHMODELINGTOOLSEDITORONLY_API USkeletonEditingTool :
 
 public:
 
-	void Init(const FToolBuilderState& InSceneState);
+	UE_API void Init(const FToolBuilderState& InSceneState);
 	
 	// UInteractiveTool overrides
-	virtual void Setup() override;
-	virtual void Shutdown(EToolShutdownType ShutdownType) override;
+	UE_API virtual void Setup() override;
+	UE_API virtual void Shutdown(EToolShutdownType ShutdownType) override;
 
-	virtual void DrawHUD(FCanvas* Canvas, IToolsContextRenderAPI* RenderAPI) override;
-	virtual void Render(IToolsContextRenderAPI* RenderAPI) override;
+	UE_API virtual void DrawHUD(FCanvas* Canvas, IToolsContextRenderAPI* RenderAPI) override;
+	UE_API virtual void Render(IToolsContextRenderAPI* RenderAPI) override;
 	
 	virtual bool HasCancel() const override { return true; }
 	virtual bool HasAccept() const override { return true; }
 	virtual bool CanAccept() const override { return true; }
 
-	virtual void OnPropertyModified(UObject* PropertySet, FProperty* Property) override;
+	UE_API virtual void OnPropertyModified(UObject* PropertySet, FProperty* Property) override;
 
 	// ICLickDragBehaviorTarget overrides
-	virtual FInputRayHit CanBeginClickDragSequence(const FInputDeviceRay& InPressPos) override;
+	UE_API virtual FInputRayHit CanBeginClickDragSequence(const FInputDeviceRay& InPressPos) override;
 	
-	virtual void OnClickPress(const FInputDeviceRay& InPressPos) override;
-	virtual void OnClickDrag(const FInputDeviceRay& InDragPos) override;
-	virtual void OnClickRelease(const FInputDeviceRay& InReleasePos) override;
-	virtual void OnTerminateDragSequence() override;
+	UE_API virtual void OnClickPress(const FInputDeviceRay& InPressPos) override;
+	UE_API virtual void OnClickDrag(const FInputDeviceRay& InDragPos) override;
+	UE_API virtual void OnClickRelease(const FInputDeviceRay& InReleasePos) override;
+	UE_API virtual void OnTerminateDragSequence() override;
 	
-	virtual void OnTick(float DeltaTime) override;
+	UE_API virtual void OnTick(float DeltaTime) override;
 
 	// IInteractiveToolCameraFocusAPI overrides
-	virtual FBox GetWorldSpaceFocusBox() override;
+	UE_API virtual FBox GetWorldSpaceFocusBox() override;
 
 	// Modifier functions
-	void MirrorBones();
-	void RenameBones();
-	void MoveBones();
-	void OrientBones();
-	void RemoveBones();
-	void UnParentBones();
-	void SnapBoneToComponentSelection(const bool bCreate);
+	UE_API void MirrorBones();
+	UE_API void RenameBones();
+	UE_API void MoveBones();
+	UE_API void OrientBones();
+	UE_API void RemoveBones();
+	UE_API void UnParentBones();
+	UE_API void SnapBoneToComponentSelection(const bool bCreate);
 	
 	// ISkeletalMeshEditionInterface overrides
-	virtual TArray<FName> GetSelectedBones() const override;
+	UE_API virtual TArray<FName> GetSelectedBones() const override;
 
-	const TArray<FName>& GetSelection() const;
+	UE_API const TArray<FName>& GetSelection() const;
 
-	const FTransform& GetTransform(const FName InBoneName, const bool bWorld) const;
-	void SetTransforms(const TArray<FName>& InBones, const TArray<FTransform>& InTransforms, const bool bWorld) const;
+	UE_API const FTransform& GetTransform(const FName InBoneName, const bool bWorld) const;
+	UE_API void SetTransforms(const TArray<FName>& InBones, const TArray<FTransform>& InTransforms, const bool bWorld) const;
 
 	// IModifierToggleBehaviorTarget overrides
-	virtual void OnUpdateModifierState(int ModifierID, bool bIsOn) override;
+	UE_API virtual void OnUpdateModifierState(int ModifierID, bool bIsOn) override;
 
 	// ISkeletalMeshEditionInterface overrides
-	virtual TWeakObjectPtr<USkeletonModifier> GetModifier() const override;
+	UE_API virtual TWeakObjectPtr<USkeletonModifier> GetModifier() const override;
 
-	EEditingOperation GetOperation() const;
-	void SetOperation(const EEditingOperation InOperation, const bool bUpdateGizmo = true);
+	UE_API EEditingOperation GetOperation() const;
+	UE_API void SetOperation(const EEditingOperation InOperation, const bool bUpdateGizmo = true);
 
-	bool HasSelectedComponent() const;
+	UE_API bool HasSelectedComponent() const;
 	
 protected:
 
 	// UInteractiveTool overrides
-	virtual void RegisterActions(FInteractiveToolActionSet& ActionSet) override;
+	UE_API virtual void RegisterActions(FInteractiveToolActionSet& ActionSet) override;
 
 	// ISkeletalMeshEditionInterface overrides
-	virtual void HandleSkeletalMeshModified(const TArray<FName>& BoneNames, const ESkeletalMeshNotifyType InNotifyType) override;
+	UE_API virtual void HandleSkeletalMeshModified(const TArray<FName>& BoneNames, const ESkeletalMeshNotifyType InNotifyType) override;
 
 	// Modifier functions
-	void CreateNewBone();
-	void ParentBones(const FName& InParentName);
+	UE_API void CreateNewBone();
+	UE_API void ParentBones(const FName& InParentName);
 
 public:
 	
@@ -221,18 +223,18 @@ protected:
 	UPROPERTY()
 	TWeakObjectPtr<USkeletalMeshEditorContextObjectBase> EditorContext = nullptr;
 
-	void UpdateGizmo() const;
+	UE_API void UpdateGizmo() const;
 	
 	// ref skeleton transactions
-	void BeginChange();
-	void EndChange();
-	void CancelChange();
+	UE_API void BeginChange();
+	UE_API void EndChange();
+	UE_API void CancelChange();
 	TUniquePtr<SkeletonEditingTool::FRefSkeletonChange> ActiveChange;
 
 	friend class SkeletonEditingTool::FRefSkeletonChange;
 	
 private:
-	TArray<int32> GetSelectedBoneIndexes() const;
+	UE_API TArray<int32> GetSelectedBoneIndexes() const;
 	
 	enum class EBoneSelectionMode : uint8
 	{
@@ -250,36 +252,36 @@ private:
 
 	TArray<FName> Selection;
 	
-	void SelectBone(const FName& InBoneName);
-	FName GetCurrentBone() const;
+	UE_API void SelectBone(const FName& InBoneName);
+	UE_API FName GetCurrentBone() const;
 
-	void NormalizeSelection();
+	UE_API void NormalizeSelection();
 
 	// setup
-	void SetupModifier(USkeletalMesh* InSkeletalMesh);
-	void SetupPreviewMesh();
-	void SetupProperties();
-	void SetupBehaviors();
-	void SetupGizmo(USkeletalMeshComponent* InComponent);
-	void SetupWatchers();
-	void SetupComponentsSelection();
+	UE_API void SetupModifier(USkeletalMesh* InSkeletalMesh);
+	UE_API void SetupPreviewMesh();
+	UE_API void SetupProperties();
+	UE_API void SetupBehaviors();
+	UE_API void SetupGizmo(USkeletalMeshComponent* InComponent);
+	UE_API void SetupWatchers();
+	UE_API void SetupComponentsSelection();
 	
 	// actions
-	void RegisterCreateAction(FInteractiveToolActionSet& InOutActionSet, const int32 InActionId);
-	void RegisterDeleteAction(FInteractiveToolActionSet& InOutActionSet, const int32 InActionId);
-	void RegisterSelectAction(FInteractiveToolActionSet& InOutActionSet, const int32 InActionId);
-	void RegisterParentAction(FInteractiveToolActionSet& InOutActionSet, const int32 InActionId);
-	void RegisterUnParentAction(FInteractiveToolActionSet& InOutActionSet, const int32 InActionId);
-	void RegisterCopyAction(FInteractiveToolActionSet& InOutActionSet, const int32 InActionId);
-	void RegisterPasteAction(FInteractiveToolActionSet& InOutActionSet, const int32 InActionId);
-	void RegisterDuplicateAction(FInteractiveToolActionSet& InOutActionSet, const int32 InActionId);
-	void RegisterSelectComponentsAction(FInteractiveToolActionSet& InOutActionSet, const int32 InActionId);
-	void RegisterSelectionFilterCyclingAction(FInteractiveToolActionSet& InOutActionSet, const int32 InActionId);
-	void RegisterSnapAction(FInteractiveToolActionSet& InOutActionSet, const int32 InActionId);
-
-	TArray<int32> GetSelectedComponents() const;
-
-	FTransform ComputeTransformFromComponents(const TArray<int32>& InIDs) const;
+	UE_API void RegisterCreateAction(FInteractiveToolActionSet& InOutActionSet, const int32 InActionId);
+	UE_API void RegisterDeleteAction(FInteractiveToolActionSet& InOutActionSet, const int32 InActionId);
+	UE_API void RegisterSelectAction(FInteractiveToolActionSet& InOutActionSet, const int32 InActionId);
+	UE_API void RegisterParentAction(FInteractiveToolActionSet& InOutActionSet, const int32 InActionId);
+	UE_API void RegisterUnParentAction(FInteractiveToolActionSet& InOutActionSet, const int32 InActionId);
+	UE_API void RegisterCopyAction(FInteractiveToolActionSet& InOutActionSet, const int32 InActionId);
+	UE_API void RegisterPasteAction(FInteractiveToolActionSet& InOutActionSet, const int32 InActionId);
+	UE_API void RegisterDuplicateAction(FInteractiveToolActionSet& InOutActionSet, const int32 InActionId);
+	UE_API void RegisterSelectComponentsAction(FInteractiveToolActionSet& InOutActionSet, const int32 InActionId);
+	UE_API void RegisterSelectionFilterCyclingAction(FInteractiveToolActionSet& InOutActionSet, const int32 InActionId);
+	UE_API void RegisterSnapAction(FInteractiveToolActionSet& InOutActionSet, const int32 InActionId);
+
+	UE_API TArray<int32> GetSelectedComponents() const;
+
+	UE_API FTransform ComputeTransformFromComponents(const TArray<int32>& InIDs) const;
 	
 	TValueWatcher<TArray<FName>> SelectionWatcher;
 	TValueWatcher<EToolContextCoordinateSystem> CoordinateSystemWatcher; 
@@ -302,12 +304,12 @@ ENUM_CLASS_FLAGS(USkeletonEditingTool::EBoneSelectionMode);
  * USkeletonEditingProperties
  */
 
-UCLASS()
-class MESHMODELINGTOOLSEDITORONLY_API USkeletonEditingProperties : public UInteractiveToolPropertySet
+UCLASS(MinimalAPI)
+class USkeletonEditingProperties : public UInteractiveToolPropertySet
 {
 	GENERATED_BODY()
 public:
-	void Initialize(USkeletonEditingTool* ParentToolIn);
+	UE_API void Initialize(USkeletonEditingTool* ParentToolIn);
 
 	UPROPERTY(EditAnywhere, Category = "Details")
 	FName Name;
@@ -328,7 +330,7 @@ public:
 	bool bEnableComponentSelection = false;
 
 #if WITH_EDITOR
-	virtual void PostEditChangeProperty(FPropertyChangedEvent &PropertyChangedEvent) override;
+	UE_API virtual void PostEditChangeProperty(FPropertyChangedEvent &PropertyChangedEvent) override;
 #endif
 	
 	TWeakObjectPtr<USkeletonEditingTool> ParentTool = nullptr;
@@ -350,16 +352,16 @@ enum class EProjectionType : uint8
  * UProjectionProperties
  */
 
-UCLASS()
-class MESHMODELINGTOOLSEDITORONLY_API UProjectionProperties: public UInteractiveToolPropertySet
+UCLASS(MinimalAPI)
+class UProjectionProperties: public UInteractiveToolPropertySet
 {
 	GENERATED_BODY()
 
 public:
-	void Initialize(USkeletonEditingTool* ParentToolIn, TObjectPtr<UPreviewMesh> PreviewMesh);
+	UE_API void Initialize(USkeletonEditingTool* ParentToolIn, TObjectPtr<UPreviewMesh> PreviewMesh);
 
-	void UpdatePlane(const UGizmoViewContext& InViewContext, const FVector& InOrigin);
-	bool GetProjectionPoint(const FInputDeviceRay& InRay, FVector& OutHitPoint) const;
+	UE_API void UpdatePlane(const UGizmoViewContext& InViewContext, const FVector& InOrigin);
+	UE_API bool GetProjectionPoint(const FInputDeviceRay& InRay, FVector& OutHitPoint) const;
 	
 	UPROPERTY(EditAnywhere, Category = "Project")
 	EProjectionType ProjectionType = EProjectionType::WithinMesh;
@@ -382,15 +384,15 @@ private:
  * UMirroringProperties
  */
 
-UCLASS()
-class MESHMODELINGTOOLSEDITORONLY_API UMirroringProperties: public UInteractiveToolPropertySet
+UCLASS(MinimalAPI)
+class UMirroringProperties: public UInteractiveToolPropertySet
 {
 	GENERATED_BODY()
 
 public:
-	void Initialize(USkeletonEditingTool* ParentToolIn);
+	UE_API void Initialize(USkeletonEditingTool* ParentToolIn);
 
-	void MirrorBones();
+	UE_API void MirrorBones();
 
 	UPROPERTY(EditAnywhere, Category = "Mirror")
 	FMirrorOptions Options;
@@ -402,15 +404,15 @@ public:
  * UOrientingProperties
  */
 
-UCLASS()
-class MESHMODELINGTOOLSEDITORONLY_API UOrientingProperties: public UInteractiveToolPropertySet
+UCLASS(MinimalAPI)
+class UOrientingProperties: public UInteractiveToolPropertySet
 {
 	GENERATED_BODY()
 
 public:
-	void Initialize(USkeletonEditingTool* ParentToolIn);
+	UE_API void Initialize(USkeletonEditingTool* ParentToolIn);
 
-	void OrientBones();
+	UE_API void OrientBones();
 
 	UPROPERTY(EditAnywhere, Category = "Orient")
 	bool bAutoOrient = false;
@@ -419,8 +421,10 @@ public:
 	FOrientOptions Options;
 
 #if WITH_EDITOR
-	virtual void PostEditChangeProperty(FPropertyChangedEvent &PropertyChangedEvent) override;
+	UE_API virtual void PostEditChangeProperty(FPropertyChangedEvent &PropertyChangedEvent) override;
 #endif
 	
 	TWeakObjectPtr<USkeletonEditingTool> ParentTool = nullptr;
-};
\ No newline at end of file
+};
+
+#undef UE_API
diff --git a/Engine/Plugins/Runtime/MeshModelingToolset/Source/MeshModelingToolsEditorOnly/Public/SkeletalMesh/SkeletonTransformProxy.h b/Engine/Plugins/Runtime/MeshModelingToolset/Source/MeshModelingToolsEditorOnly/Public/SkeletalMesh/SkeletonTransformProxy.h
--- a/Engine/Plugins/Runtime/MeshModelingToolset/Source/MeshModelingToolsEditorOnly/Public/SkeletalMesh/SkeletonTransformProxy.h
+++ b/Engine/Plugins/Runtime/MeshModelingToolset/Source/MeshModelingToolsEditorOnly/Public/SkeletalMesh/SkeletonTransformProxy.h
@@ -6,32 +6,35 @@
 
 #include "SkeletonTransformProxy.generated.h"
 
+#define UE_API MESHMODELINGTOOLSEDITORONLY_API
+
 class UObject;
 
 /**
  * USkeletonTransformProxy is a derivation of UTransformProxy that manages several bones and update their transform individually.
  */
 
-UCLASS(Transient)
-class MESHMODELINGTOOLSEDITORONLY_API USkeletonTransformProxy : public UTransformProxy
+UCLASS(MinimalAPI, Transient)
+class USkeletonTransformProxy : public UTransformProxy
 {
 	GENERATED_BODY()
 
 public:
 
-	void Initialize(const FTransform& InTransform, const EToolContextCoordinateSystem& InTransformMode);
+	UE_API void Initialize(const FTransform& InTransform, const EToolContextCoordinateSystem& InTransformMode);
 	
-	bool IsValid() const;
+	UE_API bool IsValid() const;
 
 	// UTransformProxy overrides
-	virtual FTransform GetTransform() const override;
+	UE_API virtual FTransform GetTransform() const override;
 
 	UPROPERTY()
 	EToolContextCoordinateSystem TransformMode = EToolContextCoordinateSystem::Local;
 	
 protected:
 	// UTransformProxy overrides
-	virtual void UpdateSharedTransform() override;
-	virtual void UpdateObjects() override;
+	UE_API virtual void UpdateSharedTransform() override;
+	UE_API virtual void UpdateObjects() override;
 };
 
+#undef UE_API
diff --git a/Engine/Plugins/Runtime/MeshModelingToolset/Source/MeshModelingToolsEditorOnly/Public/SkeletalMesh/SkinWeightsBindingTool.h b/Engine/Plugins/Runtime/MeshModelingToolset/Source/MeshModelingToolsEditorOnly/Public/SkeletalMesh/SkinWeightsBindingTool.h
--- a/Engine/Plugins/Runtime/MeshModelingToolset/Source/MeshModelingToolsEditorOnly/Public/SkeletalMesh/SkinWeightsBindingTool.h
+++ b/Engine/Plugins/Runtime/MeshModelingToolset/Source/MeshModelingToolsEditorOnly/Public/SkeletalMesh/SkinWeightsBindingTool.h
@@ -13,6 +13,8 @@
 
 #include "SkinWeightsBindingTool.generated.h"
 
+#define UE_API MESHMODELINGTOOLSEDITORONLY_API
+
 class UMeshOpPreviewWithBackgroundCompute;
 struct FDynamicMeshOpResult;
 struct FMeshDescription;
@@ -27,16 +29,16 @@ namespace UE::Geometry
 /**
  *
  */
-UCLASS()
-class MESHMODELINGTOOLSEDITORONLY_API USkinWeightsBindingToolBuilder : public UMultiSelectionMeshEditingToolBuilder
+UCLASS(MinimalAPI)
+class USkinWeightsBindingToolBuilder : public UMultiSelectionMeshEditingToolBuilder
 {
 	GENERATED_BODY()
 public:
-	virtual bool CanBuildTool(const FToolBuilderState& SceneState) const override;
-	virtual UMultiSelectionMeshEditingTool* CreateNewTool(const FToolBuilderState& SceneState) const override;
+	UE_API virtual bool CanBuildTool(const FToolBuilderState& SceneState) const override;
+	UE_API virtual UMultiSelectionMeshEditingTool* CreateNewTool(const FToolBuilderState& SceneState) const override;
 
 protected:
-	virtual const FToolTargetTypeRequirements& GetTargetRequirements() const override;
+	UE_API virtual const FToolTargetTypeRequirements& GetTargetRequirements() const override;
 };
 
 // A mirror of UE::Geometry::ESkinBindingType
@@ -47,8 +49,8 @@ enum class ESkinWeightsBindType : uint8
 	GeodesicVoxel = 1 UMETA(DisplayName = "Geodesic Voxel"),
 };
 
-UCLASS()
-class MESHMODELINGTOOLSEDITORONLY_API USkinWeightsBindingToolProperties : 
+UCLASS(MinimalAPI)
+class USkinWeightsBindingToolProperties : 
 	public UInteractiveToolPropertySet
 {
 	GENERATED_BODY()
@@ -79,8 +81,8 @@ public:
 /**
  *
  */
-UCLASS()
-class MESHMODELINGTOOLSEDITORONLY_API USkinWeightsBindingTool :
+UCLASS(MinimalAPI)
+class USkinWeightsBindingTool :
 	public UMultiSelectionMeshEditingTool,
 	public UE::Geometry::IDynamicMeshOperatorFactory,
 	public ISkeletalMeshEditingInterface
@@ -88,23 +90,23 @@ class MESHMODELINGTOOLSEDITORONLY_API USkinWeightsBindingTool :
 	GENERATED_BODY()
 
 public:
-	USkinWeightsBindingTool();
-	virtual ~USkinWeightsBindingTool() override;
+	UE_API USkinWeightsBindingTool();
+	UE_API virtual ~USkinWeightsBindingTool() override;
 
-	void Init(const FToolBuilderState& InSceneState);
+	UE_API void Init(const FToolBuilderState& InSceneState);
 	
-	virtual void Setup() override;
-	virtual void OnShutdown(EToolShutdownType ShutdownType) override;
+	UE_API virtual void Setup() override;
+	UE_API virtual void OnShutdown(EToolShutdownType ShutdownType) override;
 	
-	virtual void OnTick(float DeltaTime) override;
-	virtual void Render(IToolsContextRenderAPI* RenderAPI) override;
+	UE_API virtual void OnTick(float DeltaTime) override;
+	UE_API virtual void Render(IToolsContextRenderAPI* RenderAPI) override;
 
 	virtual bool HasCancel() const override { return true; }
 	virtual bool HasAccept() const override { return true; }
-	virtual bool CanAccept() const override;
+	UE_API virtual bool CanAccept() const override;
 
 	// IDynamicMeshOperatorFactory API
-	virtual TUniquePtr<UE::Geometry::FDynamicMeshOperator> MakeNewOperator() override;
+	UE_API virtual TUniquePtr<UE::Geometry::FDynamicMeshOperator> MakeNewOperator() override;
 
 	UPROPERTY()
 	TObjectPtr<USkinWeightsBindingToolProperties> Properties = nullptr;
@@ -121,20 +123,22 @@ protected:
 
 	FReferenceSkeleton ReferenceSkeleton;
 	
-	void GenerateAsset(const FDynamicMeshOpResult& Result);
+	UE_API void GenerateAsset(const FDynamicMeshOpResult& Result);
 
-	static FVector4f WeightToColor(float Value);
-	void UpdateVisualization(bool bInForce = false);
+	static UE_API FVector4f WeightToColor(float Value);
+	UE_API void UpdateVisualization(bool bInForce = false);
 
 	UPROPERTY()
 	TWeakObjectPtr<USkeletalMeshEditorContextObjectBase> EditorContext = nullptr;
 
 	// updates the skin weights using the operator's result
-	bool UpdateSkinWeightsFromDynamicMesh(UE::Geometry::FDynamicMesh3& InResultMesh) const;
+	UE_API bool UpdateSkinWeightsFromDynamicMesh(UE::Geometry::FDynamicMesh3& InResultMesh) const;
 	
 	// ISkeletalMeshEditionInterface
-	virtual void HandleSkeletalMeshModified(const TArray<FName>& InBoneNames, const ESkeletalMeshNotifyType InNotifyType) override;
+	UE_API virtual void HandleSkeletalMeshModified(const TArray<FName>& InBoneNames, const ESkeletalMeshNotifyType InNotifyType) override;
 
 	// Pending update function */
 	TFunction<void()> PendingUpdateFunction;
 };
+
+#undef UE_API
diff --git a/Engine/Plugins/Runtime/MeshModelingToolset/Source/MeshModelingToolsEditorOnly/Public/SkeletalMesh/SkinWeightsPaintTool.h b/Engine/Plugins/Runtime/MeshModelingToolset/Source/MeshModelingToolsEditorOnly/Public/SkeletalMesh/SkinWeightsPaintTool.h
--- a/Engine/Plugins/Runtime/MeshModelingToolset/Source/MeshModelingToolsEditorOnly/Public/SkeletalMesh/SkinWeightsPaintTool.h
+++ b/Engine/Plugins/Runtime/MeshModelingToolset/Source/MeshModelingToolsEditorOnly/Public/SkeletalMesh/SkinWeightsPaintTool.h
@@ -20,6 +20,8 @@
 
 #include "SkinWeightsPaintTool.generated.h"
 
+#define UE_API MESHMODELINGTOOLSEDITORONLY_API
+
 
 class USkeletalMeshComponentReadOnlyToolTarget;
 struct FMeshDescription;
@@ -320,16 +322,16 @@ namespace SkinPaintTool
 	
 }
 
-UCLASS()
-class MESHMODELINGTOOLSEDITORONLY_API USkinWeightsPaintToolBuilder : public UMeshSurfacePointMeshEditingToolBuilder
+UCLASS(MinimalAPI)
+class USkinWeightsPaintToolBuilder : public UMeshSurfacePointMeshEditingToolBuilder
 {
 	GENERATED_BODY()
 
 public:
-	UMeshSurfacePointTool* CreateNewTool(const FToolBuilderState& SceneState) const override;
+	UE_API UMeshSurfacePointTool* CreateNewTool(const FToolBuilderState& SceneState) const override;
 
 protected:
-	virtual const FToolTargetTypeRequirements& GetTargetRequirements() const override;
+	UE_API virtual const FToolTargetTypeRequirements& GetTargetRequirements() const override;
 };
 
 // for saveing/restoring the brush settings separately for each brush mode (Add, Replace, etc...)
@@ -351,26 +353,26 @@ struct FSkinWeightBrushConfig
 	EWeightBrushFalloffMode FalloffMode = EWeightBrushFalloffMode::Surface;
 };
 
-struct MESHMODELINGTOOLSEDITORONLY_API FDirectEditWeightState
+struct FDirectEditWeightState
 {
 	EWeightEditOperation EditMode;
 	float StartValue = 0.f;
 	float CurrentValue = 0.f;
 	bool bInTransaction = false;
 
-	void Reset();
-	float GetModeDefaultValue();
-	float GetModeMinValue();
-	float GetModeMaxValue();
+	UE_API void Reset();
+	UE_API float GetModeDefaultValue();
+	UE_API float GetModeMinValue();
+	UE_API float GetModeMaxValue();
 };
 
 // Container for properties displayed in Details panel while using USkinWeightsPaintTool
-UCLASS(config = EditorSettings)
-class MESHMODELINGTOOLSEDITORONLY_API USkinWeightsPaintToolProperties : public UBrushBaseProperties
+UCLASS(MinimalAPI, config = EditorSettings)
+class USkinWeightsPaintToolProperties : public UBrushBaseProperties
 {
 	GENERATED_BODY()
 
-	USkinWeightsPaintToolProperties();
+	UE_API USkinWeightsPaintToolProperties();
 	
 public:
 
@@ -416,7 +418,7 @@ public:
 	FDirectEditWeightState DirectEditState;
 
 	// save/restore user specified settings for each tool mode
-	FSkinWeightBrushConfig& GetBrushConfig();
+	UE_API FSkinWeightBrushConfig& GetBrushConfig();
 	TMap<EWeightEditOperation, FSkinWeightBrushConfig*> BrushConfigs;
 	UPROPERTY(Config)
 	FSkinWeightBrushConfig BrushConfigAdd;
@@ -440,15 +442,15 @@ public:
 		EditCondition = bShowNewProfileName, HideEditConditionToggle, NoResetToDefault))
 	FName NewSkinWeightProfile = "Profile";
 
-	FName GetActiveSkinWeightProfile() const;
+	UE_API FName GetActiveSkinWeightProfile() const;
 	
 	// pointer back to paint tool
 	TObjectPtr<USkinWeightsPaintTool> WeightTool;
 
-	void SetComponentMode(EComponentSelectionMode InComponentMode);
-	void SetFalloffMode(EWeightBrushFalloffMode InFalloffMode);
-	void SetColorMode(EWeightColorMode InColorMode);
-	void SetBrushMode(EWeightEditOperation InBrushMode);
+	UE_API void SetComponentMode(EComponentSelectionMode InComponentMode);
+	UE_API void SetFalloffMode(EWeightBrushFalloffMode InFalloffMode);
+	UE_API void SetColorMode(EWeightColorMode InColorMode);
+	UE_API void SetBrushMode(EWeightEditOperation InBrushMode);
 
 	// transfer
 	UPROPERTY(EditAnywhere, Transient, Category = "WeightTransfer")
@@ -464,13 +466,13 @@ public:
 	
 private:
 	UFUNCTION()
-	TArray<FName> GetTargetLODsFunc() const;
+	UE_API TArray<FName> GetTargetLODsFunc() const;
 	UFUNCTION()
-	TArray<FName> GetTargetSkinWeightProfilesFunc() const;
+	UE_API TArray<FName> GetTargetSkinWeightProfilesFunc() const;
 	UFUNCTION()
-	TArray<FName> GetSourceLODsFunc() const;
+	UE_API TArray<FName> GetSourceLODsFunc() const;
 	UFUNCTION()
-	TArray<FName> GetSourceSkinWeightProfilesFunc() const;
+	UE_API TArray<FName> GetSourceSkinWeightProfilesFunc() const;
 };
 
 // this class wraps the all the components to enable selection on a single mesh in the skin weights tool
@@ -478,15 +480,15 @@ private:
 // NOTE: at some point we may want to do component selections on multiple meshes in any/all viewports
 // at which time this class should be centralized and renamed to UMeshSelector or something like that.
 // But there will need to be some sort of centralized facility to manage that and make sure it interacts nicely with other tools.
-UCLASS()
-class MESHMODELINGTOOLSEDITORONLY_API UWeightToolMeshSelector : public UObject
+UCLASS(MinimalAPI)
+class UWeightToolMeshSelector : public UObject
 {
 	GENERATED_BODY()
 
 public:
 
 	// must be called during the Setup of the parent tool
-	void InitialSetup(
+	UE_API void InitialSetup(
 		UWorld* InWorld,
 		UInteractiveTool* InParentTool,
 		FEditorViewportClient* InViewportClient,
@@ -494,33 +496,33 @@ public:
 
 	// must be called AFTER InitialSetup, and any time the mesh is changed
 	// passing in a null preview mesh will disable the selector
-	void SetMesh(
+	UE_API void SetMesh(
 		UPreviewMesh* InMesh,
 		const FMeshDescription* InMeshDescription,
 		const FTransform3d& InMeshTransform);
 	
-	void UpdateAfterMeshDeformation();
+	UE_API void UpdateAfterMeshDeformation();
 	
-	void Shutdown();
+	UE_API void Shutdown();
 
-	void SetIsEnabled(bool bIsEnabled);
-	void SetComponentSelectionMode(EComponentSelectionMode InMode);
-	void SetTransform(const FTransform3d& InTargetTransform);
+	UE_API void SetIsEnabled(bool bIsEnabled);
+	UE_API void SetComponentSelectionMode(EComponentSelectionMode InMode);
+	UE_API void SetTransform(const FTransform3d& InTargetTransform);
 
 	// viewport 
-	void DrawHUD(FCanvas* Canvas, IToolsContextRenderAPI* RenderAPI);
-	void Render(IToolsContextRenderAPI* RenderAPI);
+	UE_API void DrawHUD(FCanvas* Canvas, IToolsContextRenderAPI* RenderAPI);
+	UE_API void Render(IToolsContextRenderAPI* RenderAPI);
 
 	// get a list of currently selected vertices (converting edges and faces to vertices)
-	const TArray<int32>& GetSelectedVertices();
-	bool IsAnyComponentSelected() const;
-	void GetSelectedTriangles(TArray<int32>& OutTriangleIndices) const;
+	UE_API const TArray<int32>& GetSelectedVertices();
+	UE_API bool IsAnyComponentSelected() const;
+	UE_API void GetSelectedTriangles(TArray<int32>& OutTriangleIndices) const;
 
 	// edit selection
-	void GrowSelection() const;
-	void ShrinkSelection() const;
-	void FloodSelection() const;
-	void SelectBorder() const;
+	UE_API void GrowSelection() const;
+	UE_API void ShrinkSelection() const;
+	UE_API void FloodSelection() const;
+	UE_API void SelectBorder() const;
 
 	// get access to the selection mechanic
 	UPolygonSelectionMechanic* GetSelectionMechanic() { return PolygonSelectionMechanic; };
@@ -545,31 +547,31 @@ private:
 };
 
 // this class wraps a source skeletal mesh used to transfer skin weights to the tool target mesh
-UCLASS()
-class MESHMODELINGTOOLSEDITORONLY_API UWeightToolTransferManager : public UObject
+UCLASS(MinimalAPI)
+class UWeightToolTransferManager : public UObject
 {
 	GENERATED_BODY()
 
 public:
 
 	// this must be called from within the parent tool's Setup() so that the selection mechanics are registered for capturing input
-	void InitialSetup(USkinWeightsPaintTool* InWeightTool, FEditorViewportClient* InViewportClient);
+	UE_API void InitialSetup(USkinWeightsPaintTool* InWeightTool, FEditorViewportClient* InViewportClient);
 	
 	// called when the tool is shutdown
-	void Shutdown();
+	UE_API void Shutdown();
 
 	// render the selection mechanism
-	void Render(IToolsContextRenderAPI* RenderAPI);
-	void DrawHUD(FCanvas* Canvas, IToolsContextRenderAPI* RenderAPI);
+	UE_API void Render(IToolsContextRenderAPI* RenderAPI);
+	UE_API void DrawHUD(FCanvas* Canvas, IToolsContextRenderAPI* RenderAPI);
 
 	// update the mesh we are transferring from
-	void SetSourceMesh(USkeletalMesh* InSkeletalMesh = nullptr);
+	UE_API void SetSourceMesh(USkeletalMesh* InSkeletalMesh = nullptr);
 
 	// run the weight transfer
-	void TransferWeights();
+	UE_API void TransferWeights();
 
 	// returns true if everything is setup and ready to transfer
-	bool CanTransferWeights() const;
+	UE_API bool CanTransferWeights() const;
 	
 	// gets the tool target for the source mesh
 	UToolTarget* GetTarget() const { return SourceTarget; }
@@ -581,17 +583,17 @@ public:
 	UWeightToolMeshSelector* GetMeshSelector() const { return MeshSelector; };
 	
 	// called when tool settings are modified
-	void OnPropertyModified(const USkinWeightsPaintToolProperties* WeightToolProperties, const FProperty* ModifiedProperty);
+	UE_API void OnPropertyModified(const USkinWeightsPaintToolProperties* WeightToolProperties, const FProperty* ModifiedProperty);
 
 private:
 	
 	// actually run the weight transfer to copy weights from the source to the target
-	void TransferWeightsFromOtherMeshOrSubset();
+	UE_API void TransferWeightsFromOtherMeshOrSubset();
 
 	// actually run the weight transfer to copy weights from the source to the target
-	void TransferWeightsFromSameMeshAndLOD();
+	UE_API void TransferWeightsFromSameMeshAndLOD();
 
-	void ApplyTranferredWeightsAsTransaction(
+	UE_API void ApplyTranferredWeightsAsTransaction(
 		const UE::Geometry::FDynamicMeshVertexSkinWeightsAttribute* InTransferredSkinWeights,
 		const TArray<int32>& InVertexSubset,
 		const FDynamicMesh3& InTargetMesh);
@@ -608,47 +610,47 @@ private:
 };
 
 // this class wraps all the data needed to isolate a selection of a mesh while editing skin weights
-UCLASS()
-class MESHMODELINGTOOLSEDITORONLY_API UWeightToolSelectionIsolator : public UObject
+UCLASS(MinimalAPI)
+class UWeightToolSelectionIsolator : public UObject
 {
 	GENERATED_BODY()
 
 public:
 
 	// call during tool Setup()
-	void InitialSetup(USkinWeightsPaintTool* InTool);
+	UE_API void InitialSetup(USkinWeightsPaintTool* InTool);
 
 	// call every tick to apply deferred changes to mesh
-	void UpdateIsolatedSelection();
+	UE_API void UpdateIsolatedSelection();
 
 	// returns true if any triangles are currently isolated
-	bool IsSelectionIsolated() const;
+	UE_API bool IsSelectionIsolated() const;
 
 	// isolate the current selection
-	void IsolateSelectionAsTransaction();
+	UE_API void IsolateSelectionAsTransaction();
 
 	// unisolate the current selection
-	void UnIsolateSelectionAsTransaction();
+	UE_API void UnIsolateSelectionAsTransaction();
 
 	// isolate the array of triangles
-	void SetTrianglesToIsolate(const TArray<int32>& TrianglesToIsolate);
+	UE_API void SetTrianglesToIsolate(const TArray<int32>& TrianglesToIsolate);
 	
 	// restores the whole mesh
-	void RestoreFullMesh();
+	UE_API void RestoreFullMesh();
 	
 	// get the current triangles that are isolated
 	const TArray<int32>& GetIsolatedTriangles() { return CurrentlyIsolatedTriangles; };
 
 	// convert to/from partial-isolated and full mesh vertex indices
-	int32 PartialToFullMeshVertexIndex(int32 PartialMeshVertexIndex) const;
-	int32 FullToPartialMeshVertexIndex(int32 FullMeshVertexIndex) const;
+	UE_API int32 PartialToFullMeshVertexIndex(int32 PartialMeshVertexIndex) const;
+	UE_API int32 FullToPartialMeshVertexIndex(int32 FullMeshVertexIndex) const;
 
 	// returns the isolated partial mesh (if PartialMeshDescription is not null, this will return an empty DynamicMesh3)
-	const FDynamicMesh3& GetPartialMesh() const;
+	UE_API const FDynamicMesh3& GetPartialMesh() const;
 
 private:
 
-	void CreatePartialMesh();
+	UE_API void CreatePartialMesh();
 	
 	UPROPERTY()
 	TObjectPtr<USkinWeightsPaintTool> WeightTool;
@@ -665,15 +667,15 @@ private:
 	UE::Geometry::FDynamicSubmesh3 PartialSubMesh;
 };
 
-class MESHMODELINGTOOLSEDITORONLY_API FIsolateSelectionChange : public FToolCommandChange
+class FIsolateSelectionChange : public FToolCommandChange
 {
 public:
 	TArray<int32> IsolatedTrianglesBefore;
 	TArray<int32> IsolatedTrianglesAfter;
 
-	virtual void Apply(UObject* Object) override;
-	virtual void Revert(UObject* Object) override;
-	virtual FString ToString() const override;
+	UE_API virtual void Apply(UObject* Object) override;
+	UE_API virtual void Revert(UObject* Object) override;
+	UE_API virtual FString ToString() const override;
 };
 
 // Why does this exist?
@@ -683,17 +685,17 @@ public:
 // So, rather than track these issues down individually, we leave the original mesh topology intact, and operate on a cleaned submesh.
 // When the tool exits, we copy the weight edits from the submesh to the full mesh.
 // This means weights of the floating vertices cannot be edited, however they do not matter since they do not belong to any triangles they cannot be rendered
-struct MESHMODELINGTOOLSEDITORONLY_API FCleanedEditMesh
+struct FCleanedEditMesh
 {
-	FCleanedEditMesh(const FDynamicMesh3& InDynamicMesh, const FMeshDescription& InMeshDescription);
-	void CopyWeightsToOriginalMesh(const FName InProfile);
-	FDynamicMesh3& GetEditableMesh();
-	FMeshDescription& GetEditableMeshDescription();
-	FDynamicMesh3& GetOriginalMesh();
-	FMeshDescription& GetOriginalMeshDescription();
+	UE_API FCleanedEditMesh(const FDynamicMesh3& InDynamicMesh, const FMeshDescription& InMeshDescription);
+	UE_API void CopyWeightsToOriginalMesh(const FName InProfile);
+	UE_API FDynamicMesh3& GetEditableMesh();
+	UE_API FMeshDescription& GetEditableMeshDescription();
+	UE_API FDynamicMesh3& GetOriginalMesh();
+	UE_API FMeshDescription& GetOriginalMeshDescription();
 
 	// returns the VertexID on the original mesh 
-	int32 GetOriginalVertexID(const int32 InCleanedVertexId) const;
+	UE_API int32 GetOriginalVertexID(const int32 InCleanedVertexId) const;
 	
 private:
 	FDynamicMesh3 OriginalDynamicMesh;
@@ -703,8 +705,8 @@ private:
 };
 
 // An interactive tool for painting and editing skin weights.
-UCLASS()
-class MESHMODELINGTOOLSEDITORONLY_API USkinWeightsPaintTool : public UDynamicMeshBrushTool, public ISkeletalMeshEditingInterface
+UCLASS(MinimalAPI)
+class USkinWeightsPaintTool : public UDynamicMeshBrushTool, public ISkeletalMeshEditingInterface
 {
 	GENERATED_BODY()
 
@@ -714,61 +716,61 @@ public:
 	virtual bool HasCancel() const override { return true; }
 	virtual bool HasAccept() const override { return true; }
 	virtual bool CanAccept() const override { return true; }
-	virtual bool HitTest(const FRay& Ray, FHitResult& OutHit) override;
-	virtual void OnBeginDrag(const FRay& Ray) override;
-	virtual void OnUpdateDrag(const FRay& Ray) override;
-	virtual void OnEndDrag(const FRay& Ray) override;
-	virtual bool OnUpdateHover(const FInputDeviceRay& DevicePos) override;
-	virtual double EstimateMaximumTargetDimension() override;
-
-	void Init(const FToolBuilderState& InSceneState);
+	UE_API virtual bool HitTest(const FRay& Ray, FHitResult& OutHit) override;
+	UE_API virtual void OnBeginDrag(const FRay& Ray) override;
+	UE_API virtual void OnUpdateDrag(const FRay& Ray) override;
+	UE_API virtual void OnEndDrag(const FRay& Ray) override;
+	UE_API virtual bool OnUpdateHover(const FInputDeviceRay& DevicePos) override;
+	UE_API virtual double EstimateMaximumTargetDimension() override;
+
+	UE_API void Init(const FToolBuilderState& InSceneState);
 	
 	// UInteractiveTool
-	virtual void Setup() override;
-	virtual void DrawHUD(FCanvas* Canvas, IToolsContextRenderAPI* RenderAPI) override;
-	virtual void Render(IToolsContextRenderAPI* RenderAPI) override;
+	UE_API virtual void Setup() override;
+	UE_API virtual void DrawHUD(FCanvas* Canvas, IToolsContextRenderAPI* RenderAPI) override;
+	UE_API virtual void Render(IToolsContextRenderAPI* RenderAPI) override;
 
 	// IInteractiveToolCameraFocusAPI
 	virtual bool SupportsWorldSpaceFocusBox() override { return true; }
-	virtual FBox GetWorldSpaceFocusBox() override;
+	UE_API virtual FBox GetWorldSpaceFocusBox() override;
 
 	// IClickDragBehaviorTarget implementation
-	virtual void OnUpdateModifierState(int ModifierID, bool bIsOn) override;
-	virtual FInputRayHit CanBeginClickDragSequence(const FInputDeviceRay& InPressPos) override;
+	UE_API virtual void OnUpdateModifierState(int ModifierID, bool bIsOn) override;
+	UE_API virtual FInputRayHit CanBeginClickDragSequence(const FInputDeviceRay& InPressPos) override;
 
 	// using when ToolChange is applied via Undo/Redo
-	void ExternalUpdateWeights(const int32 BoneIndex, const TMap<int32, float>& IndexValues);
-	void ExternalUpdateSkinWeightLayer(const EMeshLODIdentifier InLOD, const FName InSkinWeightProfile);
-	void ExternalAddInfluenceToVertices(const BoneIndex InfluenceToAdd, const TArray<VertexIndex>& Vertices);
-	void ExternalRemoveInfluenceFromVertices(const BoneIndex InfluenceToRemove, const TArray<VertexIndex>& Vertices);
+	UE_API void ExternalUpdateWeights(const int32 BoneIndex, const TMap<int32, float>& IndexValues);
+	UE_API void ExternalUpdateSkinWeightLayer(const EMeshLODIdentifier InLOD, const FName InSkinWeightProfile);
+	UE_API void ExternalAddInfluenceToVertices(const BoneIndex InfluenceToAdd, const TArray<VertexIndex>& Vertices);
+	UE_API void ExternalRemoveInfluenceFromVertices(const BoneIndex InfluenceToRemove, const TArray<VertexIndex>& Vertices);
 
 	// weight editing operations (selection based)
-	void MirrorWeights(EAxis::Type Axis, EMirrorDirection Direction);
-	void PruneWeights(const float Threshold, const TArray<BoneIndex>& BonesToPrune);
-	void AverageWeights(const float Strength);
-	void NormalizeWeights();
-	void HammerWeights();
-	void ClampInfluences(const int32 MaxInfluences);
+	UE_API void MirrorWeights(EAxis::Type Axis, EMirrorDirection Direction);
+	UE_API void PruneWeights(const float Threshold, const TArray<BoneIndex>& BonesToPrune);
+	UE_API void AverageWeights(const float Strength);
+	UE_API void NormalizeWeights();
+	UE_API void HammerWeights();
+	UE_API void ClampInfluences(const int32 MaxInfluences);
 
 	// HELPER functions for modifying weights
 	//
 	// given a map of BoneIndex > Weight values for a single vertex, modify the map by removing the smallest weights to fit in Max Influences
-	static void TruncateWeightMap(TMap<BoneIndex, float>& InOutWeights);
+	static UE_API void TruncateWeightMap(TMap<BoneIndex, float>& InOutWeights);
 	// given a map of BoneIndex > Weight values for a single vertex, modify the weights to sum to 1
-	static void NormalizeWeightMap(TMap<BoneIndex, float>& InOutWeights);
+	static UE_API void NormalizeWeightMap(TMap<BoneIndex, float>& InOutWeights);
 	// sum all the weights on all bones for a given list of vertices (results we not be normalized!)
-	static void AccumulateWeights(
+	static UE_API void AccumulateWeights(
 		const TArray<SkinPaintTool::VertexWeights>& AllWeights,
 		const TArray<VertexIndex>& VerticesToAccumulate,
 		TMap<BoneIndex, float>& OutWeights);
 
 	// copy paste
-	void CopyWeights();
-	void PasteWeights();
-	static const FString CopyPasteWeightsIdentifier;
+	UE_API void CopyWeights();
+	UE_API void PasteWeights();
+	static UE_API const FString CopyPasteWeightsIdentifier;
 	
 	// method to set weights directly (numeric input, for example)
-	void EditWeightsOnVertices(
+	UE_API void EditWeightsOnVertices(
 		BoneIndex Bone,
 		const float Value,
 		const int32 Iterations,
@@ -777,36 +779,36 @@ public:
 		const bool bShouldTransact);
 
 	// toggle brush / selection mode
-	void ToggleEditingMode();
+	UE_API void ToggleEditingMode();
 	// update the state of the mesh selectors
-	void UpdateSelectorState() const;
+	UE_API void UpdateSelectorState() const;
 	
 	// get access to the mesh selector for the main mesh
-	UWeightToolMeshSelector* GetMainMeshSelector();
+	UE_API UWeightToolMeshSelector* GetMainMeshSelector();
 	// get access to the currently active mesh selector (may be on the transfer source mesh)
-	UWeightToolMeshSelector* GetActiveMeshSelector();
+	UE_API UWeightToolMeshSelector* GetActiveMeshSelector();
 	// does the main mesh have an active selection ("active" meaning the selection is currently being rendered in the view and is editable)
-	bool HasActiveSelectionOnMainMesh();
+	UE_API bool HasActiveSelectionOnMainMesh();
 	// select all vertices affected by the currently selected bone(s)
-	void SelectAffected() const;
+	UE_API void SelectAffected() const;
 	// select all vertices affected by at least MinInfluences number of bones
-	void SelectByInfluenceCount(const int32 MinInfluences) const;
+	UE_API void SelectByInfluenceCount(const int32 MinInfluences) const;
 
 	// get the average weight value of each influence on the given vertices
-	void GetInfluences(const TArray<int32>& VertexIndices, TArray<BoneIndex>& OutBoneIndices);
+	UE_API void GetInfluences(const TArray<int32>& VertexIndices, TArray<BoneIndex>& OutBoneIndices);
 	// get the average weight value of a single bone on the given vertices
-	float GetAverageWeightOnBone(const BoneIndex InBoneIndex, const TArray<int32>& VertexIndices);
+	UE_API float GetAverageWeightOnBone(const BoneIndex InBoneIndex, const TArray<int32>& VertexIndices);
 	// convert an index to a name
-	FName GetBoneNameFromIndex(BoneIndex InIndex) const;
+	UE_API FName GetBoneNameFromIndex(BoneIndex InIndex) const;
 	// get the currently selected bone
-	BoneIndex GetCurrentBoneIndex() const;
+	UE_API BoneIndex GetCurrentBoneIndex() const;
 	// get a list of vertices affected by the given bone
-	void GetVerticesAffectedByBone(BoneIndex IndexOfBone, TSet<int32>& OutVertexIndices) const;
+	UE_API void GetVerticesAffectedByBone(BoneIndex IndexOfBone, TSet<int32>& OutVertexIndices) const;
 
 	// toggle the display of weights on the preview mesh (if false, uses the normal skeletal mesh material)
-	void SetDisplayVertexColors(bool bShowVertexColors=true);
+	UE_API void SetDisplayVertexColors(bool bShowVertexColors=true);
 	// set focus back to viewport so that hotkeys are immediately detected while hovering
-	void SetFocusInViewport() const;
+	UE_API void SetFocusInViewport() const;
 
 	// get the target manager (cached from Setup)
 	UToolTargetManager* GetTargetManager() const { return TargetManager.Get(); };
@@ -818,18 +820,18 @@ public:
 	UWeightToolTransferManager* GetWeightTransferManager() const { return TransferManager; };
 
 	// get the viewport this tool is operating in
-	virtual FEditorViewportClient* GetViewportClient() const;
+	UE_API virtual FEditorViewportClient* GetViewportClient() const;
 	
 	// get access to the selection isolation system
 	UWeightToolSelectionIsolator* GetSelectionIsolator() const { return SelectionIsolator; };
 
 	// get the tool properties
-	USkinWeightsPaintToolProperties* GetWeightToolProperties() const;
+	UE_API USkinWeightsPaintToolProperties* GetWeightToolProperties() const;
 
 	// get access to the mesh description for the mesh being edited
-	FCleanedEditMesh* GetCurrentCleanedEditMesh() const;
-	FMeshDescription* GetCurrentlyEditedMeshDescription() const;
-	FDynamicMesh3* GetCurrentlyEditedDynamicMesh() const;
+	UE_API FCleanedEditMesh* GetCurrentCleanedEditMesh() const;
+	UE_API FMeshDescription* GetCurrentlyEditedMeshDescription() const;
+	UE_API FDynamicMesh3* GetCurrentlyEditedDynamicMesh() const;
 
 	// HOW TO EDIT WEIGHTS WITH UNDO/REDO:
 	//
@@ -841,16 +843,16 @@ public:
 	//  2. EndChange()
 	// All the edits are stored into the "ActiveChange" and applied as a single transaction in EndChange().
 	// Deformations and vertex colors will be updated throughout the duration of the change.
-	void BeginChange();
-	void EndChange(const FText& TransactionLabel);
-	void ApplyWeightEditsWithoutTransaction(const SkinPaintTool::FMultiBoneWeightEdits& WeightEdits);
+	UE_API void BeginChange();
+	UE_API void EndChange(const FText& TransactionLabel);
+	UE_API void ApplyWeightEditsWithoutTransaction(const SkinPaintTool::FMultiBoneWeightEdits& WeightEdits);
 	// "One-off" Edits:
 	// For all one-and-done edits, you can call ApplyWeightEditsAsTransaction().
 	// It will Begin/End the change and create a transaction for it.
-	void ApplyWeightEditsAsTransaction(const SkinPaintTool::FMultiBoneWeightEdits& WeightEdits, const FText& TransactionLabel);
+	UE_API void ApplyWeightEditsAsTransaction(const SkinPaintTool::FMultiBoneWeightEdits& WeightEdits, const FText& TransactionLabel);
 
 	// call this whenever the target mesh is modified
-	void UpdateCurrentlyEditedMesh(
+	UE_API void UpdateCurrentlyEditedMesh(
 		const USkeletalMeshComponent* Component,
 		const FDynamicMesh3& InDynamicMesh,
 		const FMeshDescription& InMeshDescription);
@@ -865,17 +867,17 @@ public:
 
 protected:
 
-	virtual void ApplyStamp(const FBrushStampData& Stamp);
-	void OnShutdown(EToolShutdownType ShutdownType) override;
-	void OnTick(float DeltaTime) override;
+	UE_API virtual void ApplyStamp(const FBrushStampData& Stamp);
+	UE_API void OnShutdown(EToolShutdownType ShutdownType) override;
+	UE_API void OnTick(float DeltaTime) override;
 
 	// stamp
-	float CalculateBrushFalloff(float Distance) const;
-	void CalculateVertexROI(
+	UE_API float CalculateBrushFalloff(float Distance) const;
+	UE_API void CalculateVertexROI(
 		const FBrushStampData& InStamp,
 		TArray<VertexIndex>& OutVertexIDs,
 		TArray<float>& OutVertexFalloffs);
-	float CalculateBrushStrengthToUse(EWeightEditOperation EditMode) const;
+	UE_API float CalculateBrushStrengthToUse(EWeightEditOperation EditMode) const;
 	bool bInvertStroke = false;
 	FBrushStampData StartStamp;
 	FBrushStampData LastStamp;
@@ -885,7 +887,7 @@ protected:
 	
 	// generating bone weight edits to be stored in a transaction
 	// does not actually change the weight buffers
-	void CreateWeightEditsForVertices(
+	UE_API void CreateWeightEditsForVertices(
 		EWeightEditOperation EditOperation,
 		const BoneIndex Bone,
 		const TArray<int32>& VerticesToEdit,
@@ -893,7 +895,7 @@ protected:
 		const float InValue,
 		SkinPaintTool::FMultiBoneWeightEdits& InOutWeightEdits);
 	// same as CreateWeightEditsForVertices() but specific to relaxation (topology aware operation)
-	void CreateWeightEditsToRelaxVertices(
+	UE_API void CreateWeightEditsToRelaxVertices(
 		TArray<int32> VerticesToEdit,
 		TArray<float> VertexFalloffs,
 		const float Strength,
@@ -907,12 +909,12 @@ protected:
 	TUniquePtr<DynamicTrianglesOctree> TrianglesOctree;
 	TFuture<void> TriangleOctreeFuture;
 	TArray<int32> TrianglesToReinsert;
-	void InitializeOctrees();
+	UE_API void InitializeOctrees();
 
 	// tool properties
 	UPROPERTY()
 	TObjectPtr<USkinWeightsPaintToolProperties> WeightToolProperties;
-	virtual void OnPropertyModified(UObject* ModifiedObject, FProperty* ModifiedProperty) override;
+	UE_API virtual void OnPropertyModified(UObject* ModifiedObject, FProperty* ModifiedProperty) override;
 	
 	// the currently edited mesh descriptions
 	mutable TMap<EMeshLODIdentifier, FCleanedEditMesh> EditedMeshes;
@@ -930,36 +932,36 @@ protected:
 	// Smooth weights data source and operator
 	TUniquePtr<UE::Geometry::TBoneWeightsDataSource<int32, float>> SmoothWeightsDataSource;
 	TUniquePtr<UE::Geometry::TSmoothBoneWeights<int32, float>> SmoothWeightsOp;
-	void InitializeSmoothWeightsOperator();
+	UE_API void InitializeSmoothWeightsOperator();
 
 	// vertex colors updated when switching current bone or initializing whole mesh
-	void UpdateVertexColorForAllVertices();
+	UE_API void UpdateVertexColorForAllVertices();
 	bool bVertexColorsNeedUpdated = false;
 	// vertex colors updated when make sparse edits to subset of vertices
-	void UpdateVertexColorForSubsetOfVertices();
+	UE_API void UpdateVertexColorForSubsetOfVertices();
 	TSet<int32> VerticesToUpdateColor;
-	FVector4f GetColorOfVertex(VertexIndex InVertexIndex, BoneIndex InBoneIndex) const;
+	UE_API FVector4f GetColorOfVertex(VertexIndex InVertexIndex, BoneIndex InBoneIndex) const;
 
 	// which bone are we currently painting?
-	void UpdateCurrentBone(const FName &BoneName);
-	BoneIndex GetBoneIndexFromName(const FName BoneName) const;
+	UE_API void UpdateCurrentBone(const FName &BoneName);
+	UE_API BoneIndex GetBoneIndexFromName(const FName BoneName) const;
 	FName CurrentBone = NAME_None;
 	TOptional<FName> PendingCurrentBone;
 	TArray<FName> SelectedBoneNames;
 	TArray<BoneIndex> SelectedBoneIndices;
 
 	// ISkeletalMeshEditionInterface
-	virtual void HandleSkeletalMeshModified(const TArray<FName>& InBoneNames, const ESkeletalMeshNotifyType InNotifyType) override;
+	UE_API virtual void HandleSkeletalMeshModified(const TArray<FName>& InBoneNames, const ESkeletalMeshNotifyType InNotifyType) override;
 
 	// the selection system for the main mesh
 	UPROPERTY()
 	TObjectPtr<UWeightToolMeshSelector> MeshSelector;
 
 	// skin weight layer
-	void OnActiveLODChanged();
-	void OnActiveSkinWeightProfileChanged();
-	void OnNewSkinWeightProfileChanged();
-	bool IsProfileValid(const FName InProfileName) const;
+	UE_API void OnActiveLODChanged();
+	UE_API void OnActiveSkinWeightProfileChanged();
+	UE_API void OnNewSkinWeightProfileChanged();
+	UE_API bool IsProfileValid(const FName InProfileName) const;
 
 	// global properties stored on initialization
 	UPROPERTY()
@@ -982,3 +984,5 @@ protected:
 
 	friend SkinPaintTool::FSkinToolDeformer;
 };
+
+#undef UE_API
diff --git a/Engine/Plugins/Runtime/MeshModelingToolset/Source/ModelingComponentsEditorOnly/Public/AssetUtils/CreateMaterialUtil.h b/Engine/Plugins/Runtime/MeshModelingToolset/Source/ModelingComponentsEditorOnly/Public/AssetUtils/CreateMaterialUtil.h
--- a/Engine/Plugins/Runtime/MeshModelingToolset/Source/ModelingComponentsEditorOnly/Public/AssetUtils/CreateMaterialUtil.h
+++ b/Engine/Plugins/Runtime/MeshModelingToolset/Source/ModelingComponentsEditorOnly/Public/AssetUtils/CreateMaterialUtil.h
@@ -29,7 +29,7 @@ namespace AssetUtils
 	/**
 	 * Options for new UMaterial asset created by Create() functions below.
 	 */
-	struct MODELINGCOMPONENTSEDITORONLY_API FMaterialAssetOptions
+	struct FMaterialAssetOptions
 	{
 		// this package will be used if it is not null, otherwise a new package will be created at NewAssetPath
 		UPackage* UsePackage = nullptr;
@@ -42,7 +42,7 @@ namespace AssetUtils
 	 * Output information about a newly-created UMaterial, returned by Create functions below.
 	 * Some fields may be null, if the relevant function did not create that type of object
 	 */
-	struct MODELINGCOMPONENTSEDITORONLY_API FMaterialAssetResults
+	struct FMaterialAssetResults
 	{
 		UMaterial* NewMaterial = nullptr;
 		UMaterialInstanceConstant* NewMaterialInstance = nullptr;
diff --git a/Engine/Plugins/Runtime/MeshModelingToolset/Source/ModelingComponentsEditorOnly/Public/AssetUtils/CreateStaticMeshUtil.h b/Engine/Plugins/Runtime/MeshModelingToolset/Source/ModelingComponentsEditorOnly/Public/AssetUtils/CreateStaticMeshUtil.h
--- a/Engine/Plugins/Runtime/MeshModelingToolset/Source/ModelingComponentsEditorOnly/Public/AssetUtils/CreateStaticMeshUtil.h
+++ b/Engine/Plugins/Runtime/MeshModelingToolset/Source/ModelingComponentsEditorOnly/Public/AssetUtils/CreateStaticMeshUtil.h
@@ -37,7 +37,7 @@ namespace AssetUtils
 	 * Only one of the arrays should be initialized, and it's size should be equal to 
 	 * the number of LODs specified for the Asset.
 	 */
-	struct MODELINGCOMPONENTSEDITORONLY_API FStaticMeshAssetMeshes
+	struct FStaticMeshAssetMeshes
 	{
 		TArray<const UE::Geometry::FDynamicMesh3*> DynamicMeshes;
 		TArray<const FMeshDescription*> MeshDescriptions;
@@ -47,7 +47,7 @@ namespace AssetUtils
 	/**
 	 * Options for new UStaticMesh asset created by Create() functions below.
 	 */
-	struct MODELINGCOMPONENTSEDITORONLY_API FStaticMeshAssetOptions
+	struct FStaticMeshAssetOptions
 	{
 		// this package will be used if it is not null, otherwise a new package will be created at NewAssetPath
 		UPackage* UsePackage = nullptr;
@@ -108,7 +108,7 @@ namespace AssetUtils
 	 * Output information about a newly-created StaticMesh, returned by Create functions below.
 	 * Some fields may be null, if the relevant function did not create that type of object
 	 */
-	struct MODELINGCOMPONENTSEDITORONLY_API FStaticMeshResults
+	struct FStaticMeshResults
 	{
 		/** UStaticMesh asset that was created */
 		UStaticMesh* StaticMesh = nullptr;
diff --git a/Engine/Plugins/Runtime/MeshModelingToolset/Source/ModelingComponentsEditorOnly/Public/AssetUtils/CreateTexture2DUtil.h b/Engine/Plugins/Runtime/MeshModelingToolset/Source/ModelingComponentsEditorOnly/Public/AssetUtils/CreateTexture2DUtil.h
--- a/Engine/Plugins/Runtime/MeshModelingToolset/Source/ModelingComponentsEditorOnly/Public/AssetUtils/CreateTexture2DUtil.h
+++ b/Engine/Plugins/Runtime/MeshModelingToolset/Source/ModelingComponentsEditorOnly/Public/AssetUtils/CreateTexture2DUtil.h
@@ -28,7 +28,7 @@ namespace AssetUtils
 	/**
 	 * Options for new UTexture asset created by Create() functions below.
 	 */
-	struct MODELINGCOMPONENTSEDITORONLY_API FTexture2DAssetOptions
+	struct FTexture2DAssetOptions
 	{
 		// this package will be used if it is not null, otherwise a new package will be created at NewAssetPath
 		UPackage* UsePackage = nullptr;
@@ -44,7 +44,7 @@ namespace AssetUtils
 	 * Output information about a newly-created UTexture2D, returned by Create functions below.
 	 * Some fields may be null, if the relevant function did not create that type of object
 	 */
-	struct MODELINGCOMPONENTSEDITORONLY_API FTexture2DAssetResults
+	struct FTexture2DAssetResults
 	{
 		UTexture2D* Texture = nullptr;
 	};
