Updated headers using UnrealCodeFixup to make sure dllstorage is on methods/staticvars instead of types. Part 4/n

diff --git a/Engine/Plugins/Runtime/MeshModelingToolset/Source/ModelingComponentsEditorOnly/Public/EditorModelingObjectsCreationAPI.h b/Engine/Plugins/Runtime/MeshModelingToolset/Source/ModelingComponentsEditorOnly/Public/EditorModelingObjectsCreationAPI.h
--- a/Engine/Plugins/Runtime/MeshModelingToolset/Source/ModelingComponentsEditorOnly/Public/EditorModelingObjectsCreationAPI.h
+++ b/Engine/Plugins/Runtime/MeshModelingToolset/Source/ModelingComponentsEditorOnly/Public/EditorModelingObjectsCreationAPI.h
@@ -7,6 +7,8 @@
 
 #include "EditorModelingObjectsCreationAPI.generated.h"
 
+#define UE_API MODELINGCOMPONENTSEDITORONLY_API
+
 class UInteractiveToolsContext;
 
 /**
@@ -25,26 +27,26 @@ class UInteractiveToolsContext;
  *  - OnModelingTextureCreated is broadcast for each new created texture object
  *  - OnModelingMaterialCreated is broadcast for each new created material object
  */
-UCLASS()
-class MODELINGCOMPONENTSEDITORONLY_API UEditorModelingObjectsCreationAPI : public UModelingObjectsCreationAPI
+UCLASS(MinimalAPI)
+class UEditorModelingObjectsCreationAPI : public UModelingObjectsCreationAPI
 {
 	GENERATED_BODY()
 public:
 
 	// UFUNCTION(BlueprintCallable, Category = "Modeling Objects")
-	virtual FCreateMeshObjectResult CreateMeshObject(const FCreateMeshObjectParams& CreateMeshParams) override;
+	UE_API virtual FCreateMeshObjectResult CreateMeshObject(const FCreateMeshObjectParams& CreateMeshParams) override;
 
 	// UFUNCTION(BlueprintCallable, Category = "Modeling Objects")
-	virtual FCreateTextureObjectResult CreateTextureObject(const FCreateTextureObjectParams& CreateTexParams) override;
+	UE_API virtual FCreateTextureObjectResult CreateTextureObject(const FCreateTextureObjectParams& CreateTexParams) override;
 
 	// UFUNCTION(BlueprintCallable, Category = "Modeling Objects")
-	virtual FCreateMaterialObjectResult CreateMaterialObject(const FCreateMaterialObjectParams& CreateMaterialParams) override;
+	UE_API virtual FCreateMaterialObjectResult CreateMaterialObject(const FCreateMaterialObjectParams& CreateMaterialParams) override;
 
 	// UFUNCTION(BlueprintCallable, Category = "Modeling Objects")
-	virtual FCreateActorResult CreateNewActor(const FCreateActorParams& CreateActorParams) override;
+	UE_API virtual FCreateActorResult CreateNewActor(const FCreateActorParams& CreateActorParams) override;
 
 	// UFUNCTION(BlueprintCallable, Category = "Modeling Objects")
-	virtual FCreateComponentResult CreateNewComponentOnActor(const FCreateComponentParams& CreateComponentParams) override;
+	UE_API virtual FCreateComponentResult CreateNewComponentOnActor(const FCreateComponentParams& CreateComponentParams) override;
 
 
 	//
@@ -53,15 +55,15 @@ public:
 
 	virtual bool HasMoveVariants() const { return true; }
 
-	virtual FCreateMeshObjectResult CreateMeshObject(FCreateMeshObjectParams&& CreateMeshParams) override;
+	UE_API virtual FCreateMeshObjectResult CreateMeshObject(FCreateMeshObjectParams&& CreateMeshParams) override;
 
-	virtual FCreateTextureObjectResult CreateTextureObject(FCreateTextureObjectParams&& CreateTexParams) override;
+	UE_API virtual FCreateTextureObjectResult CreateTextureObject(FCreateTextureObjectParams&& CreateTexParams) override;
 
-	virtual FCreateMaterialObjectResult CreateMaterialObject(FCreateMaterialObjectParams&& CreateMaterialParams) override;
+	UE_API virtual FCreateMaterialObjectResult CreateMaterialObject(FCreateMaterialObjectParams&& CreateMaterialParams) override;
 
-	virtual FCreateActorResult CreateNewActor(FCreateActorParams&& CreateActorParams) override;
+	UE_API virtual FCreateActorResult CreateNewActor(FCreateActorParams&& CreateActorParams) override;
 
-	virtual FCreateComponentResult CreateNewComponentOnActor(FCreateComponentParams&& CreateComponentParams) override;
+	UE_API virtual FCreateComponentResult CreateNewComponentOnActor(FCreateComponentParams&& CreateComponentParams) override;
 
 	//
 	// Callbacks that editor can hook into to handle asset creation
@@ -89,19 +91,21 @@ public:
 	// Utility functions to handle registration/unregistration
 	//
 
-	static UEditorModelingObjectsCreationAPI* Register(UInteractiveToolsContext* ToolsContext);
-	static UEditorModelingObjectsCreationAPI* Find(UInteractiveToolsContext* ToolsContext);
-	static bool Deregister(UInteractiveToolsContext* ToolsContext);
+	static UE_API UEditorModelingObjectsCreationAPI* Register(UInteractiveToolsContext* ToolsContext);
+	static UE_API UEditorModelingObjectsCreationAPI* Find(UInteractiveToolsContext* ToolsContext);
+	static UE_API bool Deregister(UInteractiveToolsContext* ToolsContext);
 
 
 
 	//
 	// internal implementations called by public functions
 	//
-	FCreateMeshObjectResult CreateStaticMeshAsset(FCreateMeshObjectParams&& CreateMeshParams);
-	FCreateMeshObjectResult CreateVolume(FCreateMeshObjectParams&& CreateMeshParams);
-	FCreateMeshObjectResult CreateDynamicMeshActor(FCreateMeshObjectParams&& CreateMeshParams);
-	ECreateModelingObjectResult GetNewAssetPath(FString& OutNewAssetPath, const FString& BaseName, const UObject* StoreRelativeToObject, const UWorld* World);
+	UE_API FCreateMeshObjectResult CreateStaticMeshAsset(FCreateMeshObjectParams&& CreateMeshParams);
+	UE_API FCreateMeshObjectResult CreateVolume(FCreateMeshObjectParams&& CreateMeshParams);
+	UE_API FCreateMeshObjectResult CreateDynamicMeshActor(FCreateMeshObjectParams&& CreateMeshParams);
+	UE_API ECreateModelingObjectResult GetNewAssetPath(FString& OutNewAssetPath, const FString& BaseName, const UObject* StoreRelativeToObject, const UWorld* World);
 
-	TArray<UMaterialInterface*> FilterMaterials(const TArray<UMaterialInterface*>& MaterialsIn);
+	UE_API TArray<UMaterialInterface*> FilterMaterials(const TArray<UMaterialInterface*>& MaterialsIn);
 };
+
+#undef UE_API
diff --git a/Engine/Plugins/Runtime/MeshModelingToolset/Source/ModelingComponentsEditorOnly/Public/Operations/SubdividePoly.h b/Engine/Plugins/Runtime/MeshModelingToolset/Source/ModelingComponentsEditorOnly/Public/Operations/SubdividePoly.h
--- a/Engine/Plugins/Runtime/MeshModelingToolset/Source/ModelingComponentsEditorOnly/Public/Operations/SubdividePoly.h
+++ b/Engine/Plugins/Runtime/MeshModelingToolset/Source/ModelingComponentsEditorOnly/Public/Operations/SubdividePoly.h
@@ -6,6 +6,8 @@
 #include "ModelingOperators.h"
 #include "SubdividePoly.generated.h"
 
+#define UE_API MODELINGCOMPONENTSEDITORONLY_API
+
 PREDECLARE_USE_GEOMETRY_CLASS(FGroupTopology);
 PREDECLARE_USE_GEOMETRY_CLASS(FDynamicMesh3);
 class FProgressCancel;
@@ -83,15 +85,15 @@ enum class ESubdivisionOutputUVs : uint8
  * precompute the refined topology, then compute the interpolated positions (and other vertex values) as the original 
  * polymesh is deformed.
  */
-class MODELINGCOMPONENTSEDITORONLY_API FSubdividePoly
+class FSubdividePoly
 {
 public:
 
-	FSubdividePoly(const FGroupTopology& InTopology,
+	UE_API FSubdividePoly(const FGroupTopology& InTopology,
 				   const FDynamicMesh3& InOriginalMesh,
 				   int InLevel);
 
-	~FSubdividePoly();
+	UE_API ~FSubdividePoly();
 
 	enum class ETopologyCheckResult
 	{
@@ -102,11 +104,11 @@ public:
 		MultiBoundaryPolygroup = 4,
 		DegeneratePolygroup = 5
 	};
-	ETopologyCheckResult ValidateTopology();
+	UE_API ETopologyCheckResult ValidateTopology();
 
-	bool ComputeTopologySubdivision();
+	UE_API bool ComputeTopologySubdivision();
 
-	bool ComputeSubdividedMesh(FDynamicMesh3& OutMesh);
+	UE_API bool ComputeSubdividedMesh(FDynamicMesh3& OutMesh);
 
 	ESubdivisionScheme SubdivisionScheme = ESubdivisionScheme::CatmullClark;
 
@@ -130,3 +132,4 @@ private:
 	TUniquePtr<RefinerImpl> Refiner;
 };
 
+#undef UE_API
diff --git a/Engine/Plugins/Runtime/MeshModelingToolset/Source/ModelingComponentsEditorOnly/Public/Scene/LevelObjectsObserver.h b/Engine/Plugins/Runtime/MeshModelingToolset/Source/ModelingComponentsEditorOnly/Public/Scene/LevelObjectsObserver.h
--- a/Engine/Plugins/Runtime/MeshModelingToolset/Source/ModelingComponentsEditorOnly/Public/Scene/LevelObjectsObserver.h
+++ b/Engine/Plugins/Runtime/MeshModelingToolset/Source/ModelingComponentsEditorOnly/Public/Scene/LevelObjectsObserver.h
@@ -4,6 +4,8 @@
 
 #include "EditorUndoClient.h"
 
+#define UE_API MODELINGCOMPONENTSEDITORONLY_API
+
 class UWorld;
 class AActor;
 
@@ -23,27 +25,27 @@ class AActor;
  * to know this information.
  * 
  */
-class MODELINGCOMPONENTSEDITORONLY_API FLevelObjectsObserver : public FEditorUndoClient
+class FLevelObjectsObserver : public FEditorUndoClient
 {
 public:
 	/**
 	 * Destructor.
 	 */
-	virtual ~FLevelObjectsObserver();
+	UE_API virtual ~FLevelObjectsObserver();
 
 	/**
 	 * Begin watching WorldIn. This will emit OnActorAdded for all Actors in the World.
 	 */
-	void Initialize(UWorld* WorldIn);
+	UE_API void Initialize(UWorld* WorldIn);
 
 	/**
 	 * Stop watching WorldIn. This will emit OnActorRemoved for all Actors that the Observer is currently aware of.
 	 */
-	void Shutdown();
+	UE_API void Shutdown();
 
 	// FEditorUndoClient implementation
-	virtual void PostUndo(bool bSuccess) override;
-	virtual void PostRedo(bool bSuccess) override;
+	UE_API virtual void PostUndo(bool bSuccess) override;
+	UE_API virtual void PostRedo(bool bSuccess) override;
 
 	DECLARE_MULTICAST_DELEGATE_OneParam(FActorAddRemoveSignature, AActor*);
 
@@ -66,7 +68,9 @@ protected:
 	FDelegateHandle OnActorAddedHandle;
 	FDelegateHandle OnActorListChangedHandle;
 
-	void OnUntrackedLevelChange();
-	void HandleActorAddedEvent(AActor* Actor);
-	void HandleActorDeletedEvent(AActor* Actor);
+	UE_API void OnUntrackedLevelChange();
+	UE_API void HandleActorAddedEvent(AActor* Actor);
+	UE_API void HandleActorDeletedEvent(AActor* Actor);
 };
+
+#undef UE_API
diff --git a/Engine/Plugins/Runtime/MeshModelingToolset/Source/ModelingComponentsEditorOnly/Public/Selection/StaticMeshSelector.h b/Engine/Plugins/Runtime/MeshModelingToolset/Source/ModelingComponentsEditorOnly/Public/Selection/StaticMeshSelector.h
--- a/Engine/Plugins/Runtime/MeshModelingToolset/Source/ModelingComponentsEditorOnly/Public/Selection/StaticMeshSelector.h
+++ b/Engine/Plugins/Runtime/MeshModelingToolset/Source/ModelingComponentsEditorOnly/Public/Selection/StaticMeshSelector.h
@@ -8,40 +8,42 @@
 
 #include "UObject/StrongObjectPtr.h"
 
+#define UE_API MODELINGCOMPONENTSEDITORONLY_API
+
 
 class UDynamicMesh;
 class UStaticMesh;
 class UStaticMeshComponent;
 
 
-class MODELINGCOMPONENTSEDITORONLY_API FStaticMeshSelector : public FBaseDynamicMeshSelector
+class FStaticMeshSelector : public FBaseDynamicMeshSelector
 {
 public:
 	using FBaseDynamicMeshSelector::Initialize;
 
-	virtual bool Initialize(
+	UE_API virtual bool Initialize(
 		FGeometryIdentifier SourceGeometryIdentifier);
 
 	//
 	// IGeometrySelector API implementation
 	//
 
-	virtual void Shutdown() override;
+	UE_API virtual void Shutdown() override;
 	// disable sleep on StaticMesh Selector until we can properly track changes...
 	virtual bool SupportsSleep() const override { return false; }
 	//virtual bool Sleep();
 	//virtual bool Restore();
 
 
-	virtual bool IsLockable() const override;
-	virtual bool IsLocked() const override;
-	virtual void SetLockedState(bool bLocked) override;
+	UE_API virtual bool IsLockable() const override;
+	UE_API virtual bool IsLocked() const override;
+	UE_API virtual void SetLockedState(bool bLocked) override;
 
 
-	virtual IGeometrySelectionTransformer* InitializeTransformation(const FGeometrySelection& Selection) override;
-	virtual void ShutdownTransformation(IGeometrySelectionTransformer* Transformer) override;
+	UE_API virtual IGeometrySelectionTransformer* InitializeTransformation(const FGeometrySelection& Selection) override;
+	UE_API virtual void ShutdownTransformation(IGeometrySelectionTransformer* Transformer) override;
 
-	virtual void UpdateAfterGeometryEdit(
+	UE_API virtual void UpdateAfterGeometryEdit(
 		IToolsContextTransactionsAPI* TransactionsAPI,
 		bool bInTransaction,
 		TUniquePtr<UE::Geometry::FDynamicMeshChange> DynamicMeshChange,
@@ -54,15 +56,15 @@ protected:
 
 	// TODO: this is not a great design, it would be better if something external could own this mesh...
 	TStrongObjectPtr<UDynamicMesh> LocalTargetMesh;
-	void CopyFromStaticMesh();
+	UE_API void CopyFromStaticMesh();
 
 	TSharedPtr<FBasicDynamicMeshSelectionTransformer> ActiveTransformer;
-	void CommitMeshTransform();
+	UE_API void CommitMeshTransform();
 
 
 public:
-	static void SetAssetUnlockedOnCreation(UStaticMesh* StaticMesh);
-	static void ResetUnlockedStaticMeshAssets();
+	static UE_API void SetAssetUnlockedOnCreation(UStaticMesh* StaticMesh);
+	static UE_API void ResetUnlockedStaticMeshAssets();
 };
 
 
@@ -71,13 +73,14 @@ public:
  * FStaticMeshComponentSelectorFactory constructs FStaticMeshSelector instances 
  * for UStaticMeshComponents
  */
-class MODELINGCOMPONENTSEDITORONLY_API FStaticMeshComponentSelectorFactory : public IGeometrySelectorFactory
+class FStaticMeshComponentSelectorFactory : public IGeometrySelectorFactory
 {
 public:
-	virtual bool CanBuildForTarget(FGeometryIdentifier TargetIdentifier) const;
-	virtual TUniquePtr<IGeometrySelector> BuildForTarget(FGeometryIdentifier TargetIdentifier) const;
+	UE_API virtual bool CanBuildForTarget(FGeometryIdentifier TargetIdentifier) const;
+	UE_API virtual TUniquePtr<IGeometrySelector> BuildForTarget(FGeometryIdentifier TargetIdentifier) const;
 };
 
 
 
 
+#undef UE_API
diff --git a/Engine/Plugins/Runtime/MeshModelingToolset/Source/ModelingComponentsEditorOnly/Public/Selection/VolumeSelector.h b/Engine/Plugins/Runtime/MeshModelingToolset/Source/ModelingComponentsEditorOnly/Public/Selection/VolumeSelector.h
--- a/Engine/Plugins/Runtime/MeshModelingToolset/Source/ModelingComponentsEditorOnly/Public/Selection/VolumeSelector.h
+++ b/Engine/Plugins/Runtime/MeshModelingToolset/Source/ModelingComponentsEditorOnly/Public/Selection/VolumeSelector.h
@@ -9,37 +9,39 @@
 #include "UObject/StrongObjectPtr.h"
 #include "EditorUndoClient.h"
 
+#define UE_API MODELINGCOMPONENTSEDITORONLY_API
+
 class UDynamicMesh;
 class AVolume;
 class UBrushComponent;
 
 
-class MODELINGCOMPONENTSEDITORONLY_API FVolumeSelector : public FBaseDynamicMeshSelector, public FEditorUndoClient
+class FVolumeSelector : public FBaseDynamicMeshSelector, public FEditorUndoClient
 {
 public:
 	using FBaseDynamicMeshSelector::Initialize;
 
-	virtual bool Initialize(
+	UE_API virtual bool Initialize(
 		FGeometryIdentifier SourceGeometryIdentifier);
 
 	//
 	// IGeometrySelector API implementation
 	//
 
-	virtual void Shutdown() override;
+	UE_API virtual void Shutdown() override;
 	// disable sleep on VolumeSelector until we can properly track changes...
 	virtual bool SupportsSleep() const override { return false; }
 	//virtual bool Sleep();
 	//virtual bool Restore();
 
-	virtual bool IsLockable() const override;
-	virtual bool IsLocked() const override;
-	virtual void SetLockedState(bool bLocked) override;
+	UE_API virtual bool IsLockable() const override;
+	UE_API virtual bool IsLocked() const override;
+	UE_API virtual void SetLockedState(bool bLocked) override;
 
-	virtual IGeometrySelectionTransformer* InitializeTransformation(const FGeometrySelection& Selection) override;
-	virtual void ShutdownTransformation(IGeometrySelectionTransformer* Transformer) override;
+	UE_API virtual IGeometrySelectionTransformer* InitializeTransformation(const FGeometrySelection& Selection) override;
+	UE_API virtual void ShutdownTransformation(IGeometrySelectionTransformer* Transformer) override;
 
-	virtual void UpdateAfterGeometryEdit(
+	UE_API virtual void UpdateAfterGeometryEdit(
 		IToolsContextTransactionsAPI* TransactionsAPI,
 		bool bInTransaction,
 		TUniquePtr<UE::Geometry::FDynamicMeshChange> DynamicMeshChange,
@@ -47,12 +49,12 @@ public:
 
 
 	// FEditorUndoClient implementation
-	virtual void PostUndo(bool bSuccess) override;
-	virtual void PostRedo(bool bSuccess) override;
+	UE_API virtual void PostUndo(bool bSuccess) override;
+	UE_API virtual void PostRedo(bool bSuccess) override;
 
 public:
-	static void SetComponentUnlockedOnCreation(UBrushComponent* Component);
-	static void ResetUnlockedBrushComponents();
+	static UE_API void SetComponentUnlockedOnCreation(UBrushComponent* Component);
+	static UE_API void ResetUnlockedBrushComponents();
 
 protected:
 	AVolume* ParentVolume = nullptr;
@@ -61,10 +63,10 @@ protected:
 	// TODO: this is not a great design, it would be better if something external could own this mesh...
 	TStrongObjectPtr<UDynamicMesh> LocalTargetMesh;
 
-	void UpdateDynamicMeshFromVolume();
+	UE_API void UpdateDynamicMeshFromVolume();
 
 	TSharedPtr<FBasicDynamicMeshSelectionTransformer> ActiveTransformer;
-	void CommitMeshTransform();
+	UE_API void CommitMeshTransform();
 };
 
 
@@ -74,13 +76,14 @@ protected:
  * FVolumeComponentSelectorFactory constructs FVolumeSelector instances 
  * for UBrushComponents
  */
-class MODELINGCOMPONENTSEDITORONLY_API FBrushComponentSelectorFactory : public IGeometrySelectorFactory
+class FBrushComponentSelectorFactory : public IGeometrySelectorFactory
 {
 public:
-	virtual bool CanBuildForTarget(FGeometryIdentifier TargetIdentifier) const;
-	virtual TUniquePtr<IGeometrySelector> BuildForTarget(FGeometryIdentifier TargetIdentifier) const;
+	UE_API virtual bool CanBuildForTarget(FGeometryIdentifier TargetIdentifier) const;
+	UE_API virtual TUniquePtr<IGeometrySelector> BuildForTarget(FGeometryIdentifier TargetIdentifier) const;
 };
 
 
 
 
+#undef UE_API
diff --git a/Engine/Plugins/Runtime/MeshModelingToolset/Source/ModelingComponentsEditorOnly/Public/ToolTargets/DynamicMeshComponentToolTarget.h b/Engine/Plugins/Runtime/MeshModelingToolset/Source/ModelingComponentsEditorOnly/Public/ToolTargets/DynamicMeshComponentToolTarget.h
--- a/Engine/Plugins/Runtime/MeshModelingToolset/Source/ModelingComponentsEditorOnly/Public/ToolTargets/DynamicMeshComponentToolTarget.h
+++ b/Engine/Plugins/Runtime/MeshModelingToolset/Source/ModelingComponentsEditorOnly/Public/ToolTargets/DynamicMeshComponentToolTarget.h
@@ -15,13 +15,15 @@
 
 #include "DynamicMeshComponentToolTarget.generated.h"
 
+#define UE_API MODELINGCOMPONENTSEDITORONLY_API
+
 class UDynamicMesh;
 
 /**
  * A ToolTarget backed by a DynamicMeshComponent
  */
-UCLASS(Transient)
-class MODELINGCOMPONENTSEDITORONLY_API UDynamicMeshComponentToolTarget : 
+UCLASS(MinimalAPI, Transient)
+class UDynamicMeshComponentToolTarget : 
 	public UPrimitiveComponentToolTarget,
 	public IMeshDescriptionCommitter, 
 	public IMeshDescriptionProvider, 
@@ -34,40 +36,40 @@ class MODELINGCOMPONENTSEDITORONLY_API UDynamicMeshComponentToolTarget :
 	GENERATED_BODY()
 
 public:
-	virtual bool IsValid() const override;
+	UE_API virtual bool IsValid() const override;
 
 
 public:
 	// IMeshDescriptionProvider implementation
-	const FMeshDescription* GetMeshDescription(const FGetMeshParameters& GetMeshParams = FGetMeshParameters()) override;
-	virtual FMeshDescription GetEmptyMeshDescription() override;
+	UE_API const FMeshDescription* GetMeshDescription(const FGetMeshParameters& GetMeshParams = FGetMeshParameters()) override;
+	UE_API virtual FMeshDescription GetEmptyMeshDescription() override;
 
 	// IMeshDescritpionCommitter implementation
-	virtual void CommitMeshDescription(const FCommitter& Committer, const FCommitMeshParameters& CommitMeshParams = FCommitMeshParameters()) override;
+	UE_API virtual void CommitMeshDescription(const FCommitter& Committer, const FCommitMeshParameters& CommitMeshParams = FCommitMeshParameters()) override;
 	using IMeshDescriptionCommitter::CommitMeshDescription; // unhide the other overload
 
 	// IDynamicMeshProvider implementation
-	virtual UE::Geometry::FDynamicMesh3 GetDynamicMesh() override;
+	UE_API virtual UE::Geometry::FDynamicMesh3 GetDynamicMesh() override;
 
 	// IDynamicMeshCommitter implementation
-	virtual void CommitDynamicMesh(const UE::Geometry::FDynamicMesh3& Mesh, const FDynamicMeshCommitInfo& CommitInfo) override;
+	UE_API virtual void CommitDynamicMesh(const UE::Geometry::FDynamicMesh3& Mesh, const FDynamicMeshCommitInfo& CommitInfo) override;
 	using IDynamicMeshCommitter::CommitDynamicMesh; // unhide the other overload
 
 	// IMaterialProvider implementation
-	int32 GetNumMaterials() const override;
-	UMaterialInterface* GetMaterial(int32 MaterialIndex) const override;
-	void GetMaterialSet(FComponentMaterialSet& MaterialSetOut, bool bPreferAssetMaterials) const override;
-	virtual bool CommitMaterialSetUpdate(const FComponentMaterialSet& MaterialSet, bool bApplyToAsset) override;	
+	UE_API int32 GetNumMaterials() const override;
+	UE_API UMaterialInterface* GetMaterial(int32 MaterialIndex) const override;
+	UE_API void GetMaterialSet(FComponentMaterialSet& MaterialSetOut, bool bPreferAssetMaterials) const override;
+	UE_API virtual bool CommitMaterialSetUpdate(const FComponentMaterialSet& MaterialSet, bool bApplyToAsset) override;	
 
 	// IPersistentDynamicMeshSource implementation
-	virtual UDynamicMesh* GetDynamicMeshContainer() override;
-	virtual void CommitDynamicMeshChange(TUniquePtr<FToolCommandChange> Change, const FText& ChangeMessage) override;
-	virtual bool HasDynamicMeshComponent() const override;
-	virtual UDynamicMeshComponent* GetDynamicMeshComponent() override;
+	UE_API virtual UDynamicMesh* GetDynamicMeshContainer() override;
+	UE_API virtual void CommitDynamicMeshChange(TUniquePtr<FToolCommandChange> Change, const FText& ChangeMessage) override;
+	UE_API virtual bool HasDynamicMeshComponent() const override;
+	UE_API virtual UDynamicMeshComponent* GetDynamicMeshComponent() override;
 
 	// IPhysicsDataSource implementation
-	virtual UBodySetup* GetBodySetup() const override;
-	virtual IInterface_CollisionDataProvider* GetComplexCollisionProvider() const override;
+	UE_API virtual UBodySetup* GetBodySetup() const override;
+	UE_API virtual IInterface_CollisionDataProvider* GetComplexCollisionProvider() const override;
 
 	// Rest provided by parent class
 
@@ -81,14 +83,16 @@ protected:
 
 
 /** Factory for UDynamicMeshComponentToolTarget to be used by the target manager. */
-UCLASS(Transient)
-class MODELINGCOMPONENTSEDITORONLY_API UDynamicMeshComponentToolTargetFactory : public UToolTargetFactory
+UCLASS(MinimalAPI, Transient)
+class UDynamicMeshComponentToolTargetFactory : public UToolTargetFactory
 {
 	GENERATED_BODY()
 
 public:
 
-	virtual bool CanBuildTarget(UObject* SourceObject, const FToolTargetTypeRequirements& TargetTypeInfo) const override;
+	UE_API virtual bool CanBuildTarget(UObject* SourceObject, const FToolTargetTypeRequirements& TargetTypeInfo) const override;
+
+	UE_API virtual UToolTarget* BuildTarget(UObject* SourceObject, const FToolTargetTypeRequirements& TargetTypeInfo) override;
+};
 
-	virtual UToolTarget* BuildTarget(UObject* SourceObject, const FToolTargetTypeRequirements& TargetTypeInfo) override;
-};
\ No newline at end of file
+#undef UE_API
diff --git a/Engine/Plugins/Runtime/MeshModelingToolset/Source/ModelingComponentsEditorOnly/Public/ToolTargets/SkeletalMeshComponentToolTarget.h b/Engine/Plugins/Runtime/MeshModelingToolset/Source/ModelingComponentsEditorOnly/Public/ToolTargets/SkeletalMeshComponentToolTarget.h
--- a/Engine/Plugins/Runtime/MeshModelingToolset/Source/ModelingComponentsEditorOnly/Public/ToolTargets/SkeletalMeshComponentToolTarget.h
+++ b/Engine/Plugins/Runtime/MeshModelingToolset/Source/ModelingComponentsEditorOnly/Public/ToolTargets/SkeletalMeshComponentToolTarget.h
@@ -14,6 +14,8 @@
 
 #include "SkeletalMeshComponentToolTarget.generated.h"
 
+#define UE_API MODELINGCOMPONENTSEDITORONLY_API
+
 
 class USkeletalMesh;
 
@@ -22,8 +24,8 @@ class USkeletalMesh;
  * A tool target backed by a skeletal mesh component that can provide and take a mesh
  * description, without 'committer' interfaces for writing to the source skeletal mesh asset.
  */
-UCLASS(Transient)
-class MODELINGCOMPONENTSEDITORONLY_API USkeletalMeshComponentReadOnlyToolTarget :
+UCLASS(MinimalAPI, Transient)
+class USkeletalMeshComponentReadOnlyToolTarget :
 	public UPrimitiveComponentToolTarget,
 	public IMeshDescriptionProvider,
 	public IDynamicMeshProvider, 
@@ -35,28 +37,28 @@ class MODELINGCOMPONENTSEDITORONLY_API USkeletalMeshComponentReadOnlyToolTarget
 public:
 
 	// UToolTarget implementation
-	virtual bool IsValid() const override;
+	UE_API virtual bool IsValid() const override;
 
 	// IMeshDescriptionProvider implementation
-	virtual const FMeshDescription* GetMeshDescription(const FGetMeshParameters& GetMeshParams = FGetMeshParameters()) override;
-	virtual FMeshDescription GetEmptyMeshDescription() override;
+	UE_API virtual const FMeshDescription* GetMeshDescription(const FGetMeshParameters& GetMeshParams = FGetMeshParameters()) override;
+	UE_API virtual FMeshDescription GetEmptyMeshDescription() override;
 	virtual bool SupportsLODs() const override { return true; }
-	virtual TArray<EMeshLODIdentifier> GetAvailableLODs(bool bSkipAutoGenerated = true) const override;
+	UE_API virtual TArray<EMeshLODIdentifier> GetAvailableLODs(bool bSkipAutoGenerated = true) const override;
 	virtual EMeshLODIdentifier GetMeshDescriptionLOD() const override { return static_cast<EMeshLODIdentifier>(DefaultEditingLOD); }
 
 	
 	// IMaterialProvider implementation
-	virtual int32 GetNumMaterials() const override;
-	virtual UMaterialInterface* GetMaterial(int32 MaterialIndex) const override;
-	virtual void GetMaterialSet(FComponentMaterialSet& MaterialSetOut, bool bPreferAssetMaterials) const override;
-	virtual bool CommitMaterialSetUpdate(const FComponentMaterialSet& MaterialSet, bool bApplyToAsset) override;
+	UE_API virtual int32 GetNumMaterials() const override;
+	UE_API virtual UMaterialInterface* GetMaterial(int32 MaterialIndex) const override;
+	UE_API virtual void GetMaterialSet(FComponentMaterialSet& MaterialSetOut, bool bPreferAssetMaterials) const override;
+	UE_API virtual bool CommitMaterialSetUpdate(const FComponentMaterialSet& MaterialSet, bool bApplyToAsset) override;
 
 	// IDynamicMeshProvider
-	virtual UE::Geometry::FDynamicMesh3 GetDynamicMesh() override;
-	virtual UE::Geometry::FDynamicMesh3 GetDynamicMesh(const FGetMeshParameters& InGetMeshParams) override;
+	UE_API virtual UE::Geometry::FDynamicMesh3 GetDynamicMesh() override;
+	UE_API virtual UE::Geometry::FDynamicMesh3 GetDynamicMesh(const FGetMeshParameters& InGetMeshParams) override;
 
 	// ISkeletalMeshBackedTarget implementation
-	virtual USkeletalMesh* GetSkeletalMesh() const override;
+	UE_API virtual USkeletalMesh* GetSkeletalMesh() const override;
 
 protected:
 	// The LOD to edit if requested to get/commit at EMeshLODIdentifier::Default or if no specific LOD is requested.
@@ -72,8 +74,8 @@ protected:
  * A tool target backed by a skeletal mesh component that can provide and take a mesh
  * description.
  */
-UCLASS(Transient)
-class MODELINGCOMPONENTSEDITORONLY_API USkeletalMeshComponentToolTarget :
+UCLASS(MinimalAPI, Transient)
+class USkeletalMeshComponentToolTarget :
 	public USkeletalMeshComponentReadOnlyToolTarget,
 	public IMeshDescriptionCommitter,
 	public IDynamicMeshCommitter
@@ -82,11 +84,11 @@ class MODELINGCOMPONENTSEDITORONLY_API USkeletalMeshComponentToolTarget :
 
 public:
 	// IMeshDescriptionCommitter implementation
-	virtual void CommitMeshDescription(const FCommitter& Committer, const FCommitMeshParameters& CommitParams = FCommitMeshParameters()) override;
+	UE_API virtual void CommitMeshDescription(const FCommitter& Committer, const FCommitMeshParameters& CommitParams = FCommitMeshParameters()) override;
 	using IMeshDescriptionCommitter::CommitMeshDescription; // unhide the other overload
 
 	// IDynamicMeshCommitter
-	virtual void CommitDynamicMesh(const UE::Geometry::FDynamicMesh3& Mesh, const FDynamicMeshCommitInfo& CommitInfo) override;
+	UE_API virtual void CommitDynamicMesh(const UE::Geometry::FDynamicMesh3& Mesh, const FDynamicMeshCommitInfo& CommitInfo) override;
 	using IDynamicMeshCommitter::CommitDynamicMesh; // unhide the other overload
 
 protected:
@@ -96,30 +98,32 @@ protected:
 
 
 /** Factory for USkeletalMeshComponentReadOnlyToolTarget to be used by the target manager. */
-UCLASS(Transient)
-class MODELINGCOMPONENTSEDITORONLY_API USkeletalMeshComponentReadOnlyToolTargetFactory : public UToolTargetFactory
+UCLASS(MinimalAPI, Transient)
+class USkeletalMeshComponentReadOnlyToolTargetFactory : public UToolTargetFactory
 {
 	GENERATED_BODY()
 
 public:
 
-	virtual bool CanBuildTarget(UObject* SourceObject, const FToolTargetTypeRequirements& TargetTypeInfo) const override;
+	UE_API virtual bool CanBuildTarget(UObject* SourceObject, const FToolTargetTypeRequirements& TargetTypeInfo) const override;
 
-	virtual UToolTarget* BuildTarget(UObject* SourceObject, const FToolTargetTypeRequirements& TargetTypeInfo) override;
+	UE_API virtual UToolTarget* BuildTarget(UObject* SourceObject, const FToolTargetTypeRequirements& TargetTypeInfo) override;
 };
 
 
 /** Factory for USkeletalMeshComponentToolTarget to be used by the target manager. */
-UCLASS(Transient)
-class MODELINGCOMPONENTSEDITORONLY_API USkeletalMeshComponentToolTargetFactory : public UToolTargetFactory
+UCLASS(MinimalAPI, Transient)
+class USkeletalMeshComponentToolTargetFactory : public UToolTargetFactory
 {
 	GENERATED_BODY()
 
 public:
 
-	virtual bool CanBuildTarget(UObject* SourceObject, const FToolTargetTypeRequirements& TargetTypeInfo) const override;
+	UE_API virtual bool CanBuildTarget(UObject* SourceObject, const FToolTargetTypeRequirements& TargetTypeInfo) const override;
 
-	virtual UToolTarget* BuildTarget(UObject* SourceObject, const FToolTargetTypeRequirements& TargetTypeInfo) override;
+	UE_API virtual UToolTarget* BuildTarget(UObject* SourceObject, const FToolTargetTypeRequirements& TargetTypeInfo) override;
 
-	static bool CanWriteToSource(const UObject* SourceObject);
+	static UE_API bool CanWriteToSource(const UObject* SourceObject);
 };
+
+#undef UE_API
diff --git a/Engine/Plugins/Runtime/MeshModelingToolset/Source/ModelingComponentsEditorOnly/Public/ToolTargets/SkeletalMeshToolTarget.h b/Engine/Plugins/Runtime/MeshModelingToolset/Source/ModelingComponentsEditorOnly/Public/ToolTargets/SkeletalMeshToolTarget.h
--- a/Engine/Plugins/Runtime/MeshModelingToolset/Source/ModelingComponentsEditorOnly/Public/ToolTargets/SkeletalMeshToolTarget.h
+++ b/Engine/Plugins/Runtime/MeshModelingToolset/Source/ModelingComponentsEditorOnly/Public/ToolTargets/SkeletalMeshToolTarget.h
@@ -14,6 +14,8 @@
 
 #include "SkeletalMeshToolTarget.generated.h"
 
+#define UE_API MODELINGCOMPONENTSEDITORONLY_API
+
 
 class USkeletalMesh;
 
@@ -23,8 +25,8 @@ class USkeletalMesh;
  * This is a special tool target that refers to the underlying asset (in this case a skeletal mesh), rather than indirectly through a component.
  * This type of target is used in cases, such as opening an asset through the content browser, when there is no component available.
  */
-UCLASS(Transient)
-class MODELINGCOMPONENTSEDITORONLY_API USkeletalMeshReadOnlyToolTarget :
+UCLASS(MinimalAPI, Transient)
+class USkeletalMeshReadOnlyToolTarget :
 	public UToolTarget,
 	public IMeshDescriptionProvider,
 	public IDynamicMeshProvider,
@@ -35,27 +37,27 @@ class MODELINGCOMPONENTSEDITORONLY_API USkeletalMeshReadOnlyToolTarget :
 
 public:
 	// UToolTarget
-	virtual bool IsValid() const override;
+	UE_API virtual bool IsValid() const override;
 
 	// IMeshDescriptionProvider implementation
-	virtual const FMeshDescription* GetMeshDescription(const FGetMeshParameters& GetMeshParams = FGetMeshParameters()) override;
-	virtual FMeshDescription GetEmptyMeshDescription() override;
+	UE_API virtual const FMeshDescription* GetMeshDescription(const FGetMeshParameters& GetMeshParams = FGetMeshParameters()) override;
+	UE_API virtual FMeshDescription GetEmptyMeshDescription() override;
 	virtual bool SupportsLODs() const override { return true; }
-	virtual TArray<EMeshLODIdentifier> GetAvailableLODs(bool bSkipAutoGenerated = true) const override;
+	UE_API virtual TArray<EMeshLODIdentifier> GetAvailableLODs(bool bSkipAutoGenerated = true) const override;
 	virtual EMeshLODIdentifier GetMeshDescriptionLOD() const override { return static_cast<EMeshLODIdentifier>(DefaultEditingLOD); }
 
 	// IMaterialProvider implementation
-	virtual int32 GetNumMaterials() const override;
-	virtual UMaterialInterface* GetMaterial(int32 MaterialIndex) const override;
-	virtual void GetMaterialSet(FComponentMaterialSet& MaterialSetOut, bool bPreferAssetMaterials) const override;
-	virtual bool CommitMaterialSetUpdate(const FComponentMaterialSet& MaterialSet, bool bApplyToAsset) override;	
+	UE_API virtual int32 GetNumMaterials() const override;
+	UE_API virtual UMaterialInterface* GetMaterial(int32 MaterialIndex) const override;
+	UE_API virtual void GetMaterialSet(FComponentMaterialSet& MaterialSetOut, bool bPreferAssetMaterials) const override;
+	UE_API virtual bool CommitMaterialSetUpdate(const FComponentMaterialSet& MaterialSet, bool bApplyToAsset) override;	
 
 	// IDynamicMeshProvider
-	virtual UE::Geometry::FDynamicMesh3 GetDynamicMesh() override;
-	virtual UE::Geometry::FDynamicMesh3 GetDynamicMesh(const FGetMeshParameters& InGetMeshParams) override;
+	UE_API virtual UE::Geometry::FDynamicMesh3 GetDynamicMesh() override;
+	UE_API virtual UE::Geometry::FDynamicMesh3 GetDynamicMesh(const FGetMeshParameters& InGetMeshParams) override;
 
 	// ISkeletalMeshBackedTarget implementation
-	virtual USkeletalMesh* GetSkeletalMesh() const override;
+	UE_API virtual USkeletalMesh* GetSkeletalMesh() const override;
 
 protected:
 	TWeakObjectPtr<USkeletalMesh> SkeletalMesh = nullptr;
@@ -70,22 +72,22 @@ protected:
 	friend class USkeletalMeshComponentToolTarget;
 
 	UE_DEPRECATED(5.3, "This function doesn't check anything skeletal mesh specific, only (badly) checks to see if the object is alive")
-	static bool IsValid(const USkeletalMesh* USkeletalMesh);
+	static UE_API bool IsValid(const USkeletalMesh* USkeletalMesh);
 
-	static TArray<EMeshLODIdentifier> GetAvailableLODs(const USkeletalMesh* SkeletalMesh, bool bSkipAutoGenerated);	
-	static void GetMaterialSet(const USkeletalMesh* SkeletalMesh, FComponentMaterialSet& MaterialSetOut,
+	static UE_API TArray<EMeshLODIdentifier> GetAvailableLODs(const USkeletalMesh* SkeletalMesh, bool bSkipAutoGenerated);	
+	static UE_API void GetMaterialSet(const USkeletalMesh* SkeletalMesh, FComponentMaterialSet& MaterialSetOut,
 		bool bPreferAssetMaterials);
-	static bool CommitMaterialSetUpdate(USkeletalMesh* SkeletalMesh,
+	static UE_API bool CommitMaterialSetUpdate(USkeletalMesh* SkeletalMesh,
 		const FComponentMaterialSet& MaterialSet, bool bApplyToAsset);
-	static int32 GetValidEditingLOD(const USkeletalMesh* SkeletalMesh, int32 DefaultEditingLOD, bool bHaveRequestLOD, EMeshLODIdentifier RequestedEditingLOD);	
+	static UE_API int32 GetValidEditingLOD(const USkeletalMesh* SkeletalMesh, int32 DefaultEditingLOD, bool bHaveRequestLOD, EMeshLODIdentifier RequestedEditingLOD);	
 };
 
 
 /**
  * A tool target backed by a skeletal mesh.
  */
-UCLASS(Transient)
-class MODELINGCOMPONENTSEDITORONLY_API USkeletalMeshToolTarget :
+UCLASS(MinimalAPI, Transient)
+class USkeletalMeshToolTarget :
 	public USkeletalMeshReadOnlyToolTarget,
 	public IMeshDescriptionCommitter,
 	public IDynamicMeshCommitter
@@ -94,11 +96,11 @@ class MODELINGCOMPONENTSEDITORONLY_API USkeletalMeshToolTarget :
 
 public:
 	// IMeshDescriptionCommitter implementation
-	virtual void CommitMeshDescription(const FCommitter& Committer, const FCommitMeshParameters& CommitMeshParams = FCommitMeshParameters()) override;
+	UE_API virtual void CommitMeshDescription(const FCommitter& Committer, const FCommitMeshParameters& CommitMeshParams = FCommitMeshParameters()) override;
 	using IMeshDescriptionCommitter::CommitMeshDescription; // unhide the other overload
 
 	// IDynamicMeshCommitter
-	virtual void CommitDynamicMesh(const UE::Geometry::FDynamicMesh3& Mesh, 
+	UE_API virtual void CommitDynamicMesh(const UE::Geometry::FDynamicMesh3& Mesh, 
 		const FDynamicMeshCommitInfo& CommitInfo) override;
 	using IDynamicMeshCommitter::CommitDynamicMesh; // unhide the other overload
 
@@ -107,35 +109,37 @@ protected:
 	friend class USkeletalMeshToolTargetFactory;
 	friend class USkeletalMeshComponentToolTarget;
 
-	static void CommitMeshDescription(USkeletalMesh* SkeletalMesh, const FCommitter& Committer, int32 LODIndex);
+	static UE_API void CommitMeshDescription(USkeletalMesh* SkeletalMesh, const FCommitter& Committer, int32 LODIndex);
 };
 
 
 /** Factory for USkeletalMeshReadOnlyToolTarget to be used by the target manager. */
-UCLASS(Transient)
-class MODELINGCOMPONENTSEDITORONLY_API USkeletalMeshReadOnlyToolTargetFactory : public UToolTargetFactory
+UCLASS(MinimalAPI, Transient)
+class USkeletalMeshReadOnlyToolTargetFactory : public UToolTargetFactory
 {
 	GENERATED_BODY()
 
 public:
 
-	virtual bool CanBuildTarget(UObject* SourceObject, const FToolTargetTypeRequirements& TargetTypeInfo) const override;
+	UE_API virtual bool CanBuildTarget(UObject* SourceObject, const FToolTargetTypeRequirements& TargetTypeInfo) const override;
 
-	virtual UToolTarget* BuildTarget(UObject* SourceObject, const FToolTargetTypeRequirements& TargetTypeInfo) override;
+	UE_API virtual UToolTarget* BuildTarget(UObject* SourceObject, const FToolTargetTypeRequirements& TargetTypeInfo) override;
 };
 
 
 /** Factory for USkeletalMeshToolTarget to be used by the target manager. */
-UCLASS(Transient)
-class MODELINGCOMPONENTSEDITORONLY_API USkeletalMeshToolTargetFactory : public UToolTargetFactory
+UCLASS(MinimalAPI, Transient)
+class USkeletalMeshToolTargetFactory : public UToolTargetFactory
 {
 	GENERATED_BODY()
 
 public:
 
-	virtual bool CanBuildTarget(UObject* SourceObject, const FToolTargetTypeRequirements& TargetTypeInfo) const override;
+	UE_API virtual bool CanBuildTarget(UObject* SourceObject, const FToolTargetTypeRequirements& TargetTypeInfo) const override;
 
-	virtual UToolTarget* BuildTarget(UObject* SourceObject, const FToolTargetTypeRequirements& TargetTypeInfo) override;
+	UE_API virtual UToolTarget* BuildTarget(UObject* SourceObject, const FToolTargetTypeRequirements& TargetTypeInfo) override;
 
-	static bool CanWriteToSource(const UObject* SourceObject);
+	static UE_API bool CanWriteToSource(const UObject* SourceObject);
 };
+
+#undef UE_API
diff --git a/Engine/Plugins/Runtime/MeshModelingToolset/Source/ModelingComponentsEditorOnly/Public/ToolTargets/StaticMeshComponentToolTarget.h b/Engine/Plugins/Runtime/MeshModelingToolset/Source/ModelingComponentsEditorOnly/Public/ToolTargets/StaticMeshComponentToolTarget.h
--- a/Engine/Plugins/Runtime/MeshModelingToolset/Source/ModelingComponentsEditorOnly/Public/ToolTargets/StaticMeshComponentToolTarget.h
+++ b/Engine/Plugins/Runtime/MeshModelingToolset/Source/ModelingComponentsEditorOnly/Public/ToolTargets/StaticMeshComponentToolTarget.h
@@ -17,14 +17,16 @@
 
 #include "StaticMeshComponentToolTarget.generated.h"
 
+#define UE_API MODELINGCOMPONENTSEDITORONLY_API
+
 class UStaticMesh;
 
 /**
  * A tool target backed by a static mesh component that can provide and take a mesh
  * description, without 'committer' interfaces for writing to the source static mesh asset.
  */
-UCLASS(Transient)
-class MODELINGCOMPONENTSEDITORONLY_API UStaticMeshComponentReadOnlyToolTarget : public UPrimitiveComponentToolTarget,
+UCLASS(MinimalAPI, Transient)
+class UStaticMeshComponentReadOnlyToolTarget : public UPrimitiveComponentToolTarget,
 	public IMeshDescriptionProvider, public IMaterialProvider, public IStaticMeshBackedTarget,
 	public IDynamicMeshProvider, public IPhysicsDataSource
 {
@@ -35,44 +37,44 @@ public:
 	 * Configure active LOD to edit. Can only call this after Component is configured in base UPrimitiveComponentToolTarget.
 	 * If requested LOD does not exist, fall back to one that does.
 	 */
-	virtual void SetEditingLOD(EMeshLODIdentifier RequestedEditingLOD);
+	UE_API virtual void SetEditingLOD(EMeshLODIdentifier RequestedEditingLOD);
 
 	/** @return current editing LOD */
 	virtual EMeshLODIdentifier GetEditingLOD() const { return EditingLOD; }
 
 public:
-	virtual bool IsValid() const override;
+	UE_API virtual bool IsValid() const override;
 
 
 public:
 	// IMeshDescriptionProvider implementation
-	virtual const FMeshDescription* GetMeshDescription(const FGetMeshParameters& GetMeshParams = FGetMeshParameters()) override;
-	virtual FMeshDescription GetMeshDescriptionCopy(const FGetMeshParameters& GetMeshParams) override;
-	virtual FMeshDescription GetEmptyMeshDescription() override;
+	UE_API virtual const FMeshDescription* GetMeshDescription(const FGetMeshParameters& GetMeshParams = FGetMeshParameters()) override;
+	UE_API virtual FMeshDescription GetMeshDescriptionCopy(const FGetMeshParameters& GetMeshParams) override;
+	UE_API virtual FMeshDescription GetEmptyMeshDescription() override;
 	virtual bool SupportsLODs() const override { return true; }
-	virtual TArray<EMeshLODIdentifier> GetAvailableLODs(bool bSkipAutoGenerated = true) const override;
+	UE_API virtual TArray<EMeshLODIdentifier> GetAvailableLODs(bool bSkipAutoGenerated = true) const override;
 	virtual EMeshLODIdentifier GetMeshDescriptionLOD() const override { return EditingLOD; }
-	virtual TArray<int32> GetPolygonGroupToMaterialIndexMap() const override;
+	UE_API virtual TArray<int32> GetPolygonGroupToMaterialIndexMap() const override;
 
 
 
 	// IDynamicMeshProvider
-	virtual UE::Geometry::FDynamicMesh3 GetDynamicMesh() override;
-	virtual UE::Geometry::FDynamicMesh3 GetDynamicMesh(const FGetMeshParameters& InGetMeshParams) override;
+	UE_API virtual UE::Geometry::FDynamicMesh3 GetDynamicMesh() override;
+	UE_API virtual UE::Geometry::FDynamicMesh3 GetDynamicMesh(const FGetMeshParameters& InGetMeshParams) override;
 
 
 	// IMaterialProvider implementation
-	virtual int32 GetNumMaterials() const override;
-	virtual UMaterialInterface* GetMaterial(int32 MaterialIndex) const override;
-	virtual void GetMaterialSet(FComponentMaterialSet& MaterialSetOut, bool bPreferAssetMaterials) const override;
-	virtual bool CommitMaterialSetUpdate(const FComponentMaterialSet& MaterialSet, bool bApplyToAsset) override;	
+	UE_API virtual int32 GetNumMaterials() const override;
+	UE_API virtual UMaterialInterface* GetMaterial(int32 MaterialIndex) const override;
+	UE_API virtual void GetMaterialSet(FComponentMaterialSet& MaterialSetOut, bool bPreferAssetMaterials) const override;
+	UE_API virtual bool CommitMaterialSetUpdate(const FComponentMaterialSet& MaterialSet, bool bApplyToAsset) override;	
 
 	// IStaticMeshBackedTarget
-	virtual UStaticMesh* GetStaticMesh() const override;
+	UE_API virtual UStaticMesh* GetStaticMesh() const override;
 
 	// IPhysicsDataSource
-	virtual UBodySetup* GetBodySetup() const override;
-	virtual IInterface_CollisionDataProvider* GetComplexCollisionProvider() const override;
+	UE_API virtual UBodySetup* GetBodySetup() const override;
+	UE_API virtual IInterface_CollisionDataProvider* GetComplexCollisionProvider() const override;
 
 	// Rest provided by parent class
 
@@ -91,44 +93,46 @@ private:
  * A tool target backed by a static mesh component that can provide and take a mesh
  * description.
  */
-UCLASS(Transient)
-class MODELINGCOMPONENTSEDITORONLY_API UStaticMeshComponentToolTarget : 
+UCLASS(MinimalAPI, Transient)
+class UStaticMeshComponentToolTarget : 
 	public UStaticMeshComponentReadOnlyToolTarget,
 	public IMeshDescriptionCommitter, public IDynamicMeshCommitter
 {
 	GENERATED_BODY()
 
 	// IDynamicMeshCommitter
-	virtual void CommitDynamicMesh(const UE::Geometry::FDynamicMesh3& Mesh, const FDynamicMeshCommitInfo& CommitInfo) override;
+	UE_API virtual void CommitDynamicMesh(const UE::Geometry::FDynamicMesh3& Mesh, const FDynamicMeshCommitInfo& CommitInfo) override;
 	using IDynamicMeshCommitter::CommitDynamicMesh; // unhide the other overload
 	
 	// IMeshDescriptionCommitter implementation
-	virtual void CommitMeshDescription(const FCommitter& Committer, const FCommitMeshParameters& CommitParams = FCommitMeshParameters()) override;
+	UE_API virtual void CommitMeshDescription(const FCommitter& Committer, const FCommitMeshParameters& CommitParams = FCommitMeshParameters()) override;
 	using IMeshDescriptionCommitter::CommitMeshDescription; // unhide the other overload
 };
 
 
 /** Factory for UStaticMeshComponentToolTarget to be used by the target manager. */
-UCLASS(Transient)
-class MODELINGCOMPONENTSEDITORONLY_API UStaticMeshComponentToolTargetFactory : public UToolTargetFactory
+UCLASS(MinimalAPI, Transient)
+class UStaticMeshComponentToolTargetFactory : public UToolTargetFactory
 {
 	GENERATED_BODY()
 
 public:
-	virtual bool CanBuildTarget(UObject* SourceObject, const FToolTargetTypeRequirements& TargetTypeInfo) const override;
+	UE_API virtual bool CanBuildTarget(UObject* SourceObject, const FToolTargetTypeRequirements& TargetTypeInfo) const override;
 
-	virtual UToolTarget* BuildTarget(UObject* SourceObject, const FToolTargetTypeRequirements& TargetTypeInfo) override;
+	UE_API virtual UToolTarget* BuildTarget(UObject* SourceObject, const FToolTargetTypeRequirements& TargetTypeInfo) override;
 
 public:
 	virtual EMeshLODIdentifier GetActiveEditingLOD() const { return EditingLOD; }
-	virtual void SetActiveEditingLOD(EMeshLODIdentifier NewEditingLOD);
+	UE_API virtual void SetActiveEditingLOD(EMeshLODIdentifier NewEditingLOD);
 
-	static bool CanWriteToSource(const UObject* Source);
+	static UE_API bool CanWriteToSource(const UObject* Source);
 
 protected:
 	// LOD to edit, default is to edit LOD0
 	EMeshLODIdentifier EditingLOD = EMeshLODIdentifier::LOD0;
 
 private:
-	static const UStaticMesh* SourceToStaticMesh(const UObject* SourceObject);
-};
\ No newline at end of file
+	static UE_API const UStaticMesh* SourceToStaticMesh(const UObject* SourceObject);
+};
+
+#undef UE_API
diff --git a/Engine/Plugins/Runtime/MeshModelingToolset/Source/ModelingComponentsEditorOnly/Public/ToolTargets/StaticMeshToolTarget.h b/Engine/Plugins/Runtime/MeshModelingToolset/Source/ModelingComponentsEditorOnly/Public/ToolTargets/StaticMeshToolTarget.h
--- a/Engine/Plugins/Runtime/MeshModelingToolset/Source/ModelingComponentsEditorOnly/Public/ToolTargets/StaticMeshToolTarget.h
+++ b/Engine/Plugins/Runtime/MeshModelingToolset/Source/ModelingComponentsEditorOnly/Public/ToolTargets/StaticMeshToolTarget.h
@@ -15,6 +15,8 @@
 
 #include "StaticMeshToolTarget.generated.h"
 
+#define UE_API MODELINGCOMPONENTSEDITORONLY_API
+
 class UStaticMesh;
 
 /**
@@ -24,8 +26,8 @@ class UStaticMesh;
  * This is a special tool target that refers to the underlying asset (in this case a static mesh), rather than indirectly through a component.
  * This type of target is used in cases, such as opening an asset through the content browser, when there is no component available.
  */
-UCLASS(Transient)
-class MODELINGCOMPONENTSEDITORONLY_API UStaticMeshReadOnlyToolTarget : 
+UCLASS(MinimalAPI, Transient)
+class UStaticMeshReadOnlyToolTarget : 
 	public UToolTarget,
 	public IMeshDescriptionProvider, 
 	public IMaterialProvider, 
@@ -39,39 +41,39 @@ public:
 	 * Configure active LOD to edit. Can only call this after Component is configured in base UPrimitiveComponentToolTarget.
 	 * If requested LOD does not exist, fall back to one that does.
 	 */
-	virtual void SetEditingLOD(EMeshLODIdentifier RequestedEditingLOD);
+	UE_API virtual void SetEditingLOD(EMeshLODIdentifier RequestedEditingLOD);
 
 	/** @return current editing LOD */
 	virtual EMeshLODIdentifier GetEditingLOD() const { return EditingLOD; }
 
 	// UToolTarget
-	virtual bool IsValid() const override;
+	UE_API virtual bool IsValid() const override;
 
 	// IMeshDescriptionProvider implementation
-	virtual const FMeshDescription* GetMeshDescription(const FGetMeshParameters& GetMeshParams = FGetMeshParameters()) override;
-	virtual FMeshDescription GetEmptyMeshDescription() override;
+	UE_API virtual const FMeshDescription* GetMeshDescription(const FGetMeshParameters& GetMeshParams = FGetMeshParameters()) override;
+	UE_API virtual FMeshDescription GetEmptyMeshDescription() override;
 	virtual bool SupportsLODs() const override { return true; }
-	virtual TArray<EMeshLODIdentifier> GetAvailableLODs(bool bSkipAutoGenerated = true) const override;
+	UE_API virtual TArray<EMeshLODIdentifier> GetAvailableLODs(bool bSkipAutoGenerated = true) const override;
 	virtual EMeshLODIdentifier GetMeshDescriptionLOD() const override { return EditingLOD; }
-	virtual TArray<int32> GetPolygonGroupToMaterialIndexMap() const override;
+	UE_API virtual TArray<int32> GetPolygonGroupToMaterialIndexMap() const override;
 
 	// helper to get the material mapping returned by GetPolygonGroupToMaterialIndexMap for any static mesh
-	static TArray<int32> MapSectionToMaterialID(const UStaticMesh* Mesh, EMeshLODIdentifier EditingLOD);
+	static UE_API TArray<int32> MapSectionToMaterialID(const UStaticMesh* Mesh, EMeshLODIdentifier EditingLOD);
 
 
 	// IDynamicMeshProvider
-	virtual UE::Geometry::FDynamicMesh3 GetDynamicMesh() override;
-	virtual UE::Geometry::FDynamicMesh3 GetDynamicMesh(const FGetMeshParameters& InGetMeshParams) override;
+	UE_API virtual UE::Geometry::FDynamicMesh3 GetDynamicMesh() override;
+	UE_API virtual UE::Geometry::FDynamicMesh3 GetDynamicMesh(const FGetMeshParameters& InGetMeshParams) override;
 
 
 	// IMaterialProvider implementation
-	virtual int32 GetNumMaterials() const override;
-	virtual UMaterialInterface* GetMaterial(int32 MaterialIndex) const override;
-	virtual void GetMaterialSet(FComponentMaterialSet& MaterialSetOut, bool bPreferAssetMaterials) const override;
-	virtual bool CommitMaterialSetUpdate(const FComponentMaterialSet& MaterialSet, bool bApplyToAsset) override;	
+	UE_API virtual int32 GetNumMaterials() const override;
+	UE_API virtual UMaterialInterface* GetMaterial(int32 MaterialIndex) const override;
+	UE_API virtual void GetMaterialSet(FComponentMaterialSet& MaterialSetOut, bool bPreferAssetMaterials) const override;
+	UE_API virtual bool CommitMaterialSetUpdate(const FComponentMaterialSet& MaterialSet, bool bApplyToAsset) override;	
 
 	// IStaticMeshBackedTarget
-	virtual UStaticMesh* GetStaticMesh() const override;
+	UE_API virtual UStaticMesh* GetStaticMesh() const override;
 
 	// Rest provided by parent class
 
@@ -84,7 +86,7 @@ public:
 	};
 	using FMeshDescriptionCache = TMap<int32, FCachedMeshDescription>;
 	// helper to get the mesh description with build scale applied, using/updating the cache as needed
-	static const FMeshDescription* GetMeshDescriptionWithScaleApplied(const UStaticMesh* Mesh, int32 UseLOD, FMeshDescriptionCache& Cache);
+	static UE_API const FMeshDescription* GetMeshDescriptionWithScaleApplied(const UStaticMesh* Mesh, int32 UseLOD, FMeshDescriptionCache& Cache);
 
 protected:
 	TWeakObjectPtr<UStaticMesh> StaticMesh = nullptr;
@@ -98,13 +100,13 @@ protected:
 	friend class UStaticMeshComponentToolTargetFactory;
 
 	// returns true if the mesh has a non-generated LOD of the requested level.  Interprets EMeshLODIdentifier::Default as LOD0
-	static bool HasNonGeneratedLOD(const UStaticMesh* StaticMesh, EMeshLODIdentifier MeshLOD);
-	static EMeshLODIdentifier GetValidEditingLOD(const UStaticMesh* StaticMesh, 
+	static UE_API bool HasNonGeneratedLOD(const UStaticMesh* StaticMesh, EMeshLODIdentifier MeshLOD);
+	static UE_API EMeshLODIdentifier GetValidEditingLOD(const UStaticMesh* StaticMesh, 
 		EMeshLODIdentifier RequestedEditingLOD);
-	static TArray<EMeshLODIdentifier> GetAvailableLODs(const UStaticMesh* StaticMesh, bool bSkipAutoGenerated);
-	static void GetMaterialSet(const UStaticMesh* StaticMesh, 
+	static UE_API TArray<EMeshLODIdentifier> GetAvailableLODs(const UStaticMesh* StaticMesh, bool bSkipAutoGenerated);
+	static UE_API void GetMaterialSet(const UStaticMesh* StaticMesh, 
 		FComponentMaterialSet& MaterialSetOut, bool bPreferAssetMaterials);
-	static bool CommitMaterialSetUpdate(UStaticMesh* SkeletalMesh, 
+	static UE_API bool CommitMaterialSetUpdate(UStaticMesh* SkeletalMesh, 
 		const FComponentMaterialSet& MaterialSet, bool bApplyToAsset);
 
 private:
@@ -120,8 +122,8 @@ private:
  * This is a special tool target that refers to the underlying asset (in this case a static mesh), rather than indirectly through a component.
  * This type of target is used in cases, such as opening an asset through the content browser, when there is no component available.
  */
-UCLASS(Transient)
-class MODELINGCOMPONENTSEDITORONLY_API UStaticMeshToolTarget : 
+UCLASS(MinimalAPI, Transient)
+class UStaticMeshToolTarget : 
 	public UStaticMeshReadOnlyToolTarget,
 	public IMeshDescriptionCommitter,
 	public IDynamicMeshCommitter
@@ -130,42 +132,44 @@ class MODELINGCOMPONENTSEDITORONLY_API UStaticMeshToolTarget :
 	
 public:
 	// IMeshDescriptionCommitter implementation
-	virtual void CommitMeshDescription(const FCommitter& Committer, const FCommitMeshParameters& CommitParams = FCommitMeshParameters()) override;
+	UE_API virtual void CommitMeshDescription(const FCommitter& Committer, const FCommitMeshParameters& CommitParams = FCommitMeshParameters()) override;
 	using IMeshDescriptionCommitter::CommitMeshDescription; // unhide the other overload
 
 	// IDynamicMeshCommitter
-	virtual void CommitDynamicMesh(const UE::Geometry::FDynamicMesh3& Mesh, const FDynamicMeshCommitInfo& CommitInfo) override;
+	UE_API virtual void CommitDynamicMesh(const UE::Geometry::FDynamicMesh3& Mesh, const FDynamicMeshCommitInfo& CommitInfo) override;
 	using IDynamicMeshCommitter::CommitDynamicMesh; // unhide the other overload
 
 protected:
 	friend class UStaticMeshComponentToolTarget;
 
-	static void CommitMeshDescription(UStaticMesh* StaticMesh,
+	static UE_API void CommitMeshDescription(UStaticMesh* StaticMesh,
 		const FCommitter& Committer, EMeshLODIdentifier EditingLODIn);
 };
 
 
 /** Factory for UStaticMeshToolTarget to be used by the target manager. */
-UCLASS(Transient)
-class MODELINGCOMPONENTSEDITORONLY_API UStaticMeshToolTargetFactory : public UToolTargetFactory
+UCLASS(MinimalAPI, Transient)
+class UStaticMeshToolTargetFactory : public UToolTargetFactory
 {
 	GENERATED_BODY()
 
 public:
 
-	virtual bool CanBuildTarget(UObject* SourceObject, const FToolTargetTypeRequirements& TargetTypeInfo) const override;
+	UE_API virtual bool CanBuildTarget(UObject* SourceObject, const FToolTargetTypeRequirements& TargetTypeInfo) const override;
 
-	virtual UToolTarget* BuildTarget(UObject* SourceObject, const FToolTargetTypeRequirements& TargetTypeInfo) override;
+	UE_API virtual UToolTarget* BuildTarget(UObject* SourceObject, const FToolTargetTypeRequirements& TargetTypeInfo) override;
 
 
 
 public:
 	virtual EMeshLODIdentifier GetActiveEditingLOD() const { return EditingLOD; }
-	virtual void SetActiveEditingLOD(EMeshLODIdentifier NewEditingLOD);
+	UE_API virtual void SetActiveEditingLOD(EMeshLODIdentifier NewEditingLOD);
 
-	static bool CanWriteToSource(const UObject* Source);
+	static UE_API bool CanWriteToSource(const UObject* Source);
 
 protected:
 	// LOD to edit, default is to edit LOD0
 	EMeshLODIdentifier EditingLOD = EMeshLODIdentifier::LOD0;
-};
\ No newline at end of file
+};
+
+#undef UE_API
diff --git a/Engine/Plugins/Runtime/MeshModelingToolset/Source/ModelingComponentsEditorOnly/Public/ToolTargets/VolumeComponentTarget.h b/Engine/Plugins/Runtime/MeshModelingToolset/Source/ModelingComponentsEditorOnly/Public/ToolTargets/VolumeComponentTarget.h
--- a/Engine/Plugins/Runtime/MeshModelingToolset/Source/ModelingComponentsEditorOnly/Public/ToolTargets/VolumeComponentTarget.h
+++ b/Engine/Plugins/Runtime/MeshModelingToolset/Source/ModelingComponentsEditorOnly/Public/ToolTargets/VolumeComponentTarget.h
@@ -8,27 +8,29 @@
 #include "ConversionUtils/VolumeToDynamicMesh.h"
 #include "MeshDescription.h"
 
-class MODELINGCOMPONENTSEDITORONLY_API FVolumeComponentTargetFactory : public FComponentTargetFactory
+#define UE_API MODELINGCOMPONENTSEDITORONLY_API
+
+class FVolumeComponentTargetFactory : public FComponentTargetFactory
 {
 public:
-	bool CanBuild(UActorComponent* Candidate) override;
-	TUniquePtr<FPrimitiveComponentTarget> Build(UPrimitiveComponent* PrimitiveComponent) override;
+	UE_API bool CanBuild(UActorComponent* Candidate) override;
+	UE_API TUniquePtr<FPrimitiveComponentTarget> Build(UPrimitiveComponent* PrimitiveComponent) override;
 };
 
 /** Deprecated. Use the tool target instead. */
-class MODELINGCOMPONENTSEDITORONLY_API FVolumeComponentTarget : public FPrimitiveComponentTarget
+class FVolumeComponentTarget : public FPrimitiveComponentTarget
 {
 public:
 
-	FVolumeComponentTarget(UPrimitiveComponent* Component);
+	UE_API FVolumeComponentTarget(UPrimitiveComponent* Component);
 
-	FMeshDescription* GetMesh() override;
-	void CommitMesh(const FCommitter&) override;
+	UE_API FMeshDescription* GetMesh() override;
+	UE_API void CommitMesh(const FCommitter&) override;
 
-	virtual void GetMaterialSet(FComponentMaterialSet& MaterialSetOut, bool bAssetMaterials) const override;
+	UE_API virtual void GetMaterialSet(FComponentMaterialSet& MaterialSetOut, bool bAssetMaterials) const override;
 	virtual void CommitMaterialSetUpdate(const FComponentMaterialSet& MaterialSet, bool bApplyToAsset) override {};
 
-	virtual bool HasSameSourceData(const FPrimitiveComponentTarget& OtherTarget) const override;
+	UE_API virtual bool HasSameSourceData(const FPrimitiveComponentTarget& OtherTarget) const override;
 
 	const UE::Conversion::FVolumeToMeshOptions& GetVolumeToMeshOptions() { return VolumeToMeshOptions; }
 
@@ -37,3 +39,5 @@ protected:
 
 	UE::Conversion::FVolumeToMeshOptions VolumeToMeshOptions;
 };
+
+#undef UE_API
diff --git a/Engine/Plugins/Runtime/MeshModelingToolset/Source/ModelingComponentsEditorOnly/Public/ToolTargets/VolumeComponentToolTarget.h b/Engine/Plugins/Runtime/MeshModelingToolset/Source/ModelingComponentsEditorOnly/Public/ToolTargets/VolumeComponentToolTarget.h
--- a/Engine/Plugins/Runtime/MeshModelingToolset/Source/ModelingComponentsEditorOnly/Public/ToolTargets/VolumeComponentToolTarget.h
+++ b/Engine/Plugins/Runtime/MeshModelingToolset/Source/ModelingComponentsEditorOnly/Public/ToolTargets/VolumeComponentToolTarget.h
@@ -16,6 +16,8 @@
 
 #include "VolumeComponentToolTarget.generated.h"
 
+#define UE_API MODELINGCOMPONENTSEDITORONLY_API
+
 /**
  * The CVar "modeling.VolumeMaxTriCount" is used as a cap on triangles that the various Modeling Mode
  * Tools will allow an output AVolume to have. If this triangle count is exceeded, the mesh used to
@@ -28,8 +30,8 @@ extern MODELINGCOMPONENTSEDITORONLY_API TAutoConsoleVariable<int32> CVarModeling
 /**
  * A tool target backed by AVolume
  */
-UCLASS(Transient)
-class MODELINGCOMPONENTSEDITORONLY_API UVolumeComponentToolTarget : public UPrimitiveComponentToolTarget,
+UCLASS(MinimalAPI, Transient)
+class UVolumeComponentToolTarget : public UPrimitiveComponentToolTarget,
 	public IMaterialProvider, public IPhysicsDataSource,
 	public IDynamicMeshCommitter, public IDynamicMeshProvider,
 	public IMeshDescriptionCommitter, public IMeshDescriptionProvider
@@ -38,36 +40,36 @@ class MODELINGCOMPONENTSEDITORONLY_API UVolumeComponentToolTarget : public UPrim
 
 public:
 
-	UVolumeComponentToolTarget();
+	UE_API UVolumeComponentToolTarget();
 
 	const UE::Conversion::FVolumeToMeshOptions& GetVolumeToMeshOptions() { return VolumeToMeshOptions; }
 
 	// IDynamicMeshProvider implementation
-	virtual UE::Geometry::FDynamicMesh3 GetDynamicMesh() override;
+	UE_API virtual UE::Geometry::FDynamicMesh3 GetDynamicMesh() override;
 
 	// IDynamicMeshCommitter implementation
-	virtual void CommitDynamicMesh(const UE::Geometry::FDynamicMesh3& Mesh, const FDynamicMeshCommitInfo&) override;
+	UE_API virtual void CommitDynamicMesh(const UE::Geometry::FDynamicMesh3& Mesh, const FDynamicMeshCommitInfo&) override;
 	using IDynamicMeshCommitter::CommitDynamicMesh; // unhide the other overload
 
 	// IMaterialProvider implementation
-	virtual int32 GetNumMaterials() const override;
-	virtual UMaterialInterface* GetMaterial(int32 MaterialIndex) const override;
+	UE_API virtual int32 GetNumMaterials() const override;
+	UE_API virtual UMaterialInterface* GetMaterial(int32 MaterialIndex) const override;
 	// Ignores bPreferAssetMaterials
-	virtual void GetMaterialSet(FComponentMaterialSet& MaterialSetOut, bool bPreferAssetMaterials) const override;
+	UE_API virtual void GetMaterialSet(FComponentMaterialSet& MaterialSetOut, bool bPreferAssetMaterials) const override;
 
 	// Doesn't actually do anything for a volume
 	virtual bool CommitMaterialSetUpdate(const FComponentMaterialSet& MaterialSet, bool bApplyToAsset) override { return false; }
 
 	// IMeshDescriptionProvider implementation
-	virtual const FMeshDescription* GetMeshDescription(const FGetMeshParameters& GetMeshParams = FGetMeshParameters()) override;
-	virtual FMeshDescription GetEmptyMeshDescription() override;
+	UE_API virtual const FMeshDescription* GetMeshDescription(const FGetMeshParameters& GetMeshParams = FGetMeshParameters()) override;
+	UE_API virtual FMeshDescription GetEmptyMeshDescription() override;
 
 	// IMeshDescritpionCommitter implementation
-	virtual void CommitMeshDescription(const FCommitter& Committer, const FCommitMeshParameters& CommitMeshParams = FCommitMeshParameters()) override;
+	UE_API virtual void CommitMeshDescription(const FCommitter& Committer, const FCommitMeshParameters& CommitMeshParams = FCommitMeshParameters()) override;
 	using IMeshDescriptionCommitter::CommitMeshDescription; // unhide the other overload
 
 	// IPhysicsDataSource implementation
-	virtual UBodySetup* GetBodySetup() const override;
+	UE_API virtual UBodySetup* GetBodySetup() const override;
 	// always returns null because volumes do not support IInterface_CollisionDataProvider
 	virtual IInterface_CollisionDataProvider* GetComplexCollisionProvider() const override { return nullptr; }
 
@@ -84,14 +86,16 @@ protected:
 };
 
 /** Factory for UVolumeComponentToolTarget to be used by the target manager. */
-UCLASS(Transient)
-class MODELINGCOMPONENTSEDITORONLY_API UVolumeComponentToolTargetFactory : public UToolTargetFactory
+UCLASS(MinimalAPI, Transient)
+class UVolumeComponentToolTargetFactory : public UToolTargetFactory
 {
 	GENERATED_BODY()
 
 public:
 
-	virtual bool CanBuildTarget(UObject* SourceObject, const FToolTargetTypeRequirements& TargetTypeInfo) const override;
+	UE_API virtual bool CanBuildTarget(UObject* SourceObject, const FToolTargetTypeRequirements& TargetTypeInfo) const override;
+
+	UE_API virtual UToolTarget* BuildTarget(UObject* SourceObject, const FToolTargetTypeRequirements& TargetTypeInfo) override;
+};
 
-	virtual UToolTarget* BuildTarget(UObject* SourceObject, const FToolTargetTypeRequirements& TargetTypeInfo) override;
-};
\ No newline at end of file
+#undef UE_API
diff --git a/Engine/Plugins/Runtime/MeshModelingToolset/Source/ModelingOperatorsEditorOnly/Public/CleaningOps/SimplifyMeshOp.h b/Engine/Plugins/Runtime/MeshModelingToolset/Source/ModelingOperatorsEditorOnly/Public/CleaningOps/SimplifyMeshOp.h
--- a/Engine/Plugins/Runtime/MeshModelingToolset/Source/ModelingOperatorsEditorOnly/Public/CleaningOps/SimplifyMeshOp.h
+++ b/Engine/Plugins/Runtime/MeshModelingToolset/Source/ModelingOperatorsEditorOnly/Public/CleaningOps/SimplifyMeshOp.h
@@ -11,6 +11,8 @@
 
 #include "SimplifyMeshOp.generated.h"
 
+#define UE_API MODELINGOPERATORSEDITORONLY_API
+
 struct FMeshDescription;
 class IMeshReduction;
 
@@ -65,7 +67,7 @@ namespace Geometry
 {
 
 
-class MODELINGOPERATORSEDITORONLY_API FSimplifyMeshOp : public FDynamicMeshOperator
+class FSimplifyMeshOp : public FDynamicMeshOperator
 {
 public:
 	virtual ~FSimplifyMeshOp() {}
@@ -114,7 +116,7 @@ public:
 	 * @param Progress If non-null, can early-out if this is set to request cancellation
 	 * @return true if succeeded, false otherwise (e.g. if cancelled)
 	 */
-	static bool ComputeStandardSimplifier(IMeshReduction* MeshReduction, const FMeshDescription& SrcMeshDescription, FDynamicMesh3& OutResult,
+	static UE_API bool ComputeStandardSimplifier(IMeshReduction* MeshReduction, const FMeshDescription& SrcMeshDescription, FDynamicMesh3& OutResult,
 		float PercentReduction, bool bTriBasedReduction, bool bDiscardAttributes,
 		FProgressCancel* Progress = nullptr);
 
@@ -128,8 +130,10 @@ public:
 	// FDynamicMeshOperator implementation
 	// 
 
-	virtual void CalculateResult(FProgressCancel* Progress) override;
+	UE_API virtual void CalculateResult(FProgressCancel* Progress) override;
 };
 
 } // end namespace UE::Geometry
 } // end namespace UE
+
+#undef UE_API
diff --git a/Engine/Plugins/Runtime/MeshModelingToolset/Source/ModelingOperatorsEditorOnly/Public/CompositionOps/VoxelBooleanMeshesOp.h b/Engine/Plugins/Runtime/MeshModelingToolset/Source/ModelingOperatorsEditorOnly/Public/CompositionOps/VoxelBooleanMeshesOp.h
--- a/Engine/Plugins/Runtime/MeshModelingToolset/Source/ModelingOperatorsEditorOnly/Public/CompositionOps/VoxelBooleanMeshesOp.h
+++ b/Engine/Plugins/Runtime/MeshModelingToolset/Source/ModelingOperatorsEditorOnly/Public/CompositionOps/VoxelBooleanMeshesOp.h
@@ -8,6 +8,8 @@
 #include "DynamicMesh/DynamicMesh3.h"
 #include "TransformTypes.h"
 
+#define UE_API MODELINGOPERATORSEDITORONLY_API
+
 struct FMeshDescription;
 
 namespace UE
@@ -16,7 +18,7 @@ namespace Geometry
 {
 
 
-class MODELINGOPERATORSEDITORONLY_API FVoxelBooleanMeshesOp : public FDynamicMeshOperator
+class FVoxelBooleanMeshesOp : public FDynamicMeshOperator
 {
 public:
 	virtual ~FVoxelBooleanMeshesOp() {}
@@ -45,15 +47,16 @@ public:
 	// FDynamicMeshOperator implementation
 	// 
 
-	virtual void CalculateResult(FProgressCancel* Progress) override;
+	UE_API virtual void CalculateResult(FProgressCancel* Progress) override;
 
 private:
 
 	// compute the voxel size based on the voxel count and the input geometry
-	float ComputeVoxelSize() const;
+	UE_API float ComputeVoxelSize() const;
 
 };
 
 } // end namespace UE::Geometry
 } // end namespace UE
 
+#undef UE_API
diff --git a/Engine/Plugins/Runtime/MeshModelingToolset/Source/ModelingOperatorsEditorOnly/Public/CompositionOps/VoxelMergeMeshesOp.h b/Engine/Plugins/Runtime/MeshModelingToolset/Source/ModelingOperatorsEditorOnly/Public/CompositionOps/VoxelMergeMeshesOp.h
--- a/Engine/Plugins/Runtime/MeshModelingToolset/Source/ModelingOperatorsEditorOnly/Public/CompositionOps/VoxelMergeMeshesOp.h
+++ b/Engine/Plugins/Runtime/MeshModelingToolset/Source/ModelingOperatorsEditorOnly/Public/CompositionOps/VoxelMergeMeshesOp.h
@@ -8,6 +8,8 @@
 #include "DynamicMesh/DynamicMesh3.h"
 #include "TransformTypes.h"
 
+#define UE_API MODELINGOPERATORSEDITORONLY_API
+
 struct FMeshDescription;
 
 namespace UE
@@ -16,7 +18,7 @@ namespace Geometry
 {
 
 
-class MODELINGOPERATORSEDITORONLY_API FVoxelMergeMeshesOp : public FDynamicMeshOperator
+class FVoxelMergeMeshesOp : public FDynamicMeshOperator
 {
 public:
 	virtual ~FVoxelMergeMeshesOp() {}
@@ -38,14 +40,16 @@ public:
 	// FDynamicMeshOperator implementation
 	// 
 
-	virtual void CalculateResult(FProgressCancel* Progress) override;
+	UE_API virtual void CalculateResult(FProgressCancel* Progress) override;
 
 private:
 
 	// compute the voxel size based on the voxel count and the input geometry
-	float ComputeVoxelSize() const;
+	UE_API float ComputeVoxelSize() const;
 };
 
 
 } // end namespace UE::Geometry
 } // end namespace UE
+
+#undef UE_API
diff --git a/Engine/Plugins/Runtime/MeshModelingToolset/Source/ModelingOperatorsEditorOnly/Public/CuttingOps/EmbedPolygonsOp.h b/Engine/Plugins/Runtime/MeshModelingToolset/Source/ModelingOperatorsEditorOnly/Public/CuttingOps/EmbedPolygonsOp.h
--- a/Engine/Plugins/Runtime/MeshModelingToolset/Source/ModelingOperatorsEditorOnly/Public/CuttingOps/EmbedPolygonsOp.h
+++ b/Engine/Plugins/Runtime/MeshModelingToolset/Source/ModelingOperatorsEditorOnly/Public/CuttingOps/EmbedPolygonsOp.h
@@ -10,6 +10,8 @@
 
 #include "EmbedPolygonsOp.generated.h"
 
+#define UE_API MODELINGOPERATORSEDITORONLY_API
+
 
 UENUM()
 enum class EEmbeddedPolygonOpMethod : uint8
@@ -27,7 +29,7 @@ namespace Geometry
 {
 
 
-class MODELINGOPERATORSEDITORONLY_API FEmbedPolygonsOp : public FDynamicMeshOperator
+class FEmbedPolygonsOp : public FDynamicMeshOperator
 {
 public:
 	virtual ~FEmbedPolygonsOp() {}
@@ -57,7 +59,7 @@ public:
 	// FDynamicMeshOperator implementation
 	// 
 
-	virtual void CalculateResult(FProgressCancel* Progress) override;
+	UE_API virtual void CalculateResult(FProgressCancel* Progress) override;
 
 
 	// Outputs
@@ -66,16 +68,17 @@ public:
 	bool bOperationSucceeded;
 
 private:
-	void RecordEmbeddedEdges(TArray<int>& PathVertIDs);
+	UE_API void RecordEmbeddedEdges(TArray<int>& PathVertIDs);
 
-	void BooleanPath(FProgressCancel* Progress);
+	UE_API void BooleanPath(FProgressCancel* Progress);
 
 	// Fix issues like self-intersection on an input polygon
 	// Note: We could change this to return an array of general polygons in the future, to handle the most general case,
 	// but currently the embedding process does not support holes
-	TArray<FPolygon2d> CleanPolygon(const FPolygon2d& Input);
+	UE_API TArray<FPolygon2d> CleanPolygon(const FPolygon2d& Input);
 };
 
 } // end namespace UE::Geometry
 } // end namespace UE
 
+#undef UE_API
diff --git a/Engine/Plugins/Runtime/MeshModelingToolset/Source/ModelingOperatorsEditorOnly/Public/ParameterizationOps/CalculateTangentsOp.h b/Engine/Plugins/Runtime/MeshModelingToolset/Source/ModelingOperatorsEditorOnly/Public/ParameterizationOps/CalculateTangentsOp.h
--- a/Engine/Plugins/Runtime/MeshModelingToolset/Source/ModelingOperatorsEditorOnly/Public/ParameterizationOps/CalculateTangentsOp.h
+++ b/Engine/Plugins/Runtime/MeshModelingToolset/Source/ModelingOperatorsEditorOnly/Public/ParameterizationOps/CalculateTangentsOp.h
@@ -8,6 +8,8 @@
 #include "DynamicMesh/MeshTangents.h"
 #include "CalculateTangentsOp.generated.h"
 
+#define UE_API MODELINGOPERATORSEDITORONLY_API
+
 
 UENUM()
 enum class EMeshTangentsType : uint8
@@ -27,7 +29,7 @@ namespace UE
 namespace Geometry
 {
 
-class MODELINGOPERATORSEDITORONLY_API FCalculateTangentsOp : public TGenericDataOperator<FMeshTangentsd>
+class FCalculateTangentsOp : public TGenericDataOperator<FMeshTangentsd>
 {
 public:
 	virtual ~FCalculateTangentsOp() {}
@@ -46,12 +48,14 @@ public:
 	// TGenericDataOperator implementation
 	// 
 
-	virtual void CalculateResult(FProgressCancel* Progress) override;
+	UE_API virtual void CalculateResult(FProgressCancel* Progress) override;
 
-	virtual void CalculateStandard(FProgressCancel* Progress, TUniquePtr<FMeshTangentsd>& Tangents);
-	virtual void CalculateMikkT(FProgressCancel* Progress, TUniquePtr<FMeshTangentsd>& Tangents);
-	virtual void CopyFromSource(FProgressCancel* Progress, TUniquePtr<FMeshTangentsd>& Tangents);
+	UE_API virtual void CalculateStandard(FProgressCancel* Progress, TUniquePtr<FMeshTangentsd>& Tangents);
+	UE_API virtual void CalculateMikkT(FProgressCancel* Progress, TUniquePtr<FMeshTangentsd>& Tangents);
+	UE_API virtual void CopyFromSource(FProgressCancel* Progress, TUniquePtr<FMeshTangentsd>& Tangents);
 };
 
 } // end namespace UE::Geometry
 } // end namespace UE
+
+#undef UE_API
diff --git a/Engine/Plugins/Runtime/MeshModelingToolset/Source/ModelingOperatorsEditorOnly/Public/ParameterizationOps/ParameterizeMeshOp.h b/Engine/Plugins/Runtime/MeshModelingToolset/Source/ModelingOperatorsEditorOnly/Public/ParameterizationOps/ParameterizeMeshOp.h
--- a/Engine/Plugins/Runtime/MeshModelingToolset/Source/ModelingOperatorsEditorOnly/Public/ParameterizationOps/ParameterizeMeshOp.h
+++ b/Engine/Plugins/Runtime/MeshModelingToolset/Source/ModelingOperatorsEditorOnly/Public/ParameterizationOps/ParameterizeMeshOp.h
@@ -10,6 +10,8 @@
 
 #include "ParameterizeMeshOp.generated.h"
 
+#define UE_API MODELINGOPERATORSEDITORONLY_API
+
 class UParameterizeMeshToolProperties;
 class UParameterizeMeshToolUVAtlasProperties;
 class UParameterizeMeshToolXAtlasProperties;
@@ -27,7 +29,7 @@ enum class EParamOpBackend
 	XAtlas = 2
 };
 
-class MODELINGOPERATORSEDITORONLY_API FParameterizeMeshOp : public FDynamicMeshOperator
+class FParameterizeMeshOp : public FDynamicMeshOperator
 {
 public:
 	virtual ~FParameterizeMeshOp() {}
@@ -82,7 +84,7 @@ public:
 	// FDynamicMeshOperator implementation
 	// 
 
-	virtual void CalculateResult(FProgressCancel* Progress) override;
+	UE_API virtual void CalculateResult(FProgressCancel* Progress) override;
 
 
 protected:
@@ -107,11 +109,11 @@ protected:
 
 	FGeometryResult NewResultInfo;
 
-	bool ComputeUVs_UVAtlas(FDynamicMesh3& InOutMesh, TFunction<bool(float)>& Interrupter);
-	bool ComputeUVs_XAtlas(FDynamicMesh3& InOutMesh, TFunction<bool(float)>& Interrupter);
-	bool ComputeUVs_PatchBuilder(FDynamicMesh3& InOutMesh, FProgressCancel* Progress);
+	UE_API bool ComputeUVs_UVAtlas(FDynamicMesh3& InOutMesh, TFunction<bool(float)>& Interrupter);
+	UE_API bool ComputeUVs_XAtlas(FDynamicMesh3& InOutMesh, TFunction<bool(float)>& Interrupter);
+	UE_API bool ComputeUVs_PatchBuilder(FDynamicMesh3& InOutMesh, FProgressCancel* Progress);
 
-	void CopyNewUVsToMesh(
+	UE_API void CopyNewUVsToMesh(
 		FDynamicMesh3& Mesh,
 		const FLinearMesh& LinearMesh,
 		const FDynamicMesh3& FlippedMesh,
@@ -120,7 +122,7 @@ protected:
 		const TArray<int32>& VertexRemapArray,
 		bool bReverseOrientation);
 
-	void LayoutToUDIMByPolygroup(FDynamicMesh3& InOutMesh, UE::Geometry::FPolygroupSet& InputGroups);
+	UE_API void LayoutToUDIMByPolygroup(FDynamicMesh3& InOutMesh, UE::Geometry::FPolygroupSet& InputGroups);
 };
 
 } // end namespace UE::Geometry
@@ -132,14 +134,14 @@ protected:
  * Inherits from UObject so that it can hold a strong pointer to the settings UObject, which
  * needs to be a UObject to be displayed in the details panel.
  */
-UCLASS()
-class MODELINGOPERATORSEDITORONLY_API UParameterizeMeshOperatorFactory : public UObject, public UE::Geometry::IDynamicMeshOperatorFactory
+UCLASS(MinimalAPI)
+class UParameterizeMeshOperatorFactory : public UObject, public UE::Geometry::IDynamicMeshOperatorFactory
 {
 	GENERATED_BODY()
 
 public:
 	// IDynamicMeshOperatorFactory API
-	virtual TUniquePtr<UE::Geometry::FDynamicMeshOperator> MakeNewOperator() override;
+	UE_API virtual TUniquePtr<UE::Geometry::FDynamicMeshOperator> MakeNewOperator() override;
 
 	UPROPERTY()
 	TObjectPtr<UParameterizeMeshToolProperties> Settings = nullptr;
@@ -158,4 +160,6 @@ public:
 	TUniqueFunction<int32()> GetSelectedUVChannel = []() { return 0; };
 	TUniqueFunction<TSharedPtr<UE::Geometry::FPolygroupSet, ESPMode::ThreadSafe>()> GetPolygroups = []() { return nullptr; };
 	FTransform TargetTransform;
-};
\ No newline at end of file
+};
+
+#undef UE_API
diff --git a/Engine/Plugins/Runtime/MeshModelingToolset/Source/ModelingOperatorsEditorOnly/Public/Properties/ParameterizeMeshProperties.h b/Engine/Plugins/Runtime/MeshModelingToolset/Source/ModelingOperatorsEditorOnly/Public/Properties/ParameterizeMeshProperties.h
--- a/Engine/Plugins/Runtime/MeshModelingToolset/Source/ModelingOperatorsEditorOnly/Public/Properties/ParameterizeMeshProperties.h
+++ b/Engine/Plugins/Runtime/MeshModelingToolset/Source/ModelingOperatorsEditorOnly/Public/Properties/ParameterizeMeshProperties.h
@@ -22,8 +22,8 @@ enum class EParameterizeMeshUVMethod
 };
 
 
-UCLASS()
-class MODELINGOPERATORSEDITORONLY_API UParameterizeMeshToolProperties : public UInteractiveToolPropertySet
+UCLASS(MinimalAPI)
+class UParameterizeMeshToolProperties : public UInteractiveToolPropertySet
 {
 	GENERATED_BODY()
 public:
@@ -36,8 +36,8 @@ public:
 /**
  * Settings for the UVAtlas Automatic UV Generation Method
  */
-UCLASS()
-class MODELINGOPERATORSEDITORONLY_API UParameterizeMeshToolUVAtlasProperties : public UInteractiveToolPropertySet
+UCLASS(MinimalAPI)
+class UParameterizeMeshToolUVAtlasProperties : public UInteractiveToolPropertySet
 {
 	GENERATED_BODY()
 public:
@@ -71,8 +71,8 @@ public:
 };
 
 
-UCLASS()
-class MODELINGOPERATORSEDITORONLY_API UParameterizeMeshToolXAtlasProperties : public UInteractiveToolPropertySet
+UCLASS(MinimalAPI)
+class UParameterizeMeshToolXAtlasProperties : public UInteractiveToolPropertySet
 {
 	GENERATED_BODY()
 public:
@@ -82,8 +82,8 @@ public:
 };
 
 
-UCLASS()
-class MODELINGOPERATORSEDITORONLY_API UParameterizeMeshToolPatchBuilderProperties : public UInteractiveToolPropertySet
+UCLASS(MinimalAPI)
+class UParameterizeMeshToolPatchBuilderProperties : public UInteractiveToolPropertySet
 {
 	GENERATED_BODY()
 public:
diff --git a/Engine/Plugins/Runtime/MeshModelingToolset/Source/SkeletalMeshModifiers/Public/SkeletonModifier.h b/Engine/Plugins/Runtime/MeshModelingToolset/Source/SkeletalMeshModifiers/Public/SkeletonModifier.h
--- a/Engine/Plugins/Runtime/MeshModelingToolset/Source/SkeletalMeshModifiers/Public/SkeletonModifier.h
+++ b/Engine/Plugins/Runtime/MeshModelingToolset/Source/SkeletalMeshModifiers/Public/SkeletonModifier.h
@@ -11,6 +11,8 @@
 
 #include "SkeletonModifier.generated.h"
 
+#define UE_API SKELETALMESHMODIFIERS_API
+
 class USkeletalMesh;
 class FName;
 struct FMeshDescription;
@@ -54,7 +56,7 @@ ENUM_CLASS_FLAGS(ESkeletonModificationType)
  */
 
 USTRUCT(BlueprintType)
-struct SKELETALMESHMODIFIERS_API FMirrorOptions
+struct FMirrorOptions
 {
 	GENERATED_USTRUCT_BODY()
 
@@ -73,8 +75,8 @@ struct SKELETALMESHMODIFIERS_API FMirrorOptions
 	UPROPERTY(BlueprintReadWrite, EditAnywhere, Category = "Mirror")
 	bool bMirrorChildren = true;
 
-	FTransform MirrorTransform(const FTransform& InGlobalTransform) const;
-	FVector MirrorVector(const FVector& InVector) const;
+	UE_API FTransform MirrorTransform(const FTransform& InGlobalTransform) const;
+	UE_API FVector MirrorVector(const FVector& InVector) const;
 };
 
 /**
@@ -94,7 +96,7 @@ enum class EOrientAxis : uint8
 };
 
 USTRUCT(BlueprintType)
-struct SKELETALMESHMODIFIERS_API FOrientOptions
+struct FOrientOptions
 {
 	GENERATED_USTRUCT_BODY()
 
@@ -113,22 +115,22 @@ struct SKELETALMESHMODIFIERS_API FOrientOptions
 	UPROPERTY(BlueprintReadWrite, EditAnywhere, Category = "Orient", meta=(EditCondition="Primary != EOrientAxis::None"))
 	bool bOrientChildren = true;
 
-	FTransform OrientTransform(const FVector& InPrimaryTarget, const FTransform& InTransform) const;
+	UE_API FTransform OrientTransform(const FVector& InPrimaryTarget, const FTransform& InTransform) const;
 };
 
 /**
  * FSkeletalMeshSkeletonModifier
  */
 
-UCLASS(BlueprintType)
-class SKELETALMESHMODIFIERS_API USkeletonModifier : public UObject
+UCLASS(MinimalAPI, BlueprintType)
+class USkeletonModifier : public UObject
 {
 	GENERATED_BODY()
 	
 public:
 
 	UFUNCTION(BlueprintCallable, Category = "Mesh")
-	bool SetSkeletalMesh(USkeletalMesh* InSkeletalMesh);
+	UE_API bool SetSkeletalMesh(USkeletalMesh* InSkeletalMesh);
 
 	/** Creates a new bone in the skeleton hierarchy at desired transform
 	 *  @param InBoneName The new bone's name.
@@ -137,9 +139,9 @@ public:
 	 *  @return \c true if the operation succeeded, false otherwise. 
 	 */
 	UFUNCTION(BlueprintCallable, Category = "Skeleton")
-	bool AddBone(const FName InBoneName, const FName InParentName, const FTransform& InTransform);
+	UE_API bool AddBone(const FName InBoneName, const FName InParentName, const FTransform& InTransform);
 	UFUNCTION(BlueprintCallable, Category = "Skeleton")
-	bool AddBones(const TArray<FName>& InBonesName, const TArray<FName>& InParentsName, const TArray<FTransform>& InTransforms);
+	UE_API bool AddBones(const TArray<FName>& InBonesName, const TArray<FName>& InParentsName, const TArray<FTransform>& InTransforms);
 
 	/** Mirror bones
 	 *  @param InBoneName The new bone's name.
@@ -147,9 +149,9 @@ public:
 	 *  @return \c true if the operation succeeded, false otherwise. 
 	 */
 	UFUNCTION(BlueprintCallable, Category = "Skeleton")
-	bool MirrorBone(const FName InBoneName, const FMirrorOptions& InOptions = FMirrorOptions());
+	UE_API bool MirrorBone(const FName InBoneName, const FMirrorOptions& InOptions = FMirrorOptions());
 	UFUNCTION(BlueprintCallable, Category = "Skeleton")
-	bool MirrorBones(const TArray<FName>& InBonesName, const FMirrorOptions& InOptions = FMirrorOptions());
+	UE_API bool MirrorBones(const TArray<FName>& InBonesName, const FMirrorOptions& InOptions = FMirrorOptions());
 
 	/** Sets the bone the desired local transform
 	 *  @param InBoneName The new bone's name that needs to be moved.
@@ -158,9 +160,9 @@ public:
 	 *  @return \c true if the operation succeeded, false otherwise. 
 	 */
 	UFUNCTION(BlueprintCallable, Category = "Skeleton")
-	bool SetBoneTransform(const FName InBoneName, const FTransform& InNewTransform, const bool bMoveChildren);
+	UE_API bool SetBoneTransform(const FName InBoneName, const FTransform& InNewTransform, const bool bMoveChildren);
 	UFUNCTION(BlueprintCallable, Category = "Skeleton")
-	bool SetBonesTransforms(const TArray<FName>& InBoneNames, const TArray<FTransform>& InNewTransforms, const bool bMoveChildren);
+	UE_API bool SetBonesTransforms(const TArray<FName>& InBoneNames, const TArray<FTransform>& InNewTransforms, const bool bMoveChildren);
 
 	/** Remove a bone in the skeleton hierarchy
 	 *  @param InBoneName The new bone's name.
@@ -168,9 +170,9 @@ public:
 	 *  @return \c true if the operation succeeded, false otherwise. 
 	 */
 	UFUNCTION(BlueprintCallable, Category = "Skeleton")
-	bool RemoveBone(const FName InBoneName, const bool bRemoveChildren);
+	UE_API bool RemoveBone(const FName InBoneName, const bool bRemoveChildren);
 	UFUNCTION(BlueprintCallable, Category = "Skeleton")
-	bool RemoveBones(const TArray<FName>& InBoneNames, const bool bRemoveChildren);
+	UE_API bool RemoveBones(const TArray<FName>& InBoneNames, const bool bRemoveChildren);
 
 	/** Rename bones
 	 *  @param InOldBoneName The current bone's name.
@@ -178,9 +180,9 @@ public:
 	 *  @return \c true if the operation succeeded, false otherwise. 
 	 */
 	UFUNCTION(BlueprintCallable, Category = "Skeleton")
-	bool RenameBone(const FName InOldBoneName, const FName InNewBoneName);
+	UE_API bool RenameBone(const FName InOldBoneName, const FName InNewBoneName);
 	UFUNCTION(BlueprintCallable, Category = "Skeleton")
-	bool RenameBones(const TArray<FName>& InOldBoneNames, const TArray<FName>& InNewBoneNames);
+	UE_API bool RenameBones(const TArray<FName>& InOldBoneNames, const TArray<FName>& InNewBoneNames);
 	
 	/** Parent bones
 	 *  @param InBoneName The current bone's name.
@@ -188,9 +190,9 @@ public:
 	 *  @return \c true if the operation succeeded, false otherwise. 
 	 */
 	UFUNCTION(BlueprintCallable, Category = "Skeleton")
-	bool ParentBone(const FName InBoneName, const FName InParentName);
+	UE_API bool ParentBone(const FName InBoneName, const FName InParentName);
 	UFUNCTION(BlueprintCallable, Category = "Skeleton")
-	bool ParentBones(const TArray<FName>& InBoneNames, const TArray<FName>& InParentNames);
+	UE_API bool ParentBones(const TArray<FName>& InBoneNames, const TArray<FName>& InParentNames);
 
 	/** Align bones
 	 *  @param InBoneName The current bone's name.
@@ -198,9 +200,9 @@ public:
 	 *  @return \c true if the operation succeeded, false otherwise. 
 	 */
 	 UFUNCTION(BlueprintCallable, Category = "Skeleton")
-	bool OrientBone(const FName InBoneName, const FOrientOptions& InOptions = FOrientOptions());
+	UE_API bool OrientBone(const FName InBoneName, const FOrientOptions& InOptions = FOrientOptions());
 	UFUNCTION(BlueprintCallable, Category = "Skeleton")
-	bool OrientBones(const TArray<FName>& InBoneNames, const FOrientOptions& InOptions = FOrientOptions());
+	UE_API bool OrientBones(const TArray<FName>& InBoneNames, const FOrientOptions& InOptions = FOrientOptions());
 
 	/** Get Bone Transform
 	 *  @param InBoneName The current bone's name.
@@ -208,14 +210,14 @@ public:
 	 *  @return \c The current bone's transform 
 	 */
 	UFUNCTION(BlueprintCallable, Category = "Skeleton")
-	FTransform GetBoneTransform(const FName InBoneName, const bool bGlobal = false) const;
+	UE_API FTransform GetBoneTransform(const FName InBoneName, const bool bGlobal = false) const;
 
 	/** Get Parent Name
 	 *  @param InBoneName The current bone's name.
 	 *  @return \c The current bone's parent name 
 	 */
 	UFUNCTION(BlueprintCallable, Category = "Skeleton")
-	FName GetParentName(const FName InBoneName) const;
+	UE_API FName GetParentName(const FName InBoneName) const;
 
 	/** Get Children Names
 	 *  @param InBoneName The parent's name.
@@ -223,59 +225,59 @@ public:
 	 *  @return \c The children names list 
 	 */
 	UFUNCTION(BlueprintCallable, Category = "Skeleton")
-	TArray<FName> GetChildrenNames(const FName InBoneName, const bool bRecursive = false) const;
+	UE_API TArray<FName> GetChildrenNames(const FName InBoneName, const bool bRecursive = false) const;
 
 	/** Get All Bone Names
 	 *  @return \c All bone names list
 	 */
 	UFUNCTION(BlueprintCallable, Category = "Skeleton")
-	TArray<FName> GetAllBoneNames() const;
+	UE_API TArray<FName> GetAllBoneNames() const;
 	
 	/**
 	 * Actually applies the skeleton modifications to the skeletal mesh.
 	 * @return true if commit succeeded.
 	 */
 	UFUNCTION(BlueprintCallable, Category = "Mesh")
-	bool CommitSkeletonToSkeletalMesh();
+	UE_API bool CommitSkeletonToSkeletalMesh();
 
-	void ExternalUpdate(const FReferenceSkeleton& InRefSkeleton, const TArray<int32>& InIndexTracker);
+	UE_API void ExternalUpdate(const FReferenceSkeleton& InRefSkeleton, const TArray<int32>& InIndexTracker);
 	
-	FName GetUniqueName(const FName InBoneName, const TArray<FName>& InBoneNames = {}) const;
+	UE_API FName GetUniqueName(const FName InBoneName, const TArray<FName>& InBoneNames = {}) const;
 	
-	const FReferenceSkeleton& GetReferenceSkeleton() const;
-	const TArray<int32>& GetBoneIndexTracker() const;
+	UE_API const FReferenceSkeleton& GetReferenceSkeleton() const;
+	UE_API const TArray<int32>& GetBoneIndexTracker() const;
 	
-	const FTransform& GetTransform(const int32 InBoneIndex, const bool bGlobal = false) const;
+	UE_API const FTransform& GetTransform(const int32 InBoneIndex, const bool bGlobal = false) const;
 	
 private:
 
-	bool IsReferenceSkeletonValid(const bool bLog = true) const;
+	UE_API bool IsReferenceSkeletonValid(const bool bLog = true) const;
 
-	void UpdateBoneTracker(const TArray<FMeshBoneInfo>& InOtherInfos);
+	UE_API void UpdateBoneTracker(const TArray<FMeshBoneInfo>& InOtherInfos);
 	
 	// mirroring functions
-	void GetBonesToMirror(const TArray<FName>& InBonesName, const FMirrorOptions& InOptions, TArray<int32>& OutBonesToMirror) const;
-	void GetMirroredNames(const TArray<int32>& InBonesToMirror, const FMirrorOptions& InOptions, TArray<FName>& OutBonesName) const;
-	void GetMirroredBones(const TArray<int32>& InBonesToMirror, const TArray<FName>& InMirroredNames, TArray<int32>& OutMirroredBones);
-	void GetMirroredTransforms(	const TArray<int32>& InBonesToMirror, const TArray<int32>& InMirroredBones,
+	UE_API void GetBonesToMirror(const TArray<FName>& InBonesName, const FMirrorOptions& InOptions, TArray<int32>& OutBonesToMirror) const;
+	UE_API void GetMirroredNames(const TArray<int32>& InBonesToMirror, const FMirrorOptions& InOptions, TArray<FName>& OutBonesName) const;
+	UE_API void GetMirroredBones(const TArray<int32>& InBonesToMirror, const TArray<FName>& InMirroredNames, TArray<int32>& OutMirroredBones);
+	UE_API void GetMirroredTransforms(	const TArray<int32>& InBonesToMirror, const TArray<int32>& InMirroredBones,
 								const FMirrorOptions& InOptions, TArray<FTransform>& OutMirroredTransforms) const;
 
 	// orient function
-	void GetBonesToOrient(const TArray<FName>& InBonesName, const FOrientOptions& InOptions, TArray<int32>& OutBonesToOrient) const;
+	UE_API void GetBonesToOrient(const TArray<FName>& InBonesName, const FOrientOptions& InOptions, TArray<int32>& OutBonesToOrient) const;
 
 	// pre/post commit
-	ESkeletalMeshModificationType PreCommitSkeletalMesh();
-	ESkeletonModificationType PreCommitSkeleton(const ESkeletalMeshModificationType InSkeletalMeshModifications) const;
-	void PostCommitSkeleton(const ESkeletonModificationType InSkeletonModifications) const;
+	UE_API ESkeletalMeshModificationType PreCommitSkeletalMesh();
+	UE_API ESkeletonModificationType PreCommitSkeleton(const ESkeletalMeshModificationType InSkeletalMeshModifications) const;
+	UE_API void PostCommitSkeleton(const ESkeletonModificationType InSkeletonModifications) const;
 
 	// does the skeleton have any references other than the skeletal mesh currently being edited
-	bool HasAnyOtherReferences(const bool bSkeletalMeshOnly) const;
+	UE_API bool HasAnyOtherReferences(const bool bSkeletalMeshOnly) const;
 
 	// update mesh description on commit
-	void CommitChangesToMeshDescription(const ESkeletalMeshModificationType InSkeletalMeshModifications);
+	UE_API void CommitChangesToMeshDescription(const ESkeletalMeshModificationType InSkeletalMeshModifications);
 	
 	// notification
-	void NotifyFromSkeletonChanges() const;
+	UE_API void NotifyFromSkeletonChanges() const;
 	
 	UPROPERTY(transient)
 	TWeakObjectPtr<USkeletalMesh> SkeletalMesh = nullptr;
@@ -369,3 +371,5 @@ public:
 	UPROPERTY(EditAnywhere, Category="Merge", meta=(EditCondition="MergeType==ESKeletalMeshMergeType::Merge"))
 	bool bMergeAll = false;
 };
+
+#undef UE_API
diff --git a/Engine/Plugins/Runtime/MeshModelingToolset/Source/SkeletalMeshModifiers/Public/SkinWeightModifier.h b/Engine/Plugins/Runtime/MeshModelingToolset/Source/SkeletalMeshModifiers/Public/SkinWeightModifier.h
--- a/Engine/Plugins/Runtime/MeshModelingToolset/Source/SkeletalMeshModifiers/Public/SkinWeightModifier.h
+++ b/Engine/Plugins/Runtime/MeshModelingToolset/Source/SkeletalMeshModifiers/Public/SkinWeightModifier.h
@@ -8,6 +8,8 @@
 
 #include "SkinWeightModifier.generated.h"
 
+#define UE_API SKELETALMESHMODIFIERS_API
+
 class USkeletalMesh;
 
 /**
@@ -58,8 +60,8 @@ weight_modifier.commit_weights_to_skeletal_mesh()
  *
  * Though it may be useful to normalize while editing.
  */
-UCLASS(BlueprintType)
-class SKELETALMESHMODIFIERS_API USkinWeightModifier : public UObject
+UCLASS(MinimalAPI, BlueprintType)
+class USkinWeightModifier : public UObject
 {
 	GENERATED_BODY()
 
@@ -71,7 +73,7 @@ public:
 	 * @return bool - true if the mesh weights were successfully loaded.
 	 */
 	UFUNCTION(BlueprintCallable, Category = "Mesh")
-	bool SetSkeletalMesh(USkeletalMesh* InMesh);
+	UE_API bool SetSkeletalMesh(USkeletalMesh* InMesh);
 
 	/**
 	 * Actually applies the weight modifications to the skeletal mesh. This action creates an undo transaction.
@@ -79,7 +81,7 @@ public:
 	 * @return true if weights were applied to a skeletal mesh, false otherwise
 	 */
 	UFUNCTION(BlueprintCallable, Category = "Mesh")
-	bool CommitWeightsToSkeletalMesh();
+	UE_API bool CommitWeightsToSkeletalMesh();
 
 	/**
 	 * Get a reference to the skeletal mesh that was loaded
@@ -93,14 +95,14 @@ public:
 	 * @return int, number of vertices
 	 */
 	UFUNCTION(BlueprintCallable, Category = "Mesh Query")
-	int32 GetNumVertices();
+	UE_API int32 GetNumVertices();
 
 	/**
 	 * Get an array of all bone names in the skeletal mesh.
 	 * @return array of bone names
 	 */
 	UFUNCTION(BlueprintCallable, Category = "Mesh Query")
-	TArray<FName> GetAllBoneNames();
+	UE_API TArray<FName> GetAllBoneNames();
 
 	/**
 	 * Get all bone weights for a single vertex.
@@ -108,7 +110,7 @@ public:
 	 * @return a map of Bone Name to weight values for all bones that influence the specified vertex.
 	 */
 	UFUNCTION(BlueprintCallable, Category = "Weights")
-	TMap<FName, float> GetVertexWeights(const int32 VertexID);
+	UE_API TMap<FName, float> GetVertexWeights(const int32 VertexID);
 
 	/**
 	 * Set bone weights for a single vertex. The weights are stored as supplied and not normalized until
@@ -118,7 +120,7 @@ public:
 	 * @param bReplaceAll if true, all weights on this vertex will be replaced with the input weights. Default is false.
 	 */
 	UFUNCTION(BlueprintCallable, Category = "Weights")
-	bool SetVertexWeights(
+	UE_API bool SetVertexWeights(
 		const int32 VertexID,
 		const TMap<FName, float>& InWeights,
 		const bool bReplaceAll=false);
@@ -129,14 +131,14 @@ public:
 	 * @return true if normalization was performed, false otherwise
 	 */
 	UFUNCTION(BlueprintCallable, Category = "Weights")
-	bool NormalizeVertexWeights(const int32 VertexID);
+	UE_API bool NormalizeVertexWeights(const int32 VertexID);
 
 	/**
 	 * Normalize weights on all vertices in the mesh.
 	 * @return true if normalization was performed, false otherwise
 	 */
 	UFUNCTION(BlueprintCallable, Category = "Weights")
-	bool NormalizeAllWeights();
+	UE_API bool NormalizeAllWeights();
 
 	/**
 	 * Strips out smallest influences to ensure each vertex does not have weight on more influences than MaxInfluences.
@@ -145,7 +147,7 @@ public:
 	 * @return true if influences were removed, false otherwise
 	 */
 	UFUNCTION(BlueprintCallable, Category = "Weights")
-	bool EnforceMaxInfluences(const int32 MaxInfluences=-1);
+	UE_API bool EnforceMaxInfluences(const int32 MaxInfluences=-1);
 
 	/**
 	 * Remove all weights below the given threshold value, on the given vertex.
@@ -155,7 +157,7 @@ public:
 	 * @return true if influences were removed, false otherwise
 	 */
 	UFUNCTION(BlueprintCallable, Category = "Weights")
-	bool PruneVertexWeights(const int32 VertexID, const float WeightThreshold);
+	UE_API bool PruneVertexWeights(const int32 VertexID, const float WeightThreshold);
 
 	/**
 	 * Remove all weights below the given threshold value, on all vertices.
@@ -163,7 +165,7 @@ public:
 	 * @return true if influences were removed, false otherwise
 	 */
 	UFUNCTION(BlueprintCallable, Category = "Weights")
-	bool PruneAllWeights(const float WeightThreshold);
+	UE_API bool PruneAllWeights(const float WeightThreshold);
 
 private:
 
@@ -176,12 +178,14 @@ private:
 	TArray<TMap<FName, float>> Weights;
 
 	// fast internal procedures (do not validate inputs)
-	bool NormalizeVertexInternal(const int32 VertexID);
-	void EnforceMaxInfluenceInternal(const int32 VertexID, const int32 Max);
-	bool PruneVertexWeightInternal(const int32 VertexID, const float WeightThreshold);
+	UE_API bool NormalizeVertexInternal(const int32 VertexID);
+	UE_API void EnforceMaxInfluenceInternal(const int32 VertexID, const int32 Max);
+	UE_API bool PruneVertexWeightInternal(const int32 VertexID, const float WeightThreshold);
 
 	// reset all internal data when loading a new mesh
-	void Reset();
+	UE_API void Reset();
 
 	static constexpr int32 LODIndex = 0;
 };
+
+#undef UE_API
diff --git a/Engine/Plugins/Runtime/Metasound/Source/MetasoundEditor/Private/MetasoundPinAudioInspector.h b/Engine/Plugins/Runtime/Metasound/Source/MetasoundEditor/Private/MetasoundPinAudioInspector.h
--- a/Engine/Plugins/Runtime/Metasound/Source/MetasoundEditor/Private/MetasoundPinAudioInspector.h
+++ b/Engine/Plugins/Runtime/Metasound/Source/MetasoundEditor/Private/MetasoundPinAudioInspector.h
@@ -7,6 +7,8 @@
 #include "SMetasoundPinAudioInspector.h"
 #include "Templates/SharedPointer.h"
 
+#define UE_API METASOUNDEDITOR_API
+
 namespace AudioWidgets
 {
 	class FAudioOscilloscope;
@@ -18,17 +20,17 @@ namespace Metasound
 {
 	namespace Editor
 	{
-		class METASOUNDEDITOR_API FMetasoundPinAudioInspector
+		class FMetasoundPinAudioInspector
 		{
 		public:
-			explicit FMetasoundPinAudioInspector(FEdGraphPinReference InPinRef);
-			virtual ~FMetasoundPinAudioInspector();
+			UE_API explicit FMetasoundPinAudioInspector(FEdGraphPinReference InPinRef);
+			UE_API virtual ~FMetasoundPinAudioInspector();
 
-			TSharedPtr<SMetasoundPinAudioInspector> GetWidget();
+			UE_API TSharedPtr<SMetasoundPinAudioInspector> GetWidget();
 			
 		private:
-			const UMetasoundEditorGraphNode& GetReroutedNode() const;
-			FGraphConnectionManager* GetConnectionManager();
+			UE_API const UMetasoundEditorGraphNode& GetReroutedNode() const;
+			UE_API FGraphConnectionManager* GetConnectionManager();
 
 			UEdGraphPin* GraphPinObj = nullptr;
 
@@ -38,3 +40,5 @@ namespace Metasound
 		};
 	} // namespace Editor
 } // namespace Metasound
+
+#undef UE_API
diff --git a/Engine/Plugins/Runtime/Metasound/Source/MetasoundEditor/Private/SMetasoundPinAudioInspector.h b/Engine/Plugins/Runtime/Metasound/Source/MetasoundEditor/Private/SMetasoundPinAudioInspector.h
--- a/Engine/Plugins/Runtime/Metasound/Source/MetasoundEditor/Private/SMetasoundPinAudioInspector.h
+++ b/Engine/Plugins/Runtime/Metasound/Source/MetasoundEditor/Private/SMetasoundPinAudioInspector.h
@@ -6,11 +6,13 @@
 #include "Widgets/SNullWidget.h"
 #include "Widgets/SWidget.h"
 
+#define UE_API METASOUNDEDITOR_API
+
 namespace Metasound
 {
 	namespace Editor
 	{
-		class METASOUNDEDITOR_API SMetasoundPinAudioInspector : public SPinValueInspector
+		class SMetasoundPinAudioInspector : public SPinValueInspector
 		{
 
 		public:
@@ -21,7 +23,9 @@ namespace Metasound
 			SLATE_ARGUMENT(TSharedRef<SWidget>, VisualizationWidget)
 			SLATE_END_ARGS()
 
-			void Construct(const FArguments& InArgs);
+			UE_API void Construct(const FArguments& InArgs);
 		};
 	} // namespace Editor
 } // namespace Metasound
+
+#undef UE_API
diff --git a/Engine/Plugins/Runtime/Metasound/Source/MetasoundEditor/Private/SMetasoundPinValueInspector.h b/Engine/Plugins/Runtime/Metasound/Source/MetasoundEditor/Private/SMetasoundPinValueInspector.h
--- a/Engine/Plugins/Runtime/Metasound/Source/MetasoundEditor/Private/SMetasoundPinValueInspector.h
+++ b/Engine/Plugins/Runtime/Metasound/Source/MetasoundEditor/Private/SMetasoundPinValueInspector.h
@@ -19,6 +19,8 @@
 #include "Widgets/SWidget.h"
 #include "Widgets/Colors/SColorPicker.h"
 
+#define UE_API METASOUNDEDITOR_API
+
 
 // Forward Declarations
 class UMetasoundEditorGraphNode;
@@ -104,7 +106,7 @@ namespace Metasound
 			void OnValueCommitted(const FText& InValueText, FMetasoundFrontendLiteral& OutNewLiteral, int32 InIndex);
 		};
 
-		class METASOUNDEDITOR_API SMetasoundPinValueInspector : public SPinValueInspector
+		class SMetasoundPinValueInspector : public SPinValueInspector
 		{
 			TSharedPtr<FMetasoundNumericDebugLineItem> LineItem;
 
@@ -114,12 +116,13 @@ namespace Metasound
 			}
 			SLATE_END_ARGS()
 
-			void Construct(const FArguments& InArgs);
+			UE_API void Construct(const FArguments& InArgs);
 
-			void UpdateMessage();
+			UE_API void UpdateMessage();
 
-			virtual void PopulateTreeView() override;
+			UE_API virtual void PopulateTreeView() override;
 		};
 	} // namespace Editor
 } // namespace Metasound
 
+#undef UE_API
diff --git a/Engine/Plugins/Runtime/Metasound/Source/MetasoundEditor/Public/IMetasoundEditor.h b/Engine/Plugins/Runtime/Metasound/Source/MetasoundEditor/Public/IMetasoundEditor.h
--- a/Engine/Plugins/Runtime/Metasound/Source/MetasoundEditor/Public/IMetasoundEditor.h
+++ b/Engine/Plugins/Runtime/Metasound/Source/MetasoundEditor/Public/IMetasoundEditor.h
@@ -8,7 +8,7 @@
 #include "UObject/Object.h"
 
 
-class METASOUNDEDITOR_API IMetasoundEditor : public FAssetEditorToolkit
+class IMetasoundEditor : public FAssetEditorToolkit
 {
 public:
 	// Returns the editor's assigned MetaSound UObject.
