ModelingTools: Fixed a GC crash with mesh element selection's usage of PreviewGeometry.

In cases where GC destroys the PreviewActor owned by PreviewGeometry, the pointers held by the PreviewGeometry are nulled out. This change ensures that we properly check against those cases.

diff --git a/Engine/Plugins/Runtime/MeshModelingToolset/Source/ModelingComponents/Private/Drawing/PreviewGeometryActor.cpp b/Engine/Plugins/Runtime/MeshModelingToolset/Source/ModelingComponents/Private/Drawing/PreviewGeometryActor.cpp
--- a/Engine/Plugins/Runtime/MeshModelingToolset/Source/ModelingComponents/Private/Drawing/PreviewGeometryActor.cpp
+++ b/Engine/Plugins/Runtime/MeshModelingToolset/Source/ModelingComponents/Private/Drawing/PreviewGeometryActor.cpp
@@ -14,7 +14,7 @@
 
 UPreviewGeometry::~UPreviewGeometry()
 {
-	checkf(ParentActor == nullptr, TEXT("You must explicitly Disconnect() UPreviewGeometry before it is GCd"));
+	ensureMsgf(ParentActor == nullptr, TEXT("You must explicitly Disconnect() UPreviewGeometry before it is GCd"));
 }
 
 void UPreviewGeometry::CreateInWorld(UWorld* World, const FTransform& WithTransform)
@@ -71,21 +71,35 @@ void UPreviewGeometry::SetAllVisible(bool bVisible)
 {
 	for (TPair<FString, TObjectPtr<UTriangleSetComponent>> Entry : TriangleSets)
 	{
-		Entry.Value->SetVisibility(bVisible);
+		if (Entry.Value)
+		{
+			Entry.Value->SetVisibility(bVisible);
+		}
 	}
 	for (TPair<FString, TObjectPtr<ULineSetComponent>> Entry : LineSets)
 	{
-		Entry.Value->SetVisibility(bVisible);
+		if (Entry.Value)
+		{
+			Entry.Value->SetVisibility(bVisible);
+		}
 	}
 	for (TPair<FString, TObjectPtr<UPointSetComponent>> Entry : PointSets)
 	{
-		Entry.Value->SetVisibility(bVisible);
+		if (Entry.Value)
+		{
+			Entry.Value->SetVisibility(bVisible);
+		}
 	}
 }
 
 
 UTriangleSetComponent* UPreviewGeometry::AddTriangleSet(const FString& SetIdentifier)
 {
+	if (!ParentActor)
+	{
+		return nullptr;
+	}
+	
 	if (TriangleSets.Contains(SetIdentifier))
 	{
 		check(false);
@@ -136,6 +150,11 @@ UTriangleSetComponent* UPreviewGeometry::CreateOrUpdateTriangleSet(const FString
 
 ULineSetComponent* UPreviewGeometry::AddLineSet(const FString& SetIdentifier)
 {
+	if (!ParentActor)
+	{
+		return nullptr;
+	}
+	
 	if (LineSets.Contains(SetIdentifier))
 	{
 		check(false);
@@ -175,7 +194,7 @@ bool UPreviewGeometry::RemoveLineSet(const FString& LineSetIdentifier, bool bDes
 	{
 		ULineSetComponent* LineSet = *Found;
 		LineSets.Remove(LineSetIdentifier);
-		if (bDestroy)
+		if (bDestroy && LineSet)
 		{
 			LineSet->UnregisterComponent();
 			LineSet->DestroyComponent();
@@ -194,8 +213,11 @@ void UPreviewGeometry::RemoveAllLineSets(bool bDestroy)
 	{
 		for (TPair<FString, TObjectPtr<ULineSetComponent>> Entry : LineSets)
 		{
-			Entry.Value->UnregisterComponent();
-			Entry.Value->DestroyComponent();
+			if (Entry.Value)
+			{
+				Entry.Value->UnregisterComponent();
+				Entry.Value->DestroyComponent();
+			}
 		}
 	}
 	LineSets.Reset();
@@ -210,7 +232,7 @@ bool UPreviewGeometry::RemoveTriangleSet(const FString& TriangleSetIdentifier, b
 		UTriangleSetComponent* TriangleSet = *Found;
 
 		TriangleSets.Remove(TriangleSetIdentifier);
-		if (bDestroy)
+		if (bDestroy && TriangleSet)
 		{
 			TriangleSet->UnregisterComponent();
 			TriangleSet->DestroyComponent();
@@ -229,8 +251,11 @@ void UPreviewGeometry::RemoveAllTriangleSets(bool bDestroy)
 	{
 		for (TPair<FString, TObjectPtr<UTriangleSetComponent>> Entry : TriangleSets)
 		{
-			Entry.Value->UnregisterComponent();
-			Entry.Value->DestroyComponent();
+			if (Entry.Value)
+			{
+				Entry.Value->UnregisterComponent();
+				Entry.Value->DestroyComponent();
+			}
 		}
 	}
 	TriangleSets.Reset();
@@ -266,7 +291,10 @@ void UPreviewGeometry::SetAllLineSetsMaterial(UMaterialInterface* Material)
 {
 	for (TPair<FString, TObjectPtr<ULineSetComponent>> Entry : LineSets)
 	{
-		Entry.Value->SetLineMaterial(Material);
+		if (Entry.Value)
+		{
+			Entry.Value->SetLineMaterial(Material);
+		}
 	}
 }
 
@@ -295,6 +323,11 @@ ULineSetComponent* UPreviewGeometry::CreateOrUpdateLineSet(const FString& LineSe
 
 UPointSetComponent* UPreviewGeometry::AddPointSet(const FString& SetIdentifier)
 {
+	if (!ParentActor)
+	{
+		return nullptr;
+	}
+	
 	if (PointSets.Contains(SetIdentifier))
 	{
 		check(false);
@@ -352,7 +385,7 @@ bool UPreviewGeometry::RemovePointSet(const FString& PointSetIdentifier, bool bD
 	{
 		UPointSetComponent* PointSet = *Found;
 		PointSets.Remove(PointSetIdentifier);
-		if (bDestroy)
+		if (bDestroy && PointSet)
 		{
 			PointSet->UnregisterComponent();
 			PointSet->DestroyComponent();
@@ -369,8 +402,11 @@ void UPreviewGeometry::RemoveAllPointSets(bool bDestroy)
 	{
 		for (TPair<FString, TObjectPtr<UPointSetComponent>> Entry : PointSets)
 		{
-			Entry.Value->UnregisterComponent();
-			Entry.Value->DestroyComponent();
+			if (Entry.Value)
+			{
+				Entry.Value->UnregisterComponent();
+				Entry.Value->DestroyComponent();
+			}
 		}
 	}
 	PointSets.Reset();
@@ -402,6 +438,9 @@ void UPreviewGeometry::SetAllPointSetsMaterial(UMaterialInterface* Material)
 {
 	for (TPair<FString, TObjectPtr<UPointSetComponent>> Entry : PointSets)
 	{
-		Entry.Value->SetPointMaterial(Material);
+		if (Entry.Value)
+		{
+			Entry.Value->SetPointMaterial(Material);
+		}
 	}
 }
