MeshModelingToolset - Geometry Selection Changes - Add support for conditionally selectable elements and allow selector factores to be unregistered.

Add support for selectors to notify the manager that the set of selectable elements has changed to support conditional selectability.

Add virtual function to the selector interface which takes the persistent selection as input when determining doing preview raycasts in order to use the persistent selection as an input to the selectability function.

Selector factories can now be unregistered. Currently only a bulk operation. Not used in this CL, used in a later change.

diff --git a/Engine/Plugins/Runtime/MeshModelingToolset/Source/ModelingComponents/Private/Selection/GeometrySelectionManager.cpp b/Engine/Plugins/Runtime/MeshModelingToolset/Source/ModelingComponents/Private/Selection/GeometrySelectionManager.cpp
--- a/Engine/Plugins/Runtime/MeshModelingToolset/Source/ModelingComponents/Private/Selection/GeometrySelectionManager.cpp
+++ b/Engine/Plugins/Runtime/MeshModelingToolset/Source/ModelingComponents/Private/Selection/GeometrySelectionManager.cpp
@@ -45,6 +45,12 @@ void UGeometrySelectionManager::RegisterSelectorFactory(TUniquePtr<IGeometrySele
 	ResetTargetCache();
 }
 
+void UGeometrySelectionManager::UnregisterSelectorFactories()
+{
+	Factories.Empty();
+	ResetTargetCache();
+}
+
 void UGeometrySelectionManager::Shutdown()
 {
 	DiscardSavedSelection();
@@ -779,6 +785,12 @@ TSharedPtr<UGeometrySelectionManager::FGeometrySelectionTarget> UGeometrySelecti
 		TargetCache.Add(TargetIdentifier, SelectionTarget);
 	}
 	
+	SelectionTarget->Selector->GetOnSelectableElementsChanged().AddWeakLambda(this, [this]()
+	{
+		RenderCachesDirtyFlags |= EEnumerateRenderCachesDirtyFlags::UnselectedCachesDirty;
+		RebuildSelectableRenderCache();
+	});
+	
 	return SelectionTarget;
 }
 
@@ -1179,7 +1191,7 @@ bool UGeometrySelectionManager::UpdateSelectionPreviewViaRaycast(
 	QueryAPI->GetCurrentViewState(RayQueryInfo.CameraState);
 
 	FGeometrySelectionPreview NewPreview( *(Target->SelectionEditor) );
-	Target->Selector->GetSelectionPreviewForRaycast( RayQueryInfo, NewPreview );
+	Target->Selector->GetConditionalSelectionPreviewForRaycast( RayQueryInfo, NewPreview, &Target->Selection );
 	if ( ! UE::Geometry::AreSelectionsIdentical(NewPreview.PreviewSelection, ActivePreviewSelection) )
 	{
 		ActivePreviewSelection = MoveTemp(NewPreview.PreviewSelection);
@@ -1452,7 +1464,11 @@ void UGeometrySelectionManager::ExecuteSelectionCommand(UGeometrySelectionEditCo
 		SelectionArguments->ElementType = SelectionElementType;
 		SelectionArguments->TopologyMode = GetSelectionTopologyType();
 		SelectionArguments->SetTransactionsAPI(TransactionsAPI);
-		Command->ExecuteCommand(SelectionArguments, &ResultPtr);
+
+		if (Command->CanExecuteCommand(SelectionArguments))
+		{
+			Command->ExecuteCommand(SelectionArguments, &ResultPtr);
+		}
 
 		// actually clear selection after executing command. 
 		FGeometrySelectionDelta ClearDelta;
@@ -1651,39 +1667,33 @@ void UGeometrySelectionManager::RemoveSets(const TArrayView<FString>& SetIdentif
 
 void UGeometrySelectionManager::CreateOrUpdateAllSets(const FGeometrySelectionElements& Elements, const FMeshElementSelectionParams SelectionParams) const
 {
-
 	UWorld* World = GetWorld();
 	if (World && !PreviewGeometry->ParentActor)
 	{
 		PreviewGeometry->CreateInWorld(World, FTransform::Identity);
 	}
-
-	if(SelectionElementType == EGeometryElementType::Edge || SelectionElementType == EGeometryElementType::Face)
+	
+	PreviewGeometry->CreateOrUpdateLineSet(SelectionParams.Identifiers[1], Elements.Segments.Num(), [&](int32 j, TArray<FRenderableLine>& LinesOut)
 	{
-		PreviewGeometry->CreateOrUpdateLineSet(SelectionParams.Identifiers[1], Elements.Segments.Num(), [&](int32 j, TArray<FRenderableLine>& LinesOut)
-		{
 		const FSegment3d Seg = Elements.Segments[j];
 		LinesOut.Add(FRenderableLine(Seg.StartPoint(), Seg.EndPoint(),  SelectionParams.Color,  SelectionParams.LineThickness, SelectionParams.DepthBias));
-		}, 1);
-	
-		PreviewGeometry->CreateOrUpdateTriangleSet(SelectionParams.Identifiers[2], Elements.Triangles.Num(), [&](int32 k, TArray<FRenderableTriangle>& TrianglesOut)
-			{
-			const FTriangle3d Triangle = Elements.Triangles[k];
-			const FVector3d Normal = Triangle.Normal();
-			FRenderableTriangleVertex A(Triangle.V[0], FVector2D(0,0), Normal, SelectionParams.Color);
-			FRenderableTriangleVertex B(Triangle.V[1], FVector2D(1,0), Normal, SelectionParams.Color);
-			FRenderableTriangleVertex C(Triangle.V[2], FVector2D(1,1), Normal, SelectionParams.Color);
-			TrianglesOut.Add(FRenderableTriangle(SelectionParams.SelectionFillColor, A, B, C));
-			},1 );
-	}
-	else if (SelectionElementType == EGeometryElementType::Vertex)
+	}, 1);
+
+	PreviewGeometry->CreateOrUpdateTriangleSet(SelectionParams.Identifiers[2], Elements.Triangles.Num(), [&](int32 k, TArray<FRenderableTriangle>& TrianglesOut)
 	{
-		PreviewGeometry->CreateOrUpdatePointSet(SelectionParams.Identifiers[0], Elements.Points.Num(), [&](int32 k, TArray<FRenderablePoint>& PointsOut)
-		{
-		const FVector3d Point = Elements.Points[k];
-		PointsOut.Add(FRenderablePoint(Point, SelectionParams.Color, SelectionParams.PointSize, SelectionParams.DepthBias));
-		});
-	}
+		const FTriangle3d Triangle = Elements.Triangles[k];
+		const FVector3d Normal = Triangle.Normal();
+		FRenderableTriangleVertex A(Triangle.V[0], FVector2D(0,0), Normal, SelectionParams.Color);
+		FRenderableTriangleVertex B(Triangle.V[1], FVector2D(1,0), Normal, SelectionParams.Color);
+		FRenderableTriangleVertex C(Triangle.V[2], FVector2D(1,1), Normal, SelectionParams.Color);
+		TrianglesOut.Add(FRenderableTriangle(SelectionParams.SelectionFillColor, A, B, C));
+	},1 );
+
+	PreviewGeometry->CreateOrUpdatePointSet(SelectionParams.Identifiers[0], Elements.Points.Num(), [&](int32 k, TArray<FRenderablePoint>& PointsOut)
+    {
+	    const FVector3d Point = Elements.Points[k];
+	    PointsOut.Add(FRenderablePoint(Point, SelectionParams.Color, SelectionParams.PointSize, SelectionParams.DepthBias));
+    });
 }
 
 void UGeometrySelectionManager::RebuildSelectionRenderCache()
@@ -1717,20 +1727,18 @@ void UGeometrySelectionManager::RebuildSelectableRenderCache()
 	for (int32 k = 0; k < ActiveTargetReferences.Num(); ++k)
 	{
 		TSharedPtr<FGeometrySelectionTarget> Target = ActiveTargetReferences[k];
-		FGeometrySelectionElements& AllElements = CachedSelectableRenderElements[k];
 		
+		FGeometrySelectionElements& AllElements = CachedSelectableRenderElements[k];
 		AllElements.Reset();
 		
-		Target->Selector->AccumulateElementsFromPredicate(AllElements, true, false, MeshTopologyMode == EMeshTopologyMode::Polygroup, [Target, this](EGeometryElementType Type, FGeoSelectionID ID)
-		{
-			// Selectable faces are not displayed directly, just implicitly via displayed edges.
-			if (Type == EGeometryElementType::Face)
-			{
-				return false;
-			}
+		FGeometrySelection AllElementSelection;
+		AllElementSelection.ElementType = GetSelectionElementType();
+		AllElementSelection.TopologyType = GetSelectionTopologyType();
 			
-			return true;
-		});
+		// Get full selection, according to the type and topo of AllElementSelection
+		Target->Selector->InitializeSelectionFromPredicate(AllElementSelection, [](FGeoSelectionID) -> bool { return true; }, IGeometrySelector::EInitializeSelectionMode::All, &Target->Selection);
+		Target->Selector->AccumulateSelectionElements(AllElementSelection, AllElements, true, EEnumerateSelectionMapping::FacesToEdges);
+		
 		CreateOrUpdateAllSets(AllElements, UnselectedParams);
 	}
 
diff --git a/Engine/Plugins/Runtime/MeshModelingToolset/Source/ModelingComponents/Public/Selection/GeometrySelectionManager.h b/Engine/Plugins/Runtime/MeshModelingToolset/Source/ModelingComponents/Public/Selection/GeometrySelectionManager.h
--- a/Engine/Plugins/Runtime/MeshModelingToolset/Source/ModelingComponents/Public/Selection/GeometrySelectionManager.h
+++ b/Engine/Plugins/Runtime/MeshModelingToolset/Source/ModelingComponents/Public/Selection/GeometrySelectionManager.h
@@ -100,6 +100,9 @@ public:
 
 	UE_API virtual void RegisterSelectorFactory(TUniquePtr<IGeometrySelectorFactory> Factory);
 
+	UE_EXPERIMENTAL(5.7, TEXT("Selector factory unregistering is experimental while the API is under active development."))
+	UE_API virtual void UnregisterSelectorFactories();
+
 	UE_API virtual void Shutdown();
 	UE_API virtual bool HasBeenShutDown() const;
 
diff --git a/Engine/Plugins/Runtime/MeshModelingToolset/Source/ModelingComponents/Public/Selection/GeometrySelector.h b/Engine/Plugins/Runtime/MeshModelingToolset/Source/ModelingComponents/Public/Selection/GeometrySelector.h
--- a/Engine/Plugins/Runtime/MeshModelingToolset/Source/ModelingComponents/Public/Selection/GeometrySelector.h
+++ b/Engine/Plugins/Runtime/MeshModelingToolset/Source/ModelingComponents/Public/Selection/GeometrySelector.h
@@ -364,6 +364,22 @@ public:
 		const FWorldRayQueryInfo& RayInfo,
 		FGeometrySelectionEditor& PreviewEditor
 	) = 0;
+	
+	/**
+	 * Query the Selector for a potential selection based on world-space Ray. The intention of this query
+	 * is that it be used for "selection preview", eg like hover-highlighting. The PreviewEditor will
+	 * always be updated in 'Add' mode, and currently no result information will be returned.
+	 *
+	 * CurrentSelection is provided for selectors which have conditionally selectable geometry.
+	 * While not enforced, the expectation is that we add elements to the PreviewEditor only if
+	 * those elements are considered selectable with respect to CurrentSelection.
+	 */
+	UE_EXPERIMENTAL(5.7, TEXT("GetConditionalSelectionPreviewForRaycast is experimental while the API is under active development."))
+	virtual void GetConditionalSelectionPreviewForRaycast(
+		const FWorldRayQueryInfo& RayInfo,
+		FGeometrySelectionEditor& PreviewEditor,
+		const FGeometrySelection* CurrentSelection 
+	) { GetSelectionPreviewForRaycast(RayInfo, PreviewEditor); }
 
 
 
@@ -461,10 +477,23 @@ public:
 	 */
 	FOnGeometryModified& GetOnGeometryModifed() { return OnGeometryModifed; }
 
+	/**
+	 * The OnSelectableElementsChanged delegate will fire if the set of selectable elements have changed.
+	 * The set of selectable elements is a function of the state of the selection target and the set of
+	 * currently selected elements (the selection 'owned' by this selector).
+	 */
+	UE_EXPERIMENTAL(5.7, TEXT("Selectable element invalidation notification is experimental while the API is under active development."))
+	FSimpleMulticastDelegate& GetOnSelectableElementsChanged() { return OnSelectableElementsChanged; }
+
 protected:
 	FOnGeometryModified OnGeometryModifed;
 	void NotifyGeometryModified() { OnGeometryModifed.Broadcast(this); }
 
+	UE_EXPERIMENTAL(5.7, TEXT("NotifySelectableElementsChanged is experimental while the API is under active development."))
+	void NotifySelectableElementsChanged() const { OnSelectableElementsChanged.Broadcast(); }
+	
+private:
+	FSimpleMulticastDelegate OnSelectableElementsChanged;
 };
 
 
