ModelingTools: Fixed UE_EXPERIMENTAL usage in GeometrySelectionManager & GeometrySelector

diff --git a/Engine/Plugins/Runtime/MeshModelingToolset/Source/ModelingComponents/Public/Selection/GeometrySelectionManager.h b/Engine/Plugins/Runtime/MeshModelingToolset/Source/ModelingComponents/Public/Selection/GeometrySelectionManager.h
--- a/Engine/Plugins/Runtime/MeshModelingToolset/Source/ModelingComponents/Public/Selection/GeometrySelectionManager.h
+++ b/Engine/Plugins/Runtime/MeshModelingToolset/Source/ModelingComponents/Public/Selection/GeometrySelectionManager.h
@@ -100,7 +100,7 @@ public:
 
 	UE_API virtual void RegisterSelectorFactory(TUniquePtr<IGeometrySelectorFactory> Factory);
 
-	UE_EXPERIMENTAL(5.7, TEXT("Selector factory unregistering is experimental while the API is under active development."))
+	UE_EXPERIMENTAL(5.7, "Selector factory unregistering is experimental while the API is under active development.")
 	UE_API virtual void UnregisterSelectorFactories();
 
 	UE_API virtual void Shutdown();
diff --git a/Engine/Plugins/Runtime/MeshModelingToolset/Source/ModelingComponents/Public/Selection/GeometrySelector.h b/Engine/Plugins/Runtime/MeshModelingToolset/Source/ModelingComponents/Public/Selection/GeometrySelector.h
--- a/Engine/Plugins/Runtime/MeshModelingToolset/Source/ModelingComponents/Public/Selection/GeometrySelector.h
+++ b/Engine/Plugins/Runtime/MeshModelingToolset/Source/ModelingComponents/Public/Selection/GeometrySelector.h
@@ -374,7 +374,7 @@ public:
 	 * While not enforced, the expectation is that we add elements to the PreviewEditor only if
 	 * those elements are considered selectable with respect to CurrentSelection.
 	 */
-	UE_EXPERIMENTAL(5.7, TEXT("GetConditionalSelectionPreviewForRaycast is experimental while the API is under active development."))
+	UE_EXPERIMENTAL(5.7, "GetConditionalSelectionPreviewForRaycast is experimental while the API is under active development.")
 	virtual void GetConditionalSelectionPreviewForRaycast(
 		const FWorldRayQueryInfo& RayInfo,
 		FGeometrySelectionEditor& PreviewEditor,
@@ -482,14 +482,14 @@ public:
 	 * The set of selectable elements is a function of the state of the selection target and the set of
 	 * currently selected elements (the selection 'owned' by this selector).
 	 */
-	UE_EXPERIMENTAL(5.7, TEXT("Selectable element invalidation notification is experimental while the API is under active development."))
+	UE_EXPERIMENTAL(5.7, "Selectable element invalidation notification is experimental while the API is under active development.")
 	FSimpleMulticastDelegate& GetOnSelectableElementsChanged() { return OnSelectableElementsChanged; }
 
 protected:
 	FOnGeometryModified OnGeometryModifed;
 	void NotifyGeometryModified() { OnGeometryModifed.Broadcast(this); }
 
-	UE_EXPERIMENTAL(5.7, TEXT("NotifySelectableElementsChanged is experimental while the API is under active development."))
+	UE_EXPERIMENTAL(5.7, "NotifySelectableElementsChanged is experimental while the API is under active development.")
 	void NotifySelectableElementsChanged() const { OnSelectableElementsChanged.Broadcast(); }
 	
 private:
