ModelingTools: Speculative fix for crash in DisplaceMeshTool

diff --git a/Engine/Plugins/Runtime/MeshModelingToolset/Source/MeshModelingTools/Private/DisplaceMeshTool.cpp b/Engine/Plugins/Runtime/MeshModelingToolset/Source/MeshModelingTools/Private/DisplaceMeshTool.cpp
--- a/Engine/Plugins/Runtime/MeshModelingToolset/Source/MeshModelingTools/Private/DisplaceMeshTool.cpp
+++ b/Engine/Plugins/Runtime/MeshModelingToolset/Source/MeshModelingTools/Private/DisplaceMeshTool.cpp
@@ -200,7 +200,12 @@ namespace DisplaceMeshToolLocals{
 			OffsetArray.SetNumZeroed(Mesh.MaxVertexID());
 			
 			auto ComputeOffsetForElement = [&](int32 ElementID) 
-			{	
+			{
+				if (!UVOverlay->IsElement(ElementID))
+				{
+					return 0.0;
+				}
+				
 				FVector2f UV = UVOverlay->GetElement(ElementID);
 
 				// Adjust UV value and tile it. 
@@ -264,9 +269,12 @@ namespace DisplaceMeshToolLocals{
 				{
 					// Grab any triangle if it's not a uv seam vertex
 					int TID = Mesh.GetVtxSingleTriangle(VID);
-					int ElementID = UVOverlay->GetElementIDAtVertex(TID, VID);
-					double Offset = ComputeOffsetForElement(ElementID);
-					OffsetArray[VID] = Offset;
+					if (Mesh.IsTriangle(TID))
+					{
+						int ElementID = UVOverlay->GetElementIDAtVertex(TID, VID);
+						double Offset = ComputeOffsetForElement(ElementID);
+						OffsetArray[VID] = Offset;
+					}
 				}
 				
 				double Intensity = IntensityFunc(VID, Positions[VID], Normals[VID]);
