HLOD - ApproximateActors: Fixed issue where HLOD materials couldn't be created with a Material Instance parent
* Prior to this, the MI's parent was always used

diff --git a/Engine/Plugins/Runtime/MeshModelingToolset/Source/ModelingComponentsEditorOnly/Private/AssetUtils/CreateMaterialUtil.cpp b/Engine/Plugins/Runtime/MeshModelingToolset/Source/ModelingComponentsEditorOnly/Private/AssetUtils/CreateMaterialUtil.cpp
--- a/Engine/Plugins/Runtime/MeshModelingToolset/Source/ModelingComponentsEditorOnly/Private/AssetUtils/CreateMaterialUtil.cpp
+++ b/Engine/Plugins/Runtime/MeshModelingToolset/Source/ModelingComponentsEditorOnly/Private/AssetUtils/CreateMaterialUtil.cpp
@@ -39,11 +39,9 @@ UE::AssetUtils::ECreateMaterialResult UE::AssetUtils::CreateDuplicateMaterial(
 	}
 
 	FString DestinationLongPackagePath = FPackageName::GetLongPackagePath(Options.NewAssetPath);
-	//FString DestinationObjectName = FPackageName::ObjectPathToObjectName(Options.NewAssetPath);
 	FString DestinationObjectName = FPackageName::GetLongPackageAssetName(Options.NewAssetPath);
 
-	IAssetTools& AssetTools = FModuleManager::Get().LoadModuleChecked<FAssetToolsModule>("AssetTools").Get();
-	UObject* NewAsset = AssetTools.DuplicateAsset(DestinationObjectName, DestinationLongPackagePath, SourceMaterial);
+	UObject* NewAsset = IAssetTools::Get().DuplicateAsset(DestinationObjectName, DestinationLongPackagePath, SourceMaterial);
 	UMaterial* NewMaterial = Cast<UMaterial>(NewAsset);
 
 	if ( ensure(NewMaterial) == false )
@@ -60,12 +58,6 @@ UE::AssetUtils::ECreateMaterialResult UE::AssetUtils::CreateDuplicateMaterial(
 	return ECreateMaterialResult::Ok;
 }
 
-
-
-
-
-
-
 UE::AssetUtils::ECreateMaterialResult UE::AssetUtils::CreateDerivedMaterialInstance(
 	UMaterialInterface* BaseMaterial,
 	FMaterialAssetOptions& Options,
@@ -81,11 +73,6 @@ UE::AssetUtils::ECreateMaterialResult UE::AssetUtils::CreateDerivedMaterialInsta
 	{
 		return ECreateMaterialResult::InvalidBaseMaterial;
 	}
-	UMaterial* SourceMaterial = BaseMaterial->GetBaseMaterial();
-	if (ensure(SourceMaterial) == false)
-	{
-		return ECreateMaterialResult::InvalidBaseMaterial;
-	}
 
 	// cannot create new package if path already exists
 	if (FPackageName::DoesPackageExist(Options.NewAssetPath))
@@ -95,30 +82,14 @@ UE::AssetUtils::ECreateMaterialResult UE::AssetUtils::CreateDerivedMaterialInsta
 	}
 
 	FString DestinationLongPackagePath = FPackageName::GetLongPackagePath(Options.NewAssetPath);
-	//FString DestinationObjectName = FPackageName::ObjectPathToObjectName(Options.NewAssetPath);
 	FString DestinationObjectName = FPackageName::GetLongPackageAssetName(Options.NewAssetPath);
 
-	IAssetTools& AssetTools = FModuleManager::Get().LoadModuleChecked<FAssetToolsModule>("AssetTools").Get();
-
-	bool bSourceIsMIC = (Cast<UMaterialInstanceConstant>(BaseMaterial) != nullptr);
+	UMaterialInstanceConstantFactoryNew* Factory = NewObject<UMaterialInstanceConstantFactoryNew>();
+	Factory->InitialParent = BaseMaterial;
 
-	UMaterialInstanceConstant* NewMaterialInstance;
-	if (bSourceIsMIC)
-	{
-		UObject* DupeAsset = AssetTools.DuplicateAsset(DestinationObjectName, DestinationLongPackagePath, BaseMaterial);
-		NewMaterialInstance = Cast<UMaterialInstanceConstant>(DupeAsset);
-	}
-	else
-	{
-		UMaterialInstanceConstantFactoryNew* Factory = NewObject<UMaterialInstanceConstantFactoryNew>();
-		Factory->InitialParent = SourceMaterial;
-
-		UObject* NewAsset = AssetTools.CreateAsset(DestinationObjectName, DestinationLongPackagePath,
-			UMaterialInstanceConstant::StaticClass(), Factory);
-
-		NewMaterialInstance = Cast<UMaterialInstanceConstant>(NewAsset);
-	}
+	UObject* NewAsset = IAssetTools::Get().CreateAsset(DestinationObjectName, DestinationLongPackagePath, UMaterialInstanceConstant::StaticClass(), Factory);
 
+	UMaterialInstanceConstant* NewMaterialInstance = Cast<UMaterialInstanceConstant>(NewAsset);
 	if (ensure(NewMaterialInstance) == false)
 	{
 		return ECreateMaterialResult::DuplicateFailed;
