Skin Weight Paint Tool: get the mesh description from the source target instead of the tool target

#rnx

diff --git a/Engine/Plugins/Runtime/MeshModelingToolset/Source/MeshModelingToolsEditorOnly/Private/SkeletalMesh/SkinWeightsPaintTool.cpp b/Engine/Plugins/Runtime/MeshModelingToolset/Source/MeshModelingToolsEditorOnly/Private/SkeletalMesh/SkinWeightsPaintTool.cpp
--- a/Engine/Plugins/Runtime/MeshModelingToolset/Source/MeshModelingToolsEditorOnly/Private/SkeletalMesh/SkinWeightsPaintTool.cpp
+++ b/Engine/Plugins/Runtime/MeshModelingToolset/Source/MeshModelingToolsEditorOnly/Private/SkeletalMesh/SkinWeightsPaintTool.cpp
@@ -802,10 +802,13 @@ void UWeightToolTransferManager::SetSourceMesh(USkeletalMesh* InSkeletalMesh)
 	SourcePreviewMesh->SetMaterials(MaterialSet.Materials);
 
 	// setup the mesh selection for the source
-	const FMeshDescription* MeshDescription =  UE::ToolTarget::GetMeshDescription(WeightTool->GetTarget(), SourceParams);
-	MeshSelector->SetMesh(SourcePreviewMesh, MeshDescription, ToolProperties->SourcePreviewOffset);
-	MeshSelector->SetIsEnabled(ToolProperties->EditingMode == EWeightEditMode::Mesh);
-	MeshSelector->SetComponentSelectionMode(ToolProperties->ComponentSelectionMode);
+	const FMeshDescription* MeshDescription =  UE::ToolTarget::GetMeshDescription(SourceTarget, SourceParams);
+	if (ensure(MeshDescription))
+	{
+		MeshSelector->SetMesh(SourcePreviewMesh, MeshDescription, ToolProperties->SourcePreviewOffset);
+		MeshSelector->SetIsEnabled(ToolProperties->EditingMode == EWeightEditMode::Mesh);
+		MeshSelector->SetComponentSelectionMode(ToolProperties->ComponentSelectionMode);
+	}
 
 	WeightTool->UpdateSelectorState();
 }
