GeometrySelectionManager: Do not emit begin a transaction when clearing a selection if a transaction is already in progress.

diff --git a/Engine/Plugins/Runtime/MeshModelingToolset/Source/ModelingComponents/Private/Selection/GeometrySelectionManager.cpp b/Engine/Plugins/Runtime/MeshModelingToolset/Source/ModelingComponents/Private/Selection/GeometrySelectionManager.cpp
--- a/Engine/Plugins/Runtime/MeshModelingToolset/Source/ModelingComponents/Private/Selection/GeometrySelectionManager.cpp
+++ b/Engine/Plugins/Runtime/MeshModelingToolset/Source/ModelingComponents/Private/Selection/GeometrySelectionManager.cpp
@@ -842,6 +842,16 @@ bool UGeometrySelectionManager::RayHitTest(
 
 void UGeometrySelectionManager::ClearSelection(bool bSaveSelectionBeforeClear)
 {
+	bool bTransact = true;
+	
+#if WITH_EDITORONLY_DATA
+	// Cannot update the selection if we're already in a transaction (can happen e.g. when we undo out of a tool)
+	if (GIsTransacting)
+	{
+		bTransact = false;
+	}
+#endif
+	
 	if (!HasSelection())
 	{
 		return;
@@ -852,7 +862,10 @@ void UGeometrySelectionManager::ClearSelection(bool bSaveSelectionBeforeClear)
 		SaveCurrentSelection();
 	}
 
-	GetTransactionsAPI()->BeginUndoTransaction(LOCTEXT("ClearSelection", "Clear Selection"));
+	if (bTransact)
+	{
+		GetTransactionsAPI()->BeginUndoTransaction(LOCTEXT("ClearSelection", "Clear Selection"));
+	}
 
 	for (TSharedPtr<FGeometrySelectionTarget> Target : ActiveTargetReferences)
 	{
@@ -863,11 +876,18 @@ void UGeometrySelectionManager::ClearSelection(bool bSaveSelectionBeforeClear)
 			TUniquePtr<FGeometrySelectionDeltaChange> ClearChange = MakeUnique<FGeometrySelectionDeltaChange>();
 			ClearChange->Identifier = Target->TargetIdentifier;
 			ClearChange->Delta = MoveTemp(ClearDelta);
-			GetTransactionsAPI()->AppendChange(this, MoveTemp(ClearChange), LOCTEXT("ClearSelection", "Clear Selection"));
+
+			if (bTransact)
+			{
+				GetTransactionsAPI()->AppendChange(this, MoveTemp(ClearChange), LOCTEXT("ClearSelection", "Clear Selection"));
+			}
 		}
 	}
 
-	GetTransactionsAPI()->EndUndoTransaction();
+	if (bTransact)
+	{
+		GetTransactionsAPI()->EndUndoTransaction();
+	}
 
 	MarkRenderCachesDirty(false);
 	OnSelectionModified.Broadcast();
