ModelingTools: Speculative crash fix for UVLayoutPreview

diff --git a/Engine/Plugins/Runtime/MeshModelingToolset/Source/ModelingComponents/Private/Drawing/UVLayoutPreview.cpp b/Engine/Plugins/Runtime/MeshModelingToolset/Source/ModelingComponents/Private/Drawing/UVLayoutPreview.cpp
--- a/Engine/Plugins/Runtime/MeshModelingToolset/Source/ModelingComponents/Private/Drawing/UVLayoutPreview.cpp
+++ b/Engine/Plugins/Runtime/MeshModelingToolset/Source/ModelingComponents/Private/Drawing/UVLayoutPreview.cpp
@@ -160,42 +160,47 @@ void UUVLayoutPreview::UpdateUVMesh(const FDynamicMesh3* SourceMesh, int32 Sourc
 	FDynamicMeshNormalOverlay* NewNormalOverlay = UVMesh.Attributes()->PrimaryNormals();
 
 	FAxisAlignedBox2f Bounds = FAxisAlignedBox2f(FVector2f::Zero(), FVector2f::One());
-	const FDynamicMeshUVOverlay* UVOverlay = SourceMesh->Attributes()->GetUVLayer(SourceUVLayer);
-	for (int32 tid : SourceMesh->TriangleIndicesItr())
+	if (SourceMesh && SourceMesh->Attributes())
 	{
-		if (UVOverlay->IsSetTriangle(tid))
+		if (const FDynamicMeshUVOverlay* UVOverlay = SourceMesh->Attributes()->GetUVLayer(SourceUVLayer))
 		{
-			FIndex3i UVTri = UVOverlay->GetTriangle(tid);
-
-			FVector2f UVs[3];
-			for (int32 j = 0; j < 3; ++j)
-			{
-				UVs[j] = UVOverlay->GetElement(UVTri[j]);
-				Bounds.Contain(UVs[j]);
-			}
-
-			FIndex3i NewTri, NewUVTri, NewNormalTri;
-			for (int32 j = 0; j < 3; ++j)
-			{
-				NewUVTri[j] = NewUVOverlay->AppendElement(UVs[j]);
-				NewTri[j] = UVMesh.AppendVertex(FVector3d(UVs[j].X, UVs[j].Y, 0));
-				NewNormalTri[j] = NewNormalOverlay->AppendElement( FVector3f::UnitZ() );
-			}
-
-			FVector2f EdgeUV1 = UVs[1] - UVs[0];
-			FVector2f EdgeUV2 = UVs[2] - UVs[0];
-			float SignedUVArea = 0.5f * (EdgeUV1.X * EdgeUV2.Y - EdgeUV1.Y * EdgeUV2.X);
-
-			if (SignedUVArea > 0)
+			for (int32 tid : SourceMesh->TriangleIndicesItr())
 			{
-				Swap(NewTri.A, NewTri.B);
-				Swap(NewUVTri.A, NewUVTri.B);
+				if (UVOverlay->IsSetTriangle(tid))
+				{
+					FIndex3i UVTri = UVOverlay->GetTriangle(tid);
+
+					FVector2f UVs[3];
+					for (int32 j = 0; j < 3; ++j)
+					{
+						UVs[j] = UVOverlay->GetElement(UVTri[j]);
+						Bounds.Contain(UVs[j]);
+					}
+
+					FIndex3i NewTri, NewUVTri, NewNormalTri;
+					for (int32 j = 0; j < 3; ++j)
+					{
+						NewUVTri[j] = NewUVOverlay->AppendElement(UVs[j]);
+						NewTri[j] = UVMesh.AppendVertex(FVector3d(UVs[j].X, UVs[j].Y, 0));
+						NewNormalTri[j] = NewNormalOverlay->AppendElement( FVector3f::UnitZ() );
+					}
+
+					FVector2f EdgeUV1 = UVs[1] - UVs[0];
+					FVector2f EdgeUV2 = UVs[2] - UVs[0];
+					float SignedUVArea = 0.5f * (EdgeUV1.X * EdgeUV2.Y - EdgeUV1.Y * EdgeUV2.X);
+
+					if (SignedUVArea > 0)
+					{
+						Swap(NewTri.A, NewTri.B);
+						Swap(NewUVTri.A, NewUVTri.B);
+					}
+
+
+					int32 NewTriID = UVMesh.AppendTriangle(NewTri);
+					NewUVOverlay->SetTriangle(NewTriID, NewUVTri);
+					NewNormalOverlay->SetTriangle(NewTriID, NewNormalTri);
+				}
 			}
-
-
-			int32 NewTriID = UVMesh.AppendTriangle(NewTri);
-			NewUVOverlay->SetTriangle(NewTriID, NewUVTri);
-			NewNormalOverlay->SetTriangle(NewTriID, NewNormalTri);
 		}
 	}
 
