Skin Weight Paint Tool: build the cleaned mesh if there are unused or non-manifold vertices

#rnx

diff --git a/Engine/Plugins/Runtime/MeshModelingToolset/Source/MeshModelingToolsEditorOnly/Private/SkeletalMesh/SkinWeightsPaintTool.cpp b/Engine/Plugins/Runtime/MeshModelingToolset/Source/MeshModelingToolsEditorOnly/Private/SkeletalMesh/SkinWeightsPaintTool.cpp
--- a/Engine/Plugins/Runtime/MeshModelingToolset/Source/MeshModelingToolsEditorOnly/Private/SkeletalMesh/SkinWeightsPaintTool.cpp
+++ b/Engine/Plugins/Runtime/MeshModelingToolset/Source/MeshModelingToolsEditorOnly/Private/SkeletalMesh/SkinWeightsPaintTool.cpp
@@ -4755,9 +4755,10 @@ FCleanedEditMesh::FCleanedEditMesh(const FDynamicMesh3& InDynamicMesh, const FMe
 	OriginalDynamicMesh = InDynamicMesh;
 	OriginalMeshDescription = InMeshDescription;
 	
-	// if the mesh doesn't have any unused vertices, then we can skip making a cleaned duplicate and just use it directly
+	// if the mesh doesn't have any unused or non-manifold vertices, then we can skip making a cleaned duplicate and just use it directly
 	// (the getters will return the original mesh)
-	if (!OriginalDynamicMesh.HasUnusedVertices())
+	const FNonManifoldMappingSupport OriginalNonManifoldMapping(OriginalDynamicMesh);
+	if (!OriginalDynamicMesh.HasUnusedVertices() && !OriginalNonManifoldMapping.IsNonManifoldVertexInSource())
 	{
 		return;
 	}
@@ -4773,7 +4774,6 @@ FCleanedEditMesh::FCleanedEditMesh(const FDynamicMesh3& InDynamicMesh, const FMe
 	CleanedSubMesh = FDynamicSubmesh3(&OriginalDynamicMesh, AllTris);
 
 	// remap the non-manifold mapping from full mesh to partial mesh
-	FNonManifoldMappingSupport OriginalNonManifoldMapping(OriginalDynamicMesh);
 	TArray<int32> SubMeshVertexToNonManifoldVertexIDMap;
 	FDynamicMesh3& SubMeshDynamicMesh = CleanedSubMesh.GetSubmesh();
 	for (int32 SubMeshVertexID=0; SubMeshVertexID<SubMeshDynamicMesh.VertexCount(); ++SubMeshVertexID)
