Fix for crash in subdivide tool. We were passing in wrong vertex count into OpenSubdiv.

diff --git a/Engine/Plugins/Runtime/MeshModelingToolset/Source/ModelingComponentsEditorOnly/Private/Operations/SubdividePoly.cpp b/Engine/Plugins/Runtime/MeshModelingToolset/Source/ModelingComponentsEditorOnly/Private/Operations/SubdividePoly.cpp
--- a/Engine/Plugins/Runtime/MeshModelingToolset/Source/ModelingComponentsEditorOnly/Private/Operations/SubdividePoly.cpp
+++ b/Engine/Plugins/Runtime/MeshModelingToolset/Source/ModelingComponentsEditorOnly/Private/Operations/SubdividePoly.cpp
@@ -558,8 +558,6 @@ bool FSubdividePoly::ComputeTopologySubdivision()
 	auto DescriptorFromGroupTopology = [this, &NumVertsPerFace, &BoundaryVertsPerFace, &UVChannel, &UVIndexBuffer]
 	(OpenSubdiv::Far::TopologyDescriptor& Descriptor)
 	{
-		int TotalNumVertices = 0;
-
 		for (const FGroupTopology::FGroup& Group : GroupTopology.Groups)
 		{
 			if (Group.Boundaries.Num() != 1)
@@ -577,11 +575,10 @@ bool FSubdividePoly::ComputeTopologySubdivision()
 			SubdividePolyLocal::GetBoundaryCorners(Boundary, GroupTopology, Corners);
 
 			NumVertsPerFace.Add(Corners.Num());
-			TotalNumVertices += Corners.Num();
 			BoundaryVertsPerFace.Append(Corners);
 		}
 
-		Descriptor.numVertices = TotalNumVertices;
+		Descriptor.numVertices = GroupTopology.Corners.Num();
 		Descriptor.numFaces = GroupTopology.Groups.Num();
 		Descriptor.numVertsPerFace = NumVertsPerFace.GetData();
 		Descriptor.vertIndicesPerFace = BoundaryVertsPerFace.GetData();
@@ -717,6 +714,8 @@ bool FSubdividePoly::ComputeSubdividedMesh(FDynamicMesh3& OutMesh)
 		GetGroupPolyMeshVertices(OriginalMesh, GroupTopology, bInterpolateVertexNormals, SourceVertices);
 	}
 
+	ensure(SourceVertices.Num() == Refiner->TopologyRefiner->GetLevel(0).GetNumVertices());
+
 	TArray<SubdividePolyLocal::SubdVertex> RefinedVertices;
 	for (int CurrentLevel = 1; CurrentLevel <= Level; ++CurrentLevel)
 	{
