Dataflow : add support for B + vertical drag t oadjust the attribute value in the CLoth and dataflow weightmap paint tools

diff --git a/Engine/Plugins/ChaosClothAssetEditor/Source/ChaosClothAssetEditorTools/Private/ChaosClothAsset/ClothWeightMapPaintTool.cpp b/Engine/Plugins/ChaosClothAssetEditor/Source/ChaosClothAssetEditorTools/Private/ChaosClothAsset/ClothWeightMapPaintTool.cpp
--- a/Engine/Plugins/ChaosClothAssetEditor/Source/ChaosClothAssetEditorTools/Private/ChaosClothAsset/ClothWeightMapPaintTool.cpp
+++ b/Engine/Plugins/ChaosClothAssetEditor/Source/ChaosClothAssetEditorTools/Private/ChaosClothAsset/ClothWeightMapPaintTool.cpp
@@ -21,6 +21,7 @@
 #include "Polygon2.h"
 #include "Intersection/IntrLine2Line2.h"
 #include "Intersection/IntrSegment2Segment2.h"
+#include "BaseBehaviors/TwoAxisPropertyEditBehavior.h"
 
 #include "Changes/MeshVertexChange.h"
 #include "Changes/BasicChanges.h"
@@ -107,6 +108,31 @@ void UClothEditorMeshWeightMapPaintToolShowHideProperties::PostAction(EClothEdit
  * Tool
  */
 
+void UClothEditorWeightMapPaintTool::SetupBrushEditBehaviorSetup(ULocalTwoAxisPropertyEditInputBehavior& OutBehavior)
+{
+	// map the horizontal behavior to change the brush size
+	UMeshSculptToolBase::MapHorizontalBrushEditBehaviorToBrushSize(OutBehavior);
+
+	// map the vertical ehavior to change the attribiute value
+	OutBehavior.VerticalProperty.GetValueFunc = [this]()
+		{
+			return FilterProperties->AttributeValue;
+		};
+
+	OutBehavior.VerticalProperty.SetValueFunc = [this](float NewValue)
+		{
+			FilterProperties->AttributeValue = FMath::Min(1.0f, FMath::Max(NewValue, 0.f));
+#if WITH_EDITOR
+			FPropertyChangedEvent PropertyChangedEvent(UBrushBaseProperties::StaticClass()->FindPropertyByName(GET_MEMBER_NAME_CHECKED(UClothEditorWeightMapPaintBrushFilterProperties, AttributeValue)));
+			FilterProperties->PostEditChangeProperty(PropertyChangedEvent);
+#endif
+		};
+	OutBehavior.VerticalProperty.Name = LOCTEXT("AttributeValue", "Attribute Value");
+	OutBehavior.VerticalProperty.EditRate = 0.005f;
+	OutBehavior.VerticalProperty.bEnabled = true;
+}
+
+
 void UClothEditorWeightMapPaintTool::Setup()
 {
 	UMeshSculptToolBase::Setup();
diff --git a/Engine/Plugins/ChaosClothAssetEditor/Source/ChaosClothAssetEditorTools/Private/ChaosClothAsset/ClothWeightMapPaintTool.h b/Engine/Plugins/ChaosClothAssetEditor/Source/ChaosClothAssetEditorTools/Private/ChaosClothAsset/ClothWeightMapPaintTool.h
--- a/Engine/Plugins/ChaosClothAssetEditor/Source/ChaosClothAssetEditorTools/Private/ChaosClothAsset/ClothWeightMapPaintTool.h
+++ b/Engine/Plugins/ChaosClothAssetEditor/Source/ChaosClothAssetEditorTools/Private/ChaosClothAsset/ClothWeightMapPaintTool.h
@@ -349,6 +349,8 @@ public:
 
 	UE_API void SetDataflowContextObject(TObjectPtr<UDataflowContextObject> InDataflowContextObject);
 
+	virtual void SetupBrushEditBehaviorSetup(ULocalTwoAxisPropertyEditInputBehavior& OutBehavior) override;
+
 private:
 
 	// Initialization support
diff --git a/Engine/Plugins/Experimental/Dataflow/Source/DataflowEditor/Private/DataflowEditorTools/DataflowEditorWeightMapPaintTool.cpp b/Engine/Plugins/Experimental/Dataflow/Source/DataflowEditor/Private/DataflowEditorTools/DataflowEditorWeightMapPaintTool.cpp
--- a/Engine/Plugins/Experimental/Dataflow/Source/DataflowEditor/Private/DataflowEditorTools/DataflowEditorWeightMapPaintTool.cpp
+++ b/Engine/Plugins/Experimental/Dataflow/Source/DataflowEditor/Private/DataflowEditorTools/DataflowEditorWeightMapPaintTool.cpp
@@ -4,6 +4,7 @@
 
 #include "Async/ParallelFor.h"
 #include "Async/Async.h"
+#include "BaseBehaviors/TwoAxisPropertyEditBehavior.h"
 #include "BaseGizmos/BrushStampIndicator.h"
 #include "CanvasTypes.h"
 #include "CanvasItem.h"
@@ -162,6 +163,30 @@ void UDataflowEditorUpdateWeightMapProperties::PostEditChangeProperty(FPropertyC
  * Tool
  */
 
+void UDataflowEditorWeightMapPaintTool::SetupBrushEditBehaviorSetup(ULocalTwoAxisPropertyEditInputBehavior& OutBehavior) 
+{
+	// map the horizontal behavior to change the brush size
+	UMeshSculptToolBase::MapHorizontalBrushEditBehaviorToBrushSize(OutBehavior);
+
+	// map the vertical ehavior to change the attribiute value
+	OutBehavior.VerticalProperty.GetValueFunc = [this]()
+		{
+			return FilterProperties->AttributeValue;
+		};
+
+	OutBehavior.VerticalProperty.SetValueFunc = [this](float NewValue)
+		{
+			FilterProperties->AttributeValue = FMath::Min(1.0f, FMath::Max(NewValue, 0.f));
+#if WITH_EDITOR
+			FPropertyChangedEvent PropertyChangedEvent(UBrushBaseProperties::StaticClass()->FindPropertyByName(GET_MEMBER_NAME_CHECKED(UDataflowEditorWeightMapPaintBrushFilterProperties, AttributeValue)));
+			FilterProperties->PostEditChangeProperty(PropertyChangedEvent);
+#endif
+		};
+	OutBehavior.VerticalProperty.Name = LOCTEXT("AttributeValue", "Attribute Value");
+	OutBehavior.VerticalProperty.EditRate = 0.005f;
+	OutBehavior.VerticalProperty.bEnabled = true;
+}
+
 void UDataflowEditorWeightMapPaintTool::Setup()
 {
 	UMeshSculptToolBase::Setup();
diff --git a/Engine/Plugins/Experimental/Dataflow/Source/DataflowEditor/Public/DataflowEditorTools/DataflowEditorWeightMapPaintTool.h b/Engine/Plugins/Experimental/Dataflow/Source/DataflowEditor/Public/DataflowEditorTools/DataflowEditorWeightMapPaintTool.h
--- a/Engine/Plugins/Experimental/Dataflow/Source/DataflowEditor/Public/DataflowEditorTools/DataflowEditorWeightMapPaintTool.h
+++ b/Engine/Plugins/Experimental/Dataflow/Source/DataflowEditor/Public/DataflowEditorTools/DataflowEditorWeightMapPaintTool.h
@@ -296,6 +296,8 @@ public:
 	UE_API void SetDataflowEditorContextObject(TObjectPtr<UDataflowContextObject> InDataflowEditorContextObject);
 	void SetEditorMode(UDataflowEditorMode* InMode) { Mode = InMode; }
 
+	virtual void SetupBrushEditBehaviorSetup(ULocalTwoAxisPropertyEditInputBehavior& OutBehavior) override;
+
 	//~ UObject interface
 
 	static UE_API void AddReferencedObjects(UObject* InThis, FReferenceCollector& Collector);
diff --git a/Engine/Plugins/Runtime/MeshModelingToolset/Source/MeshModelingTools/Private/Sculpting/MeshSculptToolBase.cpp b/Engine/Plugins/Runtime/MeshModelingToolset/Source/MeshModelingTools/Private/Sculpting/MeshSculptToolBase.cpp
--- a/Engine/Plugins/Runtime/MeshModelingToolset/Source/MeshModelingTools/Private/Sculpting/MeshSculptToolBase.cpp
+++ b/Engine/Plugins/Runtime/MeshModelingToolset/Source/MeshModelingTools/Private/Sculpting/MeshSculptToolBase.cpp
@@ -103,7 +103,6 @@ void UMeshSculptToolBase::SetWorld(UWorld* World)
 	this->TargetWorld = World;
 }
 
-
 void UMeshSculptToolBase::Setup()
 {
 	UMeshSurfacePointTool::Setup();
@@ -169,40 +168,7 @@ void UMeshSculptToolBase::Setup()
     {
         BrushEditBehavior = NewObject<ULocalTwoAxisPropertyEditInputBehavior>(this);
 		
-        BrushEditBehavior->HorizontalProperty.GetValueFunc = [this]()
-        {
-        	if (BrushProperties->BrushSize.SizeType == EBrushToolSizeType::Adaptive)
-        	{
-        		return BrushProperties->BrushSize.AdaptiveSize;
-        	}
-            else
-            {
-                return BrushProperties->BrushSize.WorldRadius;
-            }
-        };
-        BrushEditBehavior->HorizontalProperty.SetValueFunc = [this](float NewValue)
-        {
-        	if (BrushProperties->BrushSize.SizeType == EBrushToolSizeType::Adaptive)
-        	{
-        		BrushProperties->BrushSize.AdaptiveSize = FMath::Clamp(NewValue, 0.f, 1.f);
-        	}
-        	else
-        	{
-        		BrushProperties->BrushSize.WorldRadius = FMath::Max(NewValue, 0.01f);
-        	}
-        	
-#if WITH_EDITOR
-        	FPropertyChangedEvent PropertyChangedEvent(USculptBrushProperties::StaticClass()->FindPropertyByName(GET_MEMBER_NAME_CHECKED(USculptBrushProperties, BrushSize)));
-        	BrushProperties->PostEditChangeProperty(PropertyChangedEvent);
-#endif
-        };
-		BrushEditBehavior->HorizontalProperty.MutateDeltaFunc = [this](float Delta)
-		{
-			// Scale delta if brush size is in world units.
-			return Delta * (BrushProperties->BrushSize.SizeType == EBrushToolSizeType::World ? (CameraState.Position - LastBrushFrameWorld.Origin).Length() : 1.f);
-		};
-		BrushEditBehavior->HorizontalProperty.Name = LOCTEXT("BrushRadius", "Radius");
-		BrushEditBehavior->HorizontalProperty.bEnabled = true;
+		SetupBrushEditBehaviorSetup(*BrushEditBehavior);
     	
         BrushEditBehavior->OnDragUpdated.AddWeakLambda(this, [this]()
         {
@@ -215,6 +181,48 @@ void UMeshSculptToolBase::Setup()
     }
 }
 
+void UMeshSculptToolBase::SetupBrushEditBehaviorSetup(ULocalTwoAxisPropertyEditInputBehavior& OutBehavior)
+{
+	MapHorizontalBrushEditBehaviorToBrushSize(OutBehavior);
+}
+
+void UMeshSculptToolBase::MapHorizontalBrushEditBehaviorToBrushSize(ULocalTwoAxisPropertyEditInputBehavior& OutBehavior)
+{
+	OutBehavior.HorizontalProperty.GetValueFunc = [this]()
+		{
+			if (BrushProperties->BrushSize.SizeType == EBrushToolSizeType::Adaptive)
+			{
+				return BrushProperties->BrushSize.AdaptiveSize;
+			}
+			else
+			{
+				return BrushProperties->BrushSize.WorldRadius;
+			}
+		};
+	OutBehavior.HorizontalProperty.SetValueFunc = [this](float NewValue)
+		{
+			if (BrushProperties->BrushSize.SizeType == EBrushToolSizeType::Adaptive)
+			{
+				BrushProperties->BrushSize.AdaptiveSize = FMath::Clamp(NewValue, 0.f, 1.f);
+			}
+			else
+			{
+				BrushProperties->BrushSize.WorldRadius = FMath::Max(NewValue, 0.01f);
+			}
+
+#if WITH_EDITOR
+			FPropertyChangedEvent PropertyChangedEvent(USculptBrushProperties::StaticClass()->FindPropertyByName(GET_MEMBER_NAME_CHECKED(USculptBrushProperties, BrushSize)));
+			BrushProperties->PostEditChangeProperty(PropertyChangedEvent);
+#endif
+		};
+	OutBehavior.HorizontalProperty.MutateDeltaFunc = [this](float Delta)
+		{
+			// Scale delta if brush size is in world units.
+			return Delta * (BrushProperties->BrushSize.SizeType == EBrushToolSizeType::World ? (CameraState.Position - LastBrushFrameWorld.Origin).Length() : 1.f);
+		};
+	OutBehavior.HorizontalProperty.Name = LOCTEXT("BrushRadius", "Radius");
+	OutBehavior.HorizontalProperty.bEnabled = true;
+}
 
 void UMeshSculptToolBase::OnCompleteSetup()
 {
diff --git a/Engine/Plugins/Runtime/MeshModelingToolset/Source/MeshModelingTools/Public/Sculpting/MeshSculptToolBase.h b/Engine/Plugins/Runtime/MeshModelingToolset/Source/MeshModelingTools/Public/Sculpting/MeshSculptToolBase.h
--- a/Engine/Plugins/Runtime/MeshModelingToolset/Source/MeshModelingTools/Public/Sculpting/MeshSculptToolBase.h
+++ b/Engine/Plugins/Runtime/MeshModelingToolset/Source/MeshModelingTools/Public/Sculpting/MeshSculptToolBase.h
@@ -635,6 +635,10 @@ protected:
 	UE_API virtual void SetIndicatorVisibility(bool bVisible);
 	UE_API virtual bool GetIndicatorVisibility() const;
 
+	UE_API virtual void SetupBrushEditBehaviorSetup(ULocalTwoAxisPropertyEditInputBehavior& OutBehavior);
+
+	UE_API void MapHorizontalBrushEditBehaviorToBrushSize(ULocalTwoAxisPropertyEditInputBehavior& OutBehavior);
+
 protected:
 	UPROPERTY()
 	TObjectPtr<UBrushStampIndicator> BrushIndicator;
