ModelingTools: Prevent vertext sculpt tool from crashing if wrong brush is requested.

diff --git a/Engine/Plugins/Runtime/MeshModelingToolset/Source/MeshModelingTools/Private/MeshVertexSculptTool.cpp b/Engine/Plugins/Runtime/MeshModelingToolset/Source/MeshModelingTools/Private/MeshVertexSculptTool.cpp
--- a/Engine/Plugins/Runtime/MeshModelingToolset/Source/MeshModelingTools/Private/MeshVertexSculptTool.cpp
+++ b/Engine/Plugins/Runtime/MeshModelingToolset/Source/MeshModelingTools/Private/MeshVertexSculptTool.cpp
@@ -283,7 +283,7 @@ void UMeshVertexSculptTool::Setup()
 
 	RegisterBrushes();
 
-	if (DefaultPrimaryBrushID >= 0 && ensure(SculptProperties->PrimaryBrushID < RegisteredPrimaryBrushTypes.Num()))
+	if (DefaultPrimaryBrushID >= 0 && ensure(BrushOpFactories.Contains(DefaultPrimaryBrushID)))
 	{
 		SculptProperties->PrimaryBrushID = DefaultPrimaryBrushID;
 	}
@@ -1919,6 +1919,15 @@ void UMeshVertexSculptTool::UpdateBrushType(int32 BrushID)
 
 	SetActivePrimaryBrushType(BrushID);
 
+	// If something went wrong and we were unable to activate the given brush, make sure that we have some kind
+	//  of brush active so we don't crash.
+	if (!PrimaryBrushOp && ensure(!BrushOpFactories.IsEmpty()))
+	{
+		UE_LOG(LogGeometry, Error, TEXT("Sculpt tool was unable to activate chosen brush."));
+		SetActivePrimaryBrushType(BrushOpFactories.CreateIterator()->Key);
+		ensure(PrimaryBrushOp);
+	}
+
 	if (BrushEditBehavior.IsValid())
 	{
 		// todo: Handle Kelvinlet brush props better. At the moment we are just disabling strength editing for Kelvinlet brush ops.
