ModelingTools: Fix vertex sculpt tool not applying brush on first click in spacing mode.

diff --git a/Engine/Plugins/Runtime/MeshModelingToolset/Source/MeshModelingTools/Private/Sculpting/MeshSculptToolBase.cpp b/Engine/Plugins/Runtime/MeshModelingToolset/Source/MeshModelingTools/Private/Sculpting/MeshSculptToolBase.cpp
--- a/Engine/Plugins/Runtime/MeshModelingToolset/Source/MeshModelingTools/Private/Sculpting/MeshSculptToolBase.cpp
+++ b/Engine/Plugins/Runtime/MeshModelingToolset/Source/MeshModelingTools/Private/Sculpting/MeshSculptToolBase.cpp
@@ -793,6 +793,7 @@ void UMeshSculptToolBase::OnBeginDrag(const FRay& WorldRay)
 		PrimaryBrushOp->Falloff = PrimaryFalloff;
 
 		OnBeginStroke(WorldRay);
+		OnUpdateDrag(WorldRay);
 	}
 }
 
@@ -1153,6 +1154,10 @@ void UMeshSculptToolBase::UpdateStampPendingState()
 		TGuardValue<double> CachedActivePressure(ActivePressure, ActivePressure);
 
 		TArray<FTimedStrokePoint> UnProjectedStampPoints;
+		if (!StrokePath.IsEmpty())
+		{
+			UnProjectedStampPoints.Add(StrokePath[0]);
+		}
 		if (StrokePath.Num() > 1)
 		{
 			TArray<FTimedStrokePoint> UnconsumedPathPoints;
@@ -1160,7 +1165,6 @@ void UMeshSculptToolBase::UpdateStampPendingState()
 			int32 SectionIndex = 0;
 			double SpacingSectionConsumedLength = 0.0;
 
-			UnProjectedStampPoints.Add(StrokePath[0]);
 			UnconsumedPathPoints.Add(StrokePath[0]);
 
 			for (int32 Index = 0; Index < StrokePath.Num() - 1; Index++)
