ModelingTools: Expose vertex sculpt tool function for finding whether target has sculpt layers, to control switching to layer-specific brushes externally.

diff --git a/Engine/Plugins/Runtime/MeshModelingToolset/Source/MeshModelingTools/Private/MeshVertexSculptTool.cpp b/Engine/Plugins/Runtime/MeshModelingToolset/Source/MeshModelingTools/Private/MeshVertexSculptTool.cpp
--- a/Engine/Plugins/Runtime/MeshModelingToolset/Source/MeshModelingTools/Private/MeshVertexSculptTool.cpp
+++ b/Engine/Plugins/Runtime/MeshModelingToolset/Source/MeshModelingTools/Private/MeshVertexSculptTool.cpp
@@ -405,17 +405,11 @@ void UMeshVertexSculptTool::RegisterBrushes()
 		MakeUnique<TBasicMeshSculptBrushOpFactory<FTwistKelvinletBrushOp>>(),
 		NewObject<UTwistKelvinletBrushOpProps>(this));
 
-	if (ISceneComponentBackedTarget* SceneComponentTarget = Cast<ISceneComponentBackedTarget>(Target))
+	if (DoesTargetHaveSculptLayers())
 	{
-		if (IMeshSculptLayersManager* SculptLayersManager = Cast<IMeshSculptLayersManager>(SceneComponentTarget->GetOwnerSceneComponent()))
-		{
-			if (SculptLayersManager->HasSculptLayers())
-			{
-				RegisterBrushType((int32)EMeshVertexSculptBrushType::EraseSculptLayer, LOCTEXT("EraseSculptLayer", "EraseSculptLayer"),
-					MakeUnique<TBasicMeshSculptBrushOpFactory<FEraseSculptLayerBrushOp>>(),
-					NewObject<UEraseSculptLayerBrushOpProps>(this));
-			}
-		}
+		RegisterBrushType((int32)EMeshVertexSculptBrushType::EraseSculptLayer, LOCTEXT("EraseSculptLayer", "EraseSculptLayer"),
+			MakeUnique<TBasicMeshSculptBrushOpFactory<FEraseSculptLayerBrushOp>>(),
+			NewObject<UEraseSculptLayerBrushOpProps>(this));
 	}
 
 	// secondary brushes
@@ -1974,6 +1968,17 @@ void UMeshVertexSculptTool::UpdateBrushType(int32 BrushID)
 	GetToolManager()->DisplayMessage(Builder.ToText(), EToolMessageLevel::UserNotification);
 }
 
+bool UMeshVertexSculptTool::DoesTargetHaveSculptLayers() const
+{
+	if (ISceneComponentBackedTarget* SceneComponentTarget = Cast<ISceneComponentBackedTarget>(Target))
+	{
+		if (IMeshSculptLayersManager* SculptLayersManager = Cast<IMeshSculptLayersManager>(SceneComponentTarget->GetOwnerSceneComponent()))
+		{
+			return SculptLayersManager->HasSculptLayers();
+		}
+	}
+	return false;
+}
 
 
 void UMeshVertexSculptTool::UndoRedo_RestoreSymmetryPossibleState(bool bSetToValue)
diff --git a/Engine/Plugins/Runtime/MeshModelingToolset/Source/MeshModelingTools/Public/MeshVertexSculptTool.h b/Engine/Plugins/Runtime/MeshModelingToolset/Source/MeshModelingTools/Public/MeshVertexSculptTool.h
--- a/Engine/Plugins/Runtime/MeshModelingToolset/Source/MeshModelingTools/Public/MeshVertexSculptTool.h
+++ b/Engine/Plugins/Runtime/MeshModelingToolset/Source/MeshModelingTools/Public/MeshVertexSculptTool.h
@@ -225,6 +225,8 @@ class UMeshVertexSculptTool : public UMeshSculptToolBase, public IInteractiveToo
 {
 	GENERATED_BODY()
 public:
+	UE_API bool DoesTargetHaveSculptLayers() const;
+
 	UE_API virtual void Setup() override;
 	UE_API virtual void Shutdown(EToolShutdownType ShutdownType) override;
 	UE_API virtual void Render(IToolsContextRenderAPI* RenderAPI) override;
