ModelingMode: Fix bug with the Vertex Sculpt where it doesn't update the base mesh after each stroke completes.

diff --git a/Engine/Plugins/Runtime/MeshModelingToolset/Source/MeshModelingTools/Private/MeshVertexSculptTool.cpp b/Engine/Plugins/Runtime/MeshModelingToolset/Source/MeshModelingTools/Private/MeshVertexSculptTool.cpp
--- a/Engine/Plugins/Runtime/MeshModelingToolset/Source/MeshModelingTools/Private/MeshVertexSculptTool.cpp
+++ b/Engine/Plugins/Runtime/MeshModelingToolset/Source/MeshModelingTools/Private/MeshVertexSculptTool.cpp
@@ -1553,7 +1553,7 @@ void UMeshVertexSculptTool::OnTick(float DeltaTime)
 			GetToolManager()->PostInvalidation();
 		}
 
-
+		AccumulatedTriangleROI.Append(AccumulatedTriangleROIArray);
 	};
 	
 	ProcessPerTickStamps(
