make sculpt tool skip render decomp step if the sculpt mesh has no triangles (fixing ensure when tool is run on empty mesh)

diff --git a/Engine/Plugins/Runtime/MeshModelingToolset/Source/MeshModelingTools/Private/MeshVertexSculptTool.cpp b/Engine/Plugins/Runtime/MeshModelingToolset/Source/MeshModelingTools/Private/MeshVertexSculptTool.cpp
--- a/Engine/Plugins/Runtime/MeshModelingToolset/Source/MeshModelingTools/Private/MeshVertexSculptTool.cpp
+++ b/Engine/Plugins/Runtime/MeshModelingToolset/Source/MeshModelingTools/Private/MeshVertexSculptTool.cpp
@@ -214,6 +214,10 @@ void UMeshVertexSculptTool::Setup()
 	// initialize render decomposition
 	TFuture<void> InitializeRenderDecomp = Async(VertexSculptToolAsyncExecTarget, [SculptMesh, &MaterialSet, this]()
 	{
+		if (SculptMesh->TriangleCount() == 0)
+		{
+			return;
+		}
 		TUniquePtr<FMeshRenderDecomposition> Decomp = MakeUnique<FMeshRenderDecomposition>();
 		FMeshRenderDecomposition::BuildChunkedDecomposition(SculptMesh, &MaterialSet, *Decomp);
 		Decomp->BuildAssociations(SculptMesh);
