Fix missing include
#rnx

diff --git a/Engine/Plugins/Runtime/MeshModelingToolset/Source/ModelingOperators/Private/CompositionOps/CurveSweepOp.cpp b/Engine/Plugins/Runtime/MeshModelingToolset/Source/ModelingOperators/Private/CompositionOps/CurveSweepOp.cpp
--- a/Engine/Plugins/Runtime/MeshModelingToolset/Source/ModelingOperators/Private/CompositionOps/CurveSweepOp.cpp
+++ b/Engine/Plugins/Runtime/MeshModelingToolset/Source/ModelingOperators/Private/CompositionOps/CurveSweepOp.cpp
@@ -5,6 +5,7 @@
 #include "ConstrainedDelaunay2.h"
 #include "DynamicMeshEditor.h"
 #include "MeshBoundaryLoops.h"
+#include "CompGeom/PolygonTriangulation.h"
 #include "DynamicMesh/MeshNormals.h"
 #include "DynamicMesh/MeshTransforms.h"
 #include "Operations/SimpleHoleFiller.h"
