Modeling Tools : Move the mesh selector class to a common place so it can be shared by other tools

diff --git a/Engine/Plugins/Runtime/MeshModelingToolset/Source/MeshModelingToolsEditorOnly/Private/SkeletalMesh/SkinWeightsPaintTool.cpp b/Engine/Plugins/Runtime/MeshModelingToolset/Source/MeshModelingToolsEditorOnly/Private/SkeletalMesh/SkinWeightsPaintTool.cpp
--- a/Engine/Plugins/Runtime/MeshModelingToolset/Source/MeshModelingToolsEditorOnly/Private/SkeletalMesh/SkinWeightsPaintTool.cpp
+++ b/Engine/Plugins/Runtime/MeshModelingToolset/Source/MeshModelingToolsEditorOnly/Private/SkeletalMesh/SkinWeightsPaintTool.cpp
@@ -383,333 +383,6 @@ TArray<FName> USkinWeightsPaintToolProperties::GetSourceSkinWeightProfilesFunc()
 	return SkinWeightLayer::GetProfilesFromToolTarget(SourceToolTarget, SourceLOD);
 }
 
-void UWeightToolMeshSelector::InitialSetup(
-	UWorld* InWorld,
-	UInteractiveTool* InParentTool,
-	FEditorViewportClient* InViewportClient,
-	TFunction<void()> OnSelectionChangedFunc)
-{
-	World = InWorld;
-
-	ViewportClient = InViewportClient;
-	
-	// set up vertex selection mechanic
-	PolygonSelectionMechanic = NewObject<UPolygonSelectionMechanic>(this);
-	PolygonSelectionMechanic->bAddSelectionFilterPropertiesToParentTool = false;
-	PolygonSelectionMechanic->Setup(InParentTool);
-	PolygonSelectionMechanic->SetIsEnabled(false, false);
-	PolygonSelectionMechanic->OnSelectionChanged.AddLambda(OnSelectionChangedFunc );
-
-	// set up style of vertex selection
-	constexpr FLinearColor VertexSelectedPurple = FLinearColor(0.78f, 0.f, 0.78f);
-	constexpr FLinearColor VertexSelectedYellow = FLinearColor(1.f,1.f,0.f);
-	// adjust selection rendering for this context
-	PolygonSelectionMechanic->HilightRenderer.PointColor = FLinearColor::Blue;
-	PolygonSelectionMechanic->HilightRenderer.PointSize = 10.0f;
-	// vertex highlighting once selected
-	PolygonSelectionMechanic->SelectionRenderer.LineThickness = 1.0f;
-	PolygonSelectionMechanic->SelectionRenderer.PointColor = VertexSelectedYellow;
-	PolygonSelectionMechanic->SelectionRenderer.PointSize = 5.0f;
-	PolygonSelectionMechanic->SelectionRenderer.DepthBias = 2.0f;
-	// despite the name, this renders the vertices
-	PolygonSelectionMechanic->PolyEdgesRenderer.PointColor = VertexSelectedPurple;
-	PolygonSelectionMechanic->PolyEdgesRenderer.PointSize = 5.0f;
-	PolygonSelectionMechanic->PolyEdgesRenderer.DepthBias = 2.0f;
-	PolygonSelectionMechanic->PolyEdgesRenderer.LineThickness = 1.0f;
-}
-
-void UWeightToolMeshSelector::SetMesh(
-	UPreviewMesh* InPreviewMesh,
-	const FMeshDescription* InMeshDescription,
-	const FTransform3d& InMeshTransform)
-{
-	// store the mesh this is operating on
-	PreviewMesh = InPreviewMesh;
-	// copy the mesh description for selection validation
-	MeshDescription = *InMeshDescription;
-
-	if (!ensure(World))
-	{
-		return;
-	}
-
-	if (!PreviewMesh)
-	{
-		SetIsEnabled(false);
-		return;
-	}
-
-	// reset selection topology and mesh spatial data
-	static constexpr bool bAutoBuild = true;
-	const FDynamicMesh3* DynamicMesh = PreviewMesh->GetMesh();
-	SelectionTopology = MakeUnique<UE::Geometry::FTriangleGroupTopology>(DynamicMesh, bAutoBuild);
-	MeshSpatial = MakeUnique<FDynamicMeshAABBTree3>(DynamicMesh, bAutoBuild);
-
-	// initialize the selection mechanic
-	PolygonSelectionMechanic->Initialize(
-		DynamicMesh,
-		InMeshTransform,
-		World,
-		SelectionTopology.Get(),
-		[this]() { return MeshSpatial.Get(); }
-	);
-
-	// clear the selection (old selection is invalid on new topo)
-	PolygonSelectionMechanic->ClearSelection();
-	PolygonSelectionMechanic->ClearHighlight();
-
-	// selection colors
-	constexpr FLinearColor FaceSelectedOrange = FLinearColor(0.886f, 0.672f, 0.473f);
-	// configure secondary render material for selected triangles
-	// NOTE: the material returned by ToolSetupUtil::GetSelectionMaterial has a checkerboard pattern on back faces which makes it hard to use
-	if (UMaterialInterface* Material = LoadObject<UMaterial>(nullptr, TEXT("/MeshModelingToolsetExp/Materials/SculptMaterial")))
-	{
-		if (UMaterialInstanceDynamic* MatInstance = UMaterialInstanceDynamic::Create(Material, World))
-		{
-			MatInstance->SetVectorParameterValue(TEXT("Color"), FaceSelectedOrange);
-			PreviewMesh->SetSecondaryRenderMaterial(MatInstance);
-		}
-	}
-
-	// secondary triangle buffer used to render face selection
-	PreviewMesh->EnableSecondaryTriangleBuffers([this](const FDynamicMesh3* Mesh, int32 TriangleID)
-	{
-		return PolygonSelectionMechanic->GetActiveSelection().IsSelectedTriangle(Mesh, SelectionTopology.Get(), TriangleID);
-	});
-	// notify preview mesh when triangle selection has been updated
-	PolygonSelectionMechanic->OnSelectionChanged.AddWeakLambda(this, [this]()
-	{
-		PreviewMesh->FastNotifySecondaryTrianglesChanged();
-	});
-	PolygonSelectionMechanic->OnFaceSelectionPreviewChanged.AddWeakLambda(this, [this]()
-	{
-		PreviewMesh->FastNotifySecondaryTrianglesChanged();
-	});
-}
-
-void UWeightToolMeshSelector::Shutdown()
-{
-	if (PolygonSelectionMechanic)
-	{
-		PolygonSelectionMechanic->Shutdown();
-	}
-	
-	PolygonSelectionMechanic = nullptr;
-}
-
-void UWeightToolMeshSelector::SetIsEnabled(bool bIsEnabled)
-{
-	if (!PolygonSelectionMechanic)
-	{
-		return;
-	}
-
-	// force off if there's no preview mesh
-	bIsEnabled = PreviewMesh ? bIsEnabled : false; 
-	
-	PolygonSelectionMechanic->SetIsEnabled(bIsEnabled, bIsEnabled);
-}
-
-void UWeightToolMeshSelector::SetComponentSelectionMode(EComponentSelectionMode InMode)
-{
-	if (!(PolygonSelectionMechanic))
-	{
-		return;
-	}
-	
-	PolygonSelectionMechanic->Properties->bSelectVertices = InMode == EComponentSelectionMode::Vertices;
-	PolygonSelectionMechanic->Properties->bSelectEdges = InMode == EComponentSelectionMode::Edges;
-	PolygonSelectionMechanic->Properties->bSelectFaces = InMode == EComponentSelectionMode::Faces;
-	PolygonSelectionMechanic->SetShowSelectableCorners(InMode == EComponentSelectionMode::Vertices);
-	PolygonSelectionMechanic->SetShowEdges(InMode == EComponentSelectionMode::Edges);
-}
-
-void UWeightToolMeshSelector::SetTransform(const FTransform3d& InTargetTransform)
-{
-	if (PolygonSelectionMechanic)
-	{
-		PolygonSelectionMechanic->SetTransform(InTargetTransform);
-	}
-}
-
-void UWeightToolMeshSelector::UpdateAfterMeshDeformation()
-{
-	MeshSpatial->Build();
-	constexpr bool bTopologyDeformed = true;
-	constexpr bool bTopologyModified = false;
-	PolygonSelectionMechanic->GetTopologySelector()->Invalidate(bTopologyDeformed, bTopologyModified);
-}
-
-void UWeightToolMeshSelector::DrawHUD(FCanvas* Canvas, IToolsContextRenderAPI* RenderAPI)
-{
-	if (!PolygonSelectionMechanic)
-	{
-		return;
-	}
-
-	PolygonSelectionMechanic->DrawHUD(Canvas, RenderAPI);
-}
-	
-void UWeightToolMeshSelector::Render(IToolsContextRenderAPI* RenderAPI)
-{
-	if (!(PolygonSelectionMechanic && PreviewMesh))
-	{
-		return;
-	}
-	
-	PolygonSelectionMechanic->Render(RenderAPI);
-}
-
-const TArray<int32>& UWeightToolMeshSelector::GetSelectedVertices()
-{
-	using UE::Geometry::FNonManifoldMappingSupport;
-	
-	SelectedVerticesInternal.Empty();
-	if (!(PolygonSelectionMechanic && PreviewMesh))
-	{
-		return SelectedVerticesInternal;
-	}
-	
-	const FGroupTopologySelection& Selection = PolygonSelectionMechanic->GetActiveSelection();
-	const FDynamicMesh3* DynamicMesh = PreviewMesh->GetMesh();
-	const FVertexArray& AllVertices = MeshDescription.Vertices();
-	const FNonManifoldMappingSupport NonManifoldMappingSupport(*DynamicMesh);
-
-	// validate and add vertices to the output array
-	auto AddVertices = [this, &AllVertices, &NonManifoldMappingSupport](const TSet<int32>& VerticesToAdd)
-	{
-		// we have to make sure that the vertex ids are safe to use as PolygonSelectionMechanic does not act on the
-		// mesh description but on the dynamic mesh that can duplicate vertices when dealing with degenerate triangles.
-		// cf. FMeshDescriptionToDynamicMesh::Convert for more details.
-		for (const int32 VertexToAdd : VerticesToAdd)
-		{
-			const int32 SrcVertexID = NonManifoldMappingSupport.GetOriginalNonManifoldVertexID(VertexToAdd);
-			if (AllVertices.IsValid(SrcVertexID))
-			{
-				SelectedVerticesInternal.Add(SrcVertexID);
-			}
-		}
-	};
-
-	// add selected vertices
-	AddVertices(Selection.SelectedCornerIDs);
-
-	// add vertices on selected edges
-	{
-		TSet<int32> VerticesInSelectedEdges;
-		for (const int32 SelectedEdgeIndex : Selection.SelectedEdgeIDs)
-		{
-			FDynamicMesh3::FEdge CurrentEdge = DynamicMesh->GetEdge(SelectedEdgeIndex);
-			VerticesInSelectedEdges.Add(CurrentEdge.Vert.A);
-			VerticesInSelectedEdges.Add(CurrentEdge.Vert.B);
-		}
-		
-		AddVertices(VerticesInSelectedEdges);
-	}
-
-	// add vertices in selected faces
-	{
-		TSet<int32> VerticesInSelectedFaces;
-		for (const int32 SelectedFaceIndex : Selection.SelectedGroupIDs)
-		{
-			UE::Geometry::FIndex3i TriangleVertices = DynamicMesh->GetTriangleRef(SelectedFaceIndex);
-			VerticesInSelectedFaces.Add(TriangleVertices[0]);
-			VerticesInSelectedFaces.Add(TriangleVertices[1]);
-			VerticesInSelectedFaces.Add(TriangleVertices[2]);
-		}
-		
-		AddVertices(VerticesInSelectedFaces);
-	}
-	
-	return SelectedVerticesInternal;
-}
-
-bool UWeightToolMeshSelector::IsAnyComponentSelected() const
-{
-	if (!PolygonSelectionMechanic)
-	{
-		return false;
-	}
-
-	return PolygonSelectionMechanic->HasSelection();
-}
-
-void UWeightToolMeshSelector::GetSelectedTriangles(TArray<int32>& OutTriangleIndices) const
-{
-	OutTriangleIndices.Empty();
-	if (!ensure(PolygonSelectionMechanic))
-	{
-		return;
-	}
-	
-	const FGroupTopologySelection& Selection = PolygonSelectionMechanic->GetActiveSelection();
-	const FDynamicMesh3* DynamicMesh = PreviewMesh->GetMesh();
-	TSet<int32> TriangleSet;
-
-	// add triangles connected to selected vertices
-	for (const int32 VertexIndex : Selection.SelectedCornerIDs)
-	{
-		DynamicMesh->EnumerateVertexTriangles(VertexIndex, [&TriangleSet](int32 TriangleIndex)
-		{
-			TriangleSet.Add(TriangleIndex);
-		});
-	}
-
-	// add triangles connected to selected edges
-	for (const int32 EdgeIndex : Selection.SelectedEdgeIDs)
-	{
-		DynamicMesh->EnumerateEdgeTriangles(EdgeIndex, [&TriangleSet](int32 TriangleIndex)
-		{
-			TriangleSet.Add(TriangleIndex);
-		});
-	}
-
-	// add selected triangles
-	TriangleSet.Append(Selection.SelectedGroupIDs);	
-
-	OutTriangleIndices = TriangleSet.Array();
-}
-
-void UWeightToolMeshSelector::GrowSelection() const
-{
-	if (!ensure(PolygonSelectionMechanic))
-	{
-		return;
-	}
-	
-	PolygonSelectionMechanic->GrowSelection(/*bAsTriangleTopology*/ true);
-}
-
-void UWeightToolMeshSelector::ShrinkSelection() const
-{
-	if (!ensure(PolygonSelectionMechanic))
-	{
-		return;
-	}
-	
-	PolygonSelectionMechanic->ShrinkSelection(/*bAsTriangleTopology*/ true);
-}
-
-void UWeightToolMeshSelector::FloodSelection() const
-{
-	if (!ensure(PolygonSelectionMechanic))
-	{
-		return;
-	}
-	
-	PolygonSelectionMechanic->FloodSelection();
-}
-
-void UWeightToolMeshSelector::SelectBorder() const
-{
-	if (!ensure(PolygonSelectionMechanic))
-	{
-		return;
-	}
-	
-	PolygonSelectionMechanic->ConvertSelectionToBorderVertices(/*bAsTriangleTopology*/ true);
-}
-
 void UWeightToolTransferManager::InitialSetup(USkinWeightsPaintTool* InWeightTool, FEditorViewportClient* InViewportClient)
 {
 	WeightTool = InWeightTool;
@@ -719,9 +392,9 @@ void UWeightToolTransferManager::InitialSetup(USkinWeightsPaintTool* InWeightToo
 
 	// create the mesh selector and run initial setup
 	// NOTE: currently this must happen inside the Setup of a UInteractiveTool so that input is routed to the selection mechanic
-	MeshSelector = NewObject<UWeightToolMeshSelector>(this);
+	MeshSelector = NewObject<UToolMeshSelector>(this);
 	auto DoNothingOnSelection = [](){};
-	MeshSelector->InitialSetup(WeightTool->GetTargetWorld(), WeightTool.Get(), InViewportClient,  DoNothingOnSelection);
+	MeshSelector->InitialSetup(WeightTool->GetTargetWorld(), WeightTool.Get(), DoNothingOnSelection);
 }
 
 void UWeightToolTransferManager::SetSourceMesh(USkeletalMesh* InSkeletalMesh)
@@ -2236,9 +1909,9 @@ void USkinWeightsPaintTool::Setup()
 	SelectionIsolator->InitialSetup(this);
 
 	// setup selection for the main mesh
-	MeshSelector = NewObject<UWeightToolMeshSelector>(this);
+	MeshSelector = NewObject<UToolMeshSelector>(this);
 	auto OnSelectionChangedLambda = [this](){OnSelectionChanged.Broadcast();};
-	MeshSelector->InitialSetup(TargetWorld.Get(), this, GetViewportClient(), OnSelectionChangedLambda);
+	MeshSelector->InitialSetup(TargetWorld.Get(), this, OnSelectionChangedLambda);
 
 	// run all initialization for mesh/weights
 	UpdateCurrentlyEditedMesh(Component, *GetCurrentlyEditedDynamicMesh(), *GetCurrentlyEditedMeshDescription());
@@ -4448,7 +4121,7 @@ void USkinWeightsPaintTool::UpdateSelectorState() const
 	const bool bEnableTargetMeshSelector = bIsMeshEditing && !bWasSetToSelectSource;
 	
 	// source mesh selector
-	UWeightToolMeshSelector* SourceMeshSelector = TransferManager->GetMeshSelector();
+	UToolMeshSelector* SourceMeshSelector = TransferManager->GetMeshSelector();
 	SourceMeshSelector->SetIsEnabled(bEnableSourceMeshSelector);
 	
 	// main mesh selector
@@ -4459,12 +4132,12 @@ void USkinWeightsPaintTool::UpdateSelectorState() const
 	SourceMeshSelector->SetComponentSelectionMode(WeightToolProperties->ComponentSelectionMode);
 }
 
-UWeightToolMeshSelector* USkinWeightsPaintTool::GetMainMeshSelector()
+UToolMeshSelector* USkinWeightsPaintTool::GetMainMeshSelector()
 {
 	return MeshSelector;
 }
 
-UWeightToolMeshSelector* USkinWeightsPaintTool::GetActiveMeshSelector()
+UToolMeshSelector* USkinWeightsPaintTool::GetActiveMeshSelector()
 {
 	if (WeightToolProperties->MeshSelectMode == EMeshTransferOption::Source)
 	{
@@ -4897,7 +4570,7 @@ bool USkinWeightsPaintTool::HasActiveSelectionOnMainMesh()
 		return false;
 	}
 
-	UWeightToolMeshSelector* MainMeshSelector = GetMainMeshSelector();
+	UToolMeshSelector* MainMeshSelector = GetMainMeshSelector();
 	if (!MainMeshSelector)
 	{
 		return false;
diff --git a/Engine/Plugins/Runtime/MeshModelingToolset/Source/MeshModelingToolsEditorOnly/Private/ToolMeshSelector.cpp b/Engine/Plugins/Runtime/MeshModelingToolset/Source/MeshModelingToolsEditorOnly/Private/ToolMeshSelector.cpp
new file mode 100644
--- /dev/null
+++ b/Engine/Plugins/Runtime/MeshModelingToolset/Source/MeshModelingToolsEditorOnly/Private/ToolMeshSelector.cpp
@@ -0,0 +1,338 @@
+// Copyright Epic Games, Inc. All Rights Reserved.
+
+#include "ToolMeshSelector.h"
+
+#include "DynamicMesh/NonManifoldMappingSupport.h"
+#include "Engine/World.h"
+#include "Materials/Material.h"
+#include "Materials/MaterialInterface.h"
+#include "Materials/MaterialInstanceDynamic.h"
+#include "PreviewMesh.h"
+#include "Selection/PolygonSelectionMechanic.h"
+
+// thise method is deprecated - InViewportClient is no longer used 
+void UToolMeshSelector::InitialSetup(UWorld* InWorld, UInteractiveTool* InParentTool, FEditorViewportClient* InViewportClient, TFunction<void()> OnSelectionChangedFunc)
+{
+	InitialSetup(InWorld, InParentTool, OnSelectionChangedFunc);
+}
+
+void UToolMeshSelector::InitialSetup(UWorld* InWorld, UInteractiveTool* InParentTool, TFunction<void()> OnSelectionChangedFunc)
+{
+	World = InWorld;
+
+	// set up vertex selection mechanic
+	PolygonSelectionMechanic = NewObject<UPolygonSelectionMechanic>(this);
+	PolygonSelectionMechanic->bAddSelectionFilterPropertiesToParentTool = false;
+	PolygonSelectionMechanic->Setup(InParentTool);
+	PolygonSelectionMechanic->SetIsEnabled(false, false);
+	PolygonSelectionMechanic->OnSelectionChanged.AddLambda(OnSelectionChangedFunc);
+
+	// set up style of vertex selection
+	constexpr FLinearColor VertexSelectedPurple = FLinearColor(0.78f, 0.f, 0.78f);
+	constexpr FLinearColor VertexSelectedYellow = FLinearColor(1.f, 1.f, 0.f);
+	// adjust selection rendering for this context
+	PolygonSelectionMechanic->HilightRenderer.PointColor = FLinearColor::Blue;
+	PolygonSelectionMechanic->HilightRenderer.PointSize = 10.0f;
+	// vertex highlighting once selected
+	PolygonSelectionMechanic->SelectionRenderer.LineThickness = 1.0f;
+	PolygonSelectionMechanic->SelectionRenderer.PointColor = VertexSelectedYellow;
+	PolygonSelectionMechanic->SelectionRenderer.PointSize = 5.0f;
+	PolygonSelectionMechanic->SelectionRenderer.DepthBias = 2.0f;
+	// despite the name, this renders the vertices
+	PolygonSelectionMechanic->PolyEdgesRenderer.PointColor = VertexSelectedPurple;
+	PolygonSelectionMechanic->PolyEdgesRenderer.PointSize = 5.0f;
+	PolygonSelectionMechanic->PolyEdgesRenderer.DepthBias = 2.0f;
+	PolygonSelectionMechanic->PolyEdgesRenderer.LineThickness = 1.0f;
+}
+
+void UToolMeshSelector::SetMesh(
+	UPreviewMesh* InPreviewMesh,
+	const FMeshDescription* InMeshDescription,
+	const FTransform3d& InMeshTransform)
+{
+	// store the mesh this is operating on
+	PreviewMesh = InPreviewMesh;
+	// copy the mesh description for selection validation
+	MeshDescription = *InMeshDescription;
+
+	if (!ensure(World))
+	{
+		return;
+	}
+
+	if (!PreviewMesh)
+	{
+		SetIsEnabled(false);
+		return;
+	}
+
+	// reset selection topology and mesh spatial data
+	static constexpr bool bAutoBuild = true;
+	const FDynamicMesh3* DynamicMesh = PreviewMesh->GetMesh();
+	SelectionTopology = MakeUnique<UE::Geometry::FTriangleGroupTopology>(DynamicMesh, bAutoBuild);
+	MeshSpatial = MakeUnique<FDynamicMeshAABBTree3>(DynamicMesh, bAutoBuild);
+
+	// initialize the selection mechanic
+	PolygonSelectionMechanic->Initialize(
+		DynamicMesh,
+		InMeshTransform,
+		World,
+		SelectionTopology.Get(),
+		[this]() { return MeshSpatial.Get(); }
+	);
+
+	// clear the selection (old selection is invalid on new topo)
+	PolygonSelectionMechanic->ClearSelection();
+	PolygonSelectionMechanic->ClearHighlight();
+
+	// selection colors
+	constexpr FLinearColor FaceSelectedOrange = FLinearColor(0.886f, 0.672f, 0.473f);
+	// configure secondary render material for selected triangles
+	// NOTE: the material returned by ToolSetupUtil::GetSelectionMaterial has a checkerboard pattern on back faces which makes it hard to use
+	if (UMaterialInterface* Material = LoadObject<UMaterial>(nullptr, TEXT("/MeshModelingToolsetExp/Materials/SculptMaterial")))
+	{
+		if (UMaterialInstanceDynamic* MatInstance = UMaterialInstanceDynamic::Create(Material, World))
+		{
+			MatInstance->SetVectorParameterValue(TEXT("Color"), FaceSelectedOrange);
+			PreviewMesh->SetSecondaryRenderMaterial(MatInstance);
+		}
+	}
+
+	// secondary triangle buffer used to render face selection
+	PreviewMesh->EnableSecondaryTriangleBuffers([this](const FDynamicMesh3* Mesh, int32 TriangleID)
+		{
+			return PolygonSelectionMechanic->GetActiveSelection().IsSelectedTriangle(Mesh, SelectionTopology.Get(), TriangleID);
+		});
+	// notify preview mesh when triangle selection has been updated
+	PolygonSelectionMechanic->OnSelectionChanged.AddWeakLambda(this, [this]()
+		{
+			PreviewMesh->FastNotifySecondaryTrianglesChanged();
+		});
+	PolygonSelectionMechanic->OnFaceSelectionPreviewChanged.AddWeakLambda(this, [this]()
+		{
+			PreviewMesh->FastNotifySecondaryTrianglesChanged();
+		});
+}
+
+void UToolMeshSelector::Shutdown()
+{
+	if (PolygonSelectionMechanic)
+	{
+		PolygonSelectionMechanic->Shutdown();
+	}
+
+	PolygonSelectionMechanic = nullptr;
+}
+
+void UToolMeshSelector::SetIsEnabled(bool bIsEnabled)
+{
+	if (!PolygonSelectionMechanic)
+	{
+		return;
+	}
+
+	// force off if there's no preview mesh
+	bIsEnabled = PreviewMesh ? bIsEnabled : false;
+
+	PolygonSelectionMechanic->SetIsEnabled(bIsEnabled, bIsEnabled);
+}
+
+void UToolMeshSelector::SetComponentSelectionMode(EComponentSelectionMode InMode)
+{
+	if (!(PolygonSelectionMechanic))
+	{
+		return;
+	}
+
+	PolygonSelectionMechanic->Properties->bSelectVertices = InMode == EComponentSelectionMode::Vertices;
+	PolygonSelectionMechanic->Properties->bSelectEdges = InMode == EComponentSelectionMode::Edges;
+	PolygonSelectionMechanic->Properties->bSelectFaces = InMode == EComponentSelectionMode::Faces;
+	PolygonSelectionMechanic->SetShowSelectableCorners(InMode == EComponentSelectionMode::Vertices);
+	PolygonSelectionMechanic->SetShowEdges(InMode == EComponentSelectionMode::Edges);
+}
+
+void UToolMeshSelector::SetTransform(const FTransform3d& InTargetTransform)
+{
+	if (PolygonSelectionMechanic)
+	{
+		PolygonSelectionMechanic->SetTransform(InTargetTransform);
+	}
+}
+
+void UToolMeshSelector::UpdateAfterMeshDeformation()
+{
+	MeshSpatial->Build();
+	constexpr bool bTopologyDeformed = true;
+	constexpr bool bTopologyModified = false;
+	PolygonSelectionMechanic->GetTopologySelector()->Invalidate(bTopologyDeformed, bTopologyModified);
+}
+
+void UToolMeshSelector::DrawHUD(FCanvas* Canvas, IToolsContextRenderAPI* RenderAPI)
+{
+	if (!PolygonSelectionMechanic)
+	{
+		return;
+	}
+
+	PolygonSelectionMechanic->DrawHUD(Canvas, RenderAPI);
+}
+
+void UToolMeshSelector::Render(IToolsContextRenderAPI* RenderAPI)
+{
+	if (!(PolygonSelectionMechanic && PreviewMesh))
+	{
+		return;
+	}
+
+	PolygonSelectionMechanic->Render(RenderAPI);
+}
+
+const TArray<int32>& UToolMeshSelector::GetSelectedVertices()
+{
+	using UE::Geometry::FNonManifoldMappingSupport;
+
+	SelectedVerticesInternal.Empty();
+	if (!(PolygonSelectionMechanic && PreviewMesh))
+	{
+		return SelectedVerticesInternal;
+	}
+
+	const FGroupTopologySelection& Selection = PolygonSelectionMechanic->GetActiveSelection();
+	const FDynamicMesh3* DynamicMesh = PreviewMesh->GetMesh();
+	const FVertexArray& AllVertices = MeshDescription.Vertices();
+	const FNonManifoldMappingSupport NonManifoldMappingSupport(*DynamicMesh);
+
+	// validate and add vertices to the output array
+	auto AddVertices = [this, &AllVertices, &NonManifoldMappingSupport](const TSet<int32>& VerticesToAdd)
+		{
+			// we have to make sure that the vertex ids are safe to use as PolygonSelectionMechanic does not act on the
+			// mesh description but on the dynamic mesh that can duplicate vertices when dealing with degenerate triangles.
+			// cf. FMeshDescriptionToDynamicMesh::Convert for more details.
+			for (const int32 VertexToAdd : VerticesToAdd)
+			{
+				const int32 SrcVertexID = NonManifoldMappingSupport.GetOriginalNonManifoldVertexID(VertexToAdd);
+				if (AllVertices.IsValid(SrcVertexID))
+				{
+					SelectedVerticesInternal.Add(SrcVertexID);
+				}
+			}
+		};
+
+	// add selected vertices
+	AddVertices(Selection.SelectedCornerIDs);
+
+	// add vertices on selected edges
+	{
+		TSet<int32> VerticesInSelectedEdges;
+		for (const int32 SelectedEdgeIndex : Selection.SelectedEdgeIDs)
+		{
+			FDynamicMesh3::FEdge CurrentEdge = DynamicMesh->GetEdge(SelectedEdgeIndex);
+			VerticesInSelectedEdges.Add(CurrentEdge.Vert.A);
+			VerticesInSelectedEdges.Add(CurrentEdge.Vert.B);
+		}
+
+		AddVertices(VerticesInSelectedEdges);
+	}
+
+	// add vertices in selected faces
+	{
+		TSet<int32> VerticesInSelectedFaces;
+		for (const int32 SelectedFaceIndex : Selection.SelectedGroupIDs)
+		{
+			UE::Geometry::FIndex3i TriangleVertices = DynamicMesh->GetTriangleRef(SelectedFaceIndex);
+			VerticesInSelectedFaces.Add(TriangleVertices[0]);
+			VerticesInSelectedFaces.Add(TriangleVertices[1]);
+			VerticesInSelectedFaces.Add(TriangleVertices[2]);
+		}
+
+		AddVertices(VerticesInSelectedFaces);
+	}
+
+	return SelectedVerticesInternal;
+}
+
+bool UToolMeshSelector::IsAnyComponentSelected() const
+{
+	if (!PolygonSelectionMechanic)
+	{
+		return false;
+	}
+
+	return PolygonSelectionMechanic->HasSelection();
+}
+
+void UToolMeshSelector::GetSelectedTriangles(TArray<int32>& OutTriangleIndices) const
+{
+	OutTriangleIndices.Empty();
+	if (!ensure(PolygonSelectionMechanic))
+	{
+		return;
+	}
+
+	const FGroupTopologySelection& Selection = PolygonSelectionMechanic->GetActiveSelection();
+	const FDynamicMesh3* DynamicMesh = PreviewMesh->GetMesh();
+	TSet<int32> TriangleSet;
+
+	// add triangles connected to selected vertices
+	for (const int32 VertexIndex : Selection.SelectedCornerIDs)
+	{
+		DynamicMesh->EnumerateVertexTriangles(VertexIndex, [&TriangleSet](int32 TriangleIndex)
+			{
+				TriangleSet.Add(TriangleIndex);
+			});
+	}
+
+	// add triangles connected to selected edges
+	for (const int32 EdgeIndex : Selection.SelectedEdgeIDs)
+	{
+		DynamicMesh->EnumerateEdgeTriangles(EdgeIndex, [&TriangleSet](int32 TriangleIndex)
+			{
+				TriangleSet.Add(TriangleIndex);
+			});
+	}
+
+	// add selected triangles
+	TriangleSet.Append(Selection.SelectedGroupIDs);
+
+	OutTriangleIndices = TriangleSet.Array();
+}
+
+void UToolMeshSelector::GrowSelection() const
+{
+	if (!ensure(PolygonSelectionMechanic))
+	{
+		return;
+	}
+
+	PolygonSelectionMechanic->GrowSelection(/*bAsTriangleTopology*/ true);
+}
+
+void UToolMeshSelector::ShrinkSelection() const
+{
+	if (!ensure(PolygonSelectionMechanic))
+	{
+		return;
+	}
+
+	PolygonSelectionMechanic->ShrinkSelection(/*bAsTriangleTopology*/ true);
+}
+
+void UToolMeshSelector::FloodSelection() const
+{
+	if (!ensure(PolygonSelectionMechanic))
+	{
+		return;
+	}
+
+	PolygonSelectionMechanic->FloodSelection();
+}
+
+void UToolMeshSelector::SelectBorder() const
+{
+	if (!ensure(PolygonSelectionMechanic))
+	{
+		return;
+	}
+
+	PolygonSelectionMechanic->ConvertSelectionToBorderVertices(/*bAsTriangleTopology*/ true);
+}
\ No newline at end of file
diff --git a/Engine/Plugins/Runtime/MeshModelingToolset/Source/MeshModelingToolsEditorOnly/Public/SkeletalMesh/SkinWeightsPaintTool.h b/Engine/Plugins/Runtime/MeshModelingToolset/Source/MeshModelingToolsEditorOnly/Public/SkeletalMesh/SkinWeightsPaintTool.h
--- a/Engine/Plugins/Runtime/MeshModelingToolset/Source/MeshModelingToolsEditorOnly/Public/SkeletalMesh/SkinWeightsPaintTool.h
+++ b/Engine/Plugins/Runtime/MeshModelingToolset/Source/MeshModelingToolsEditorOnly/Public/SkeletalMesh/SkinWeightsPaintTool.h
@@ -17,6 +17,7 @@
 #include "Engine/SkeletalMesh.h"
 #include "Selections/GeometrySelection.h"
 #include "TargetInterfaces/MeshTargetInterfaceTypes.h"
+#include "ToolMeshSelector.h"
 
 #include "SkinWeightsPaintTool.generated.h"
 
@@ -38,7 +39,6 @@ namespace UE::Geometry
 }
 
 using BoneIndex = int32;
-using VertexIndex = int32;
 
 // weight edit mode
 UENUM()
@@ -49,15 +49,6 @@ enum class EWeightEditMode : uint8
 	Bones,
 };
 
-// component selection mode
-UENUM()
-enum class EComponentSelectionMode : uint8
-{
-	Vertices,
-	Edges,
-	Faces
-};
-
 // weight transfers happen between a source and target
 UENUM()
 enum class EMeshTransferOption : uint8
@@ -480,70 +471,10 @@ private:
 // NOTE: at some point we may want to do component selections on multiple meshes in any/all viewports
 // at which time this class should be centralized and renamed to UMeshSelector or something like that.
 // But there will need to be some sort of centralized facility to manage that and make sure it interacts nicely with other tools.
-UCLASS(MinimalAPI)
-class UWeightToolMeshSelector : public UObject
+UCLASS(MinimalAPI, Deprecated)
+class UDEPRECATED_WeightToolMeshSelector : public UToolMeshSelector
 {
 	GENERATED_BODY()
-
-public:
-
-	// must be called during the Setup of the parent tool
-	UE_API void InitialSetup(
-		UWorld* InWorld,
-		UInteractiveTool* InParentTool,
-		FEditorViewportClient* InViewportClient,
-		TFunction<void()> OnSelectionChangedFunc);
-
-	// must be called AFTER InitialSetup, and any time the mesh is changed
-	// passing in a null preview mesh will disable the selector
-	UE_API void SetMesh(
-		UPreviewMesh* InMesh,
-		const FMeshDescription* InMeshDescription,
-		const FTransform3d& InMeshTransform);
-	
-	UE_API void UpdateAfterMeshDeformation();
-	
-	UE_API void Shutdown();
-
-	UE_API void SetIsEnabled(bool bIsEnabled);
-	UE_API void SetComponentSelectionMode(EComponentSelectionMode InMode);
-	UE_API void SetTransform(const FTransform3d& InTargetTransform);
-
-	// viewport 
-	UE_API void DrawHUD(FCanvas* Canvas, IToolsContextRenderAPI* RenderAPI);
-	UE_API void Render(IToolsContextRenderAPI* RenderAPI);
-
-	// get a list of currently selected vertices (converting edges and faces to vertices)
-	UE_API const TArray<int32>& GetSelectedVertices();
-	UE_API bool IsAnyComponentSelected() const;
-	UE_API void GetSelectedTriangles(TArray<int32>& OutTriangleIndices) const;
-
-	// edit selection
-	UE_API void GrowSelection() const;
-	UE_API void ShrinkSelection() const;
-	UE_API void FloodSelection() const;
-	UE_API void SelectBorder() const;
-
-	// get access to the selection mechanic
-	UPolygonSelectionMechanic* GetSelectionMechanic() { return PolygonSelectionMechanic; };
-
-private:
-
-	UPROPERTY()
-	TObjectPtr<UInteractiveTool> ParentTool;
-	UPROPERTY()
-	TObjectPtr<UWorld> World;
-	FEditorViewportClient* ViewportClient;
-	UPROPERTY()
-	TObjectPtr<UPreviewMesh> PreviewMesh;
-	FMeshDescription MeshDescription;
-	UPROPERTY()
-	TObjectPtr<UPolygonSelectionMechanic> PolygonSelectionMechanic;
-	
-	TUniquePtr<UE::Geometry::FDynamicMeshAABBTree3> MeshSpatial = nullptr;
-	TUniquePtr<UE::Geometry::FTriangleGroupTopology> SelectionTopology = nullptr;
-	
-	TArray<VertexIndex> SelectedVerticesInternal;
 };
 
 // this class wraps a source skeletal mesh used to transfer skin weights to the tool target mesh
@@ -580,7 +511,7 @@ public:
 	UPreviewMesh* GetPreviewMesh() const { return SourcePreviewMesh; };
 	
 	// get the mesh selector for the source mesh
-	UWeightToolMeshSelector* GetMeshSelector() const { return MeshSelector; };
+	UToolMeshSelector* GetMeshSelector() const { return MeshSelector; };
 	
 	// called when tool settings are modified
 	UE_API void OnPropertyModified(const USkinWeightsPaintToolProperties* WeightToolProperties, const FProperty* ModifiedProperty);
@@ -605,7 +536,7 @@ private:
 	UPROPERTY()
 	TObjectPtr<UToolTarget> SourceTarget = nullptr;
 	UPROPERTY()
-	TObjectPtr<UWeightToolMeshSelector> MeshSelector;
+	TObjectPtr<UToolMeshSelector> MeshSelector;
 	TWeakObjectPtr<USkinWeightsPaintTool> WeightTool;
 };
 
@@ -784,9 +715,9 @@ public:
 	UE_API void UpdateSelectorState() const;
 	
 	// get access to the mesh selector for the main mesh
-	UE_API UWeightToolMeshSelector* GetMainMeshSelector();
+	UE_API UToolMeshSelector* GetMainMeshSelector();
 	// get access to the currently active mesh selector (may be on the transfer source mesh)
-	UE_API UWeightToolMeshSelector* GetActiveMeshSelector();
+	UE_API UToolMeshSelector* GetActiveMeshSelector();
 	// does the main mesh have an active selection ("active" meaning the selection is currently being rendered in the view and is editable)
 	UE_API bool HasActiveSelectionOnMainMesh();
 	// select all vertices affected by the currently selected bone(s)
@@ -955,7 +886,7 @@ protected:
 
 	// the selection system for the main mesh
 	UPROPERTY()
-	TObjectPtr<UWeightToolMeshSelector> MeshSelector;
+	TObjectPtr<UToolMeshSelector> MeshSelector;
 
 	// skin weight layer
 	UE_API void OnActiveLODChanged();
diff --git a/Engine/Plugins/Runtime/MeshModelingToolset/Source/MeshModelingToolsEditorOnly/Public/ToolMeshSelector.h b/Engine/Plugins/Runtime/MeshModelingToolset/Source/MeshModelingToolsEditorOnly/Public/ToolMeshSelector.h
new file mode 100644
--- /dev/null
+++ b/Engine/Plugins/Runtime/MeshModelingToolset/Source/MeshModelingToolsEditorOnly/Public/ToolMeshSelector.h
@@ -0,0 +1,105 @@
+// Copyright Epic Games, Inc. All Rights Reserved.
+
+#pragma once
+
+#include "Containers/Array.h"
+#include "DynamicMesh/DynamicMeshAABBTree3.h"
+#include "MeshDescription.h"
+#include "GroupTopology.h"
+#include "Templates/UniquePtr.h"
+#include "UObject/ObjectPtr.h"
+
+#include "Math/MathFwd.h"
+
+#include "ToolMeshSelector.generated.h"
+
+#define UE_API MESHMODELINGTOOLSEDITORONLY_API
+
+class FCanvas;
+class FEditorViewportClient;
+class IToolsContextRenderAPI;
+class UInteractiveTool;
+class UPolygonSelectionMechanic;
+class UPreviewMesh;
+class UWorld;
+
+// component selection mode
+UENUM()
+enum class EComponentSelectionMode : uint8
+{
+	Vertices,
+	Edges,
+	Faces
+};
+
+using VertexIndex = int32;
+
+// this class wraps the all the components to enable selection on a single mesh in the skin weights tool
+// this allows us to make selections on multiple different meshes
+// NOTE: at some point we may want to do component selections on multiple meshes in any/all viewports
+// at which time this class should be centralized and renamed to UMeshSelector or something like that.
+// But there will need to be some sort of centralized facility to manage that and make sure it interacts nicely with other tools.
+UCLASS(MinimalAPI)
+class UToolMeshSelector : public UObject
+{
+	GENERATED_BODY()
+
+public:
+
+	// must be called during the Setup of the parent tool
+	UE_DEPRECATED(5.7, "Use InitialSetup that does not need use a Viewport Client")
+	UE_API void InitialSetup(UWorld* InWorld, UInteractiveTool* InParentTool, FEditorViewportClient* InViewportClient, TFunction<void()> OnSelectionChangedFunc);
+	UE_API void InitialSetup(UWorld* InWorld, UInteractiveTool* InParentTool, TFunction<void()> OnSelectionChangedFunc);
+
+	// must be called AFTER InitialSetup, and any time the mesh is changed
+	// passing in a null preview mesh will disable the selector
+	UE_API void SetMesh(
+		UPreviewMesh* InMesh,
+		const FMeshDescription* InMeshDescription,
+		const FTransform3d& InMeshTransform);
+
+	UE_API void UpdateAfterMeshDeformation();
+
+	UE_API void Shutdown();
+
+	UE_API void SetIsEnabled(bool bIsEnabled);
+	UE_API void SetComponentSelectionMode(EComponentSelectionMode InMode);
+	UE_API void SetTransform(const FTransform3d& InTargetTransform);
+
+	// viewport 
+	UE_API void DrawHUD(FCanvas* Canvas, IToolsContextRenderAPI* RenderAPI);
+	UE_API void Render(IToolsContextRenderAPI* RenderAPI);
+
+	// get a list of currently selected vertices (converting edges and faces to vertices)
+	UE_API const TArray<int32>& GetSelectedVertices();
+	UE_API bool IsAnyComponentSelected() const;
+	UE_API void GetSelectedTriangles(TArray<int32>& OutTriangleIndices) const;
+
+	// edit selection
+	UE_API void GrowSelection() const;
+	UE_API void ShrinkSelection() const;
+	UE_API void FloodSelection() const;
+	UE_API void SelectBorder() const;
+
+	// get access to the selection mechanic
+	UPolygonSelectionMechanic* GetSelectionMechanic() { return PolygonSelectionMechanic; };
+
+private:
+
+	UPROPERTY()
+	TObjectPtr<UInteractiveTool> ParentTool;
+	UPROPERTY()
+	TObjectPtr<UWorld> World;
+	UPROPERTY()
+	TObjectPtr<UPreviewMesh> PreviewMesh;
+	FMeshDescription MeshDescription;
+	UPROPERTY()
+	TObjectPtr<UPolygonSelectionMechanic> PolygonSelectionMechanic;
+
+	TUniquePtr<UE::Geometry::FDynamicMeshAABBTree3> MeshSpatial = nullptr;
+	TUniquePtr<UE::Geometry::FTriangleGroupTopology> SelectionTopology = nullptr;
+
+	TArray<VertexIndex> SelectedVerticesInternal;
+};
+
+#undef UE_API
