Skin Weights Paint Tools: reset source LOD and profile when changing the source mesh

#rnx

diff --git a/Engine/Plugins/Runtime/MeshModelingToolset/Source/MeshModelingToolsEditorOnly/Private/SkeletalMesh/SkinWeightsPaintTool.cpp b/Engine/Plugins/Runtime/MeshModelingToolset/Source/MeshModelingToolsEditorOnly/Private/SkeletalMesh/SkinWeightsPaintTool.cpp
--- a/Engine/Plugins/Runtime/MeshModelingToolset/Source/MeshModelingToolsEditorOnly/Private/SkeletalMesh/SkinWeightsPaintTool.cpp
+++ b/Engine/Plugins/Runtime/MeshModelingToolset/Source/MeshModelingToolsEditorOnly/Private/SkeletalMesh/SkinWeightsPaintTool.cpp
@@ -399,9 +399,14 @@ void UWeightToolTransferManager::InitialSetup(USkinWeightsPaintTool* InWeightToo
 
 void UWeightToolTransferManager::SetSourceMesh(USkeletalMesh* InSkeletalMesh)
 {
-	SourceSkeletalMesh = InSkeletalMesh;
-
 	USkinWeightsPaintToolProperties* ToolProperties = WeightTool->GetWeightToolProperties();
+	if (!InSkeletalMesh || InSkeletalMesh != SourceSkeletalMesh)
+	{
+		ToolProperties->SourceSkinWeightProfile = FSkeletalMeshAttributesShared::DefaultSkinWeightProfileName;
+		ToolProperties->SourceLOD = "LOD0";
+	}
+		
+	SourceSkeletalMesh = InSkeletalMesh;
 	
 	// reset to prepare for new mesh (or possibly no mesh)
 	{
@@ -425,8 +430,6 @@ void UWeightToolTransferManager::SetSourceMesh(USkeletalMesh* InSkeletalMesh)
 	
 	if (!InSkeletalMesh)
 	{
-		ToolProperties->SourceSkinWeightProfile = FSkeletalMeshAttributesShared::DefaultSkinWeightProfileName;
-		ToolProperties->SourceLOD = "LOD0";
 		WeightTool->UpdateSelectorState();
 		return;
 	}
