[HLOD] Support Spline Mesh Components in Approximate Mesh HLOD

diff --git a/Engine/Plugins/Runtime/MeshModelingToolset/Source/ModelingComponents/Private/Scene/MeshSceneAdapter.cpp b/Engine/Plugins/Runtime/MeshModelingToolset/Source/ModelingComponents/Private/Scene/MeshSceneAdapter.cpp
--- a/Engine/Plugins/Runtime/MeshModelingToolset/Source/ModelingComponents/Private/Scene/MeshSceneAdapter.cpp
+++ b/Engine/Plugins/Runtime/MeshModelingToolset/Source/ModelingComponents/Private/Scene/MeshSceneAdapter.cpp
@@ -22,6 +22,7 @@
 #include "Spatial/SparseDynamicOctree3.h"
 
 #include "GameFramework/Actor.h"
+#include "Components/SplineMeshComponent.h"
 #include "Components/StaticMeshComponent.h"
 #include "Components/InstancedStaticMeshComponent.h"
 #include "Engine/StaticMesh.h"
@@ -1172,14 +1173,20 @@ public:
 	TUniquePtr<TMeshAABBTree3<FColliderMesh>> AABBTree;
 	TUniquePtr<TFastWindingTree<FColliderMesh>> FWNTree;
 
+	virtual TUniquePtr<TStaticMeshSpatialWrapperBase<FMeshDescriptionTriangleMeshSurfaceAdapter>> GetIntermediateWrapper()
+	{
+		TUniquePtr<FStaticMeshSourceDataSpatialWrapper> TempWrapper = MakeUnique<FStaticMeshSourceDataSpatialWrapper>();
+		TempWrapper->SourceContainer = SourceContainer;
+		TempWrapper->StaticMesh = StaticMesh;
+		return TempWrapper;
+	}
+
 	virtual bool Build(const FMeshSceneAdapterBuildOptions& BuildOptions) override
 	{
 		// this is because Build() will be called twice
 		if (ColliderMesh.TriangleCount() == 0)
 		{
-			TUniquePtr<FStaticMeshSourceDataSpatialWrapper> TempWrapper = MakeUnique<FStaticMeshSourceDataSpatialWrapper>();
-			TempWrapper->SourceContainer = SourceContainer;
-			TempWrapper->StaticMesh = StaticMesh;
+			TUniquePtr<TStaticMeshSpatialWrapperBase<FMeshDescriptionTriangleMeshSurfaceAdapter>> TempWrapper = GetIntermediateWrapper();
 
 			FMeshSceneAdapterBuildOptions TempBuildOptions;
 			TempBuildOptions.bBuildSpatialDataStructures = false;	// we are just copying from this so do not build spatials
@@ -1301,13 +1308,76 @@ public:
 };
 
 
+class FSplineStaticMeshSpatialWrapperBase : public TStaticMeshSpatialWrapperBase<FMeshDescriptionTriangleMeshSurfaceAdapter>
+{
+public:
+	FMeshDescription DeformedSourceMeshDescription;
+	USplineMeshComponent* SplineMeshComponent = nullptr;
 
+	virtual void InitializeSourceMeshDescription() = 0;
 
+	virtual TUniquePtr<FMeshDescriptionTriangleMeshSurfaceAdapter> MakeMeshAdapter(const FMeshSceneAdapterBuildOptions& BuildOptions) override
+	{
+#if WITH_EDITOR
+		InitializeSourceMeshDescription();
+		SplineMeshComponent->PropagateSplineDeformationToMesh(DeformedSourceMeshDescription);
 
+		TUniquePtr<FMeshDescriptionTriangleMeshSurfaceAdapter> SurfaceAdapter = MakeUnique<FMeshDescriptionTriangleMeshSurfaceAdapter>(&DeformedSourceMeshDescription, StaticMesh, BuildOptions.bOnlySurfaceMaterials);
+		const FMeshBuildSettings& LODBuildSettings = StaticMesh->GetSourceModel(LODIndex).BuildSettings;
+		SurfaceAdapter->SetBuildScale(LODBuildSettings.BuildScale3D, false);
+		return SurfaceAdapter;
+#else
+		// cannot use this path in non-editor builds, but this should have been handled higher up by instantiating 
+		// the runtime subclass
+		check(false);
+		return nullptr;
+#endif
+	}
+};
 
 
+class FSplineStaticMeshSourceDataSpatialWrapper : public FSplineStaticMeshSpatialWrapperBase
+{
+	virtual void InitializeSourceMeshDescription() override
+	{
+#if WITH_EDITOR
+		StaticMesh->CloneMeshDescription(LODIndex, DeformedSourceMeshDescription);
+#endif
+	}
+};
 
-static TUniquePtr<IMeshSpatialWrapper> SpatialWrapperFactory( const FMeshTypeContainer& MeshContainer, const FMeshSceneAdapterBuildOptions& BuildOptions )
+
+class FSplineStaticMeshRenderDataSpatialWrapper : public FSplineStaticMeshSpatialWrapperBase
+{
+	virtual void InitializeSourceMeshDescription() override
+	{
+#if WITH_EDITOR
+		if (const FStaticMeshLODResources* LODResource = StaticMesh->GetRenderData()->GetCurrentFirstLOD(0))
+		{
+			StaticMesh->ExportStaticMeshLOD(*LODResource, DeformedSourceMeshDescription);
+		}
+#endif
+	}
+};
+
+
+class FCompressedSplineStaticMeshSpatialWrapper : public FCompressedStaticMeshSpatialWrapper
+{
+public:
+	USplineMeshComponent* SplineMeshComponent = nullptr;
+
+	virtual TUniquePtr<TStaticMeshSpatialWrapperBase<FMeshDescriptionTriangleMeshSurfaceAdapter>> GetIntermediateWrapper() override
+	{
+		TUniquePtr<FSplineStaticMeshSourceDataSpatialWrapper> TempWrapper = MakeUnique<FSplineStaticMeshSourceDataSpatialWrapper>();
+		TempWrapper->SourceContainer = SourceContainer;
+		TempWrapper->StaticMesh = StaticMesh;
+		TempWrapper->SplineMeshComponent = SplineMeshComponent;
+		return TempWrapper;
+	}
+};
+
+
+static TUniquePtr<IMeshSpatialWrapper> SpatialWrapperFactory( const FMeshTypeContainer& MeshContainer, const TArray<FActorChildMesh*>& ParentMeshes, const FMeshSceneAdapterBuildOptions& BuildOptions )
 {
 	if (MeshContainer.MeshType == ESceneMeshType::StaticMeshAsset)
 	{
@@ -1319,27 +1389,57 @@ static TUniquePtr<IMeshSpatialWrapper> SpatialWrapperFactory( const FMeshTypeCon
 #else
 			bool bUseRenderMeshes = true;
 #endif
-
-			if (bUseRenderMeshes)
-			{
-				TUniquePtr<FStaticMeshRenderDataSpatialWrapper> SMWrapper = MakeUnique<FStaticMeshRenderDataSpatialWrapper>();
-				SMWrapper->SourceContainer = MeshContainer;
-				SMWrapper->StaticMesh = StaticMesh;
-				return SMWrapper;
-			}
-			else if (BuildOptions.bEnableUVQueries == false && BuildOptions.bEnableNormalsQueries == false)
+			const bool bIsSplineMesh = (ParentMeshes.Num() == 1 && ParentMeshes[0]->ComponentType == EActorMeshComponentType::SplineStaticMesh);
+			if (bIsSplineMesh)
 			{
-				TUniquePtr<FCompressedStaticMeshSpatialWrapper> SMWrapper = MakeUnique<FCompressedStaticMeshSpatialWrapper>();
-				SMWrapper->SourceContainer = MeshContainer;
-				SMWrapper->StaticMesh = StaticMesh;
-				return SMWrapper;
+				if (bUseRenderMeshes)
+				{
+					TUniquePtr<FSplineStaticMeshRenderDataSpatialWrapper> SplineSMWrapper = MakeUnique<FSplineStaticMeshRenderDataSpatialWrapper>();
+					SplineSMWrapper->SourceContainer = MeshContainer;
+					SplineSMWrapper->StaticMesh = StaticMesh;
+					SplineSMWrapper->SplineMeshComponent = Cast<USplineMeshComponent>(ParentMeshes[0]->SourceComponent);
+					return SplineSMWrapper;
+				}
+				else if (BuildOptions.bEnableUVQueries == false && BuildOptions.bEnableNormalsQueries == false)
+				{
+					TUniquePtr<FCompressedSplineStaticMeshSpatialWrapper> SplineSMWrapper = MakeUnique<FCompressedSplineStaticMeshSpatialWrapper>();
+					SplineSMWrapper->SourceContainer = MeshContainer;
+					SplineSMWrapper->StaticMesh = StaticMesh;
+					SplineSMWrapper->SplineMeshComponent = Cast<USplineMeshComponent>(ParentMeshes[0]->SourceComponent);
+					return SplineSMWrapper;
+				}
+				else
+				{
+					TUniquePtr<FSplineStaticMeshSourceDataSpatialWrapper> SplineSMWrapper = MakeUnique<FSplineStaticMeshSourceDataSpatialWrapper>();
+					SplineSMWrapper->SourceContainer = MeshContainer;
+					SplineSMWrapper->StaticMesh = StaticMesh;
+					SplineSMWrapper->SplineMeshComponent = Cast<USplineMeshComponent>(ParentMeshes[0]->SourceComponent);
+					return SplineSMWrapper;
+				}
 			}
 			else
 			{
-				TUniquePtr<FStaticMeshSourceDataSpatialWrapper> SMWrapper = MakeUnique<FStaticMeshSourceDataSpatialWrapper>();
-				SMWrapper->SourceContainer = MeshContainer;
-				SMWrapper->StaticMesh = StaticMesh;
-				return SMWrapper;
+				if (bUseRenderMeshes)
+				{
+					TUniquePtr<FStaticMeshRenderDataSpatialWrapper> SMWrapper = MakeUnique<FStaticMeshRenderDataSpatialWrapper>();
+					SMWrapper->SourceContainer = MeshContainer;
+					SMWrapper->StaticMesh = StaticMesh;
+					return SMWrapper;
+				}
+				else if (BuildOptions.bEnableUVQueries == false && BuildOptions.bEnableNormalsQueries == false)
+				{
+					TUniquePtr<FCompressedStaticMeshSpatialWrapper> SMWrapper = MakeUnique<FCompressedStaticMeshSpatialWrapper>();
+					SMWrapper->SourceContainer = MeshContainer;
+					SMWrapper->StaticMesh = StaticMesh;
+					return SMWrapper;
+				}
+				else
+				{
+					TUniquePtr<FStaticMeshSourceDataSpatialWrapper> SMWrapper = MakeUnique<FStaticMeshSourceDataSpatialWrapper>();
+					SMWrapper->SourceContainer = MeshContainer;
+					SMWrapper->StaticMesh = StaticMesh;
+					return SMWrapper;
+				}
 			}
 		}
 	}
@@ -1393,7 +1493,14 @@ static void CollectComponentChildMeshes(UActorComponent* Component, FActorAdapte
 			else
 			{
 				// base StaticMeshComponent
-				ChildMesh->ComponentType = EActorMeshComponentType::StaticMesh;
+				if (StaticMeshComponent->IsA<USplineMeshComponent>())
+				{
+					ChildMesh->ComponentType = EActorMeshComponentType::SplineStaticMesh;
+				}
+				else
+				{
+					ChildMesh->ComponentType = EActorMeshComponentType::StaticMesh;
+				}
 				ChildMesh->ComponentIndex = 0;
 				ChildMesh->WorldTransform.Append(StaticMeshComponent->GetComponentTransform());
 				ChildMesh->bIsNonUniformScaled = ChildMesh->WorldTransform.HasNonUniformScale();
@@ -1472,27 +1579,41 @@ void FMeshSceneAdapter::AddComponents(const TArray<UActorComponent*>& ComponentS
 
 void FMeshSceneAdapter::InitializeSpatialWrappers(const TArray<FActorAdapter*>& NewItemsToProcess)
 {
+	auto CreateNewWrapperInfo = [](TUniquePtr<FActorChildMesh>& ChildMesh)
+	{
+		TSharedPtr<FSpatialWrapperInfo> NewWrapperInfo = MakeShared<FSpatialWrapperInfo>();
+		NewWrapperInfo->SourceContainer = ChildMesh->MeshContainer;
+		NewWrapperInfo->SpatialWrapper = nullptr; // these are now initialized at beginning of Build() function
+		return NewWrapperInfo;
+	};
+
 	// Find IMeshSpatialWrapper for each child mesh component. If one does not exist
 	// and we have not seen the underlying unique mesh (eg StaticMesh Asset, etc, construct a new one
 	for (FActorAdapter* Actor : NewItemsToProcess)
 	{
 		for (TUniquePtr<FActorChildMesh>& ChildMesh : Actor->ChildMeshes)
 		{
-			void* MeshKey = ChildMesh->MeshContainer.GetMeshKey();
-			TSharedPtr<FSpatialWrapperInfo>* Found = SpatialAdapters.Find(MeshKey);
 			FSpatialWrapperInfo* MeshInfo;
-			if (Found == nullptr)
+			if (ChildMesh->ComponentType == EActorMeshComponentType::SplineStaticMesh)
 			{
-				TSharedPtr<FSpatialWrapperInfo> NewWrapperInfo = MakeShared<FSpatialWrapperInfo>();
-				SpatialAdapters.Add(MeshKey, NewWrapperInfo);
-
-				NewWrapperInfo->SourceContainer = ChildMesh->MeshContainer;
-				NewWrapperInfo->SpatialWrapper = nullptr;	// these are now initialized at beginning of Build() function
+				TSharedPtr<FSpatialWrapperInfo> NewWrapperInfo = CreateNewWrapperInfo(ChildMesh);
+				SplineMeshSpatialAdapters.Add(NewWrapperInfo);
 				MeshInfo = NewWrapperInfo.Get();
 			}
 			else
 			{
-				MeshInfo = (*Found).Get();
+				void* MeshKey = ChildMesh->MeshContainer.GetMeshKey();
+				TSharedPtr<FSpatialWrapperInfo>* Found = SpatialAdapters.Find(MeshKey);
+				if (Found == nullptr)
+				{
+					TSharedPtr<FSpatialWrapperInfo> NewWrapperInfo = CreateNewWrapperInfo(ChildMesh);
+					SpatialAdapters.Add(MeshKey, NewWrapperInfo);
+					MeshInfo = NewWrapperInfo.Get();
+				}
+				else
+				{
+					MeshInfo = (*Found).Get();
+				}
 			}
 
 			MeshInfo->ParentMeshes.Add(ChildMesh.Get());
@@ -1505,9 +1626,22 @@ void FMeshSceneAdapter::InitializeSpatialWrappers(const TArray<FActorAdapter*>&
 	}
 }
 
+void FMeshSceneAdapter::ForEachSpatialWrapper(TFunctionRef<void(TSharedPtr<FSpatialWrapperInfo>)> Func) const
+{
+	for (TPair<void*, TSharedPtr<FSpatialWrapperInfo>> Pair : SpatialAdapters)
+	{
+		Func(Pair.Value);
+	}
+
+	for (TSharedPtr<FSpatialWrapperInfo> Adapter : SplineMeshSpatialAdapters)
+	{
+		Func(Adapter);
+	}
+}
+
 bool FMeshSceneAdapter::IsValid() const
 {
-	return !SpatialAdapters.IsEmpty();
+	return !SpatialAdapters.IsEmpty() || !SplineMeshSpatialAdapters.IsEmpty();
 }
 
 void FMeshSceneAdapter::Build(const FMeshSceneAdapterBuildOptions& BuildOptions)
@@ -1515,15 +1649,15 @@ void FMeshSceneAdapter::Build(const FMeshSceneAdapterBuildOptions& BuildOptions)
 	TRACE_CPUPROFILER_EVENT_SCOPE(MeshScene_Build);
 
 	// Initialize all Spatial Wrappers using the default factory
-	for (TPair<void*, TSharedPtr<FSpatialWrapperInfo>> Pair : SpatialAdapters)
+	ForEachSpatialWrapper([&BuildOptions](TSharedPtr<FSpatialWrapperInfo> Adapter)
 	{
-		Pair.Value->SpatialWrapper = SpatialWrapperFactory(Pair.Value->SourceContainer, BuildOptions);
+		Adapter->SpatialWrapper = SpatialWrapperFactory(Adapter->SourceContainer, Adapter->ParentMeshes, BuildOptions);
 		// populate the MeshSpatial members of the FActorChildMeshes
-		for (FActorChildMesh* ParentMesh : Pair.Value->ParentMeshes)
+		for (FActorChildMesh* ParentMesh : Adapter->ParentMeshes)
 		{
-			ParentMesh->MeshSpatial = Pair.Value->SpatialWrapper.Get();
+			ParentMesh->MeshSpatial = Adapter->SpatialWrapper.Get();
 		}
-	}
+	});
 
 	if (BuildOptions.bThickenThinMeshes)
 	{
@@ -1532,10 +1666,10 @@ void FMeshSceneAdapter::Build(const FMeshSceneAdapterBuildOptions& BuildOptions)
 	else
 	{
 		TArray<FSpatialWrapperInfo*> ToBuild;
-		for (TPair<void*, TSharedPtr<FSpatialWrapperInfo>> Pair : SpatialAdapters)
+		ForEachSpatialWrapper([&ToBuild](TSharedPtr<FSpatialWrapperInfo> Adapter)
 		{
-			ToBuild.Add(Pair.Value.Get());
-		}
+			ToBuild.Add(Adapter.Get());
+		});
 
 		FCriticalSection ListsLock;
 
@@ -1665,10 +1799,10 @@ void FMeshSceneAdapter::Build_FullDecompose(const FMeshSceneAdapterBuildOptions&
 
 	// initial list of spatial wrappers that need to be built
 	TArray<FSpatialWrapperInfo*> ToBuild;
-	for (TPair<void*, TSharedPtr<FSpatialWrapperInfo>> Pair : SpatialAdapters)
+	ForEachSpatialWrapper([&ToBuild](TSharedPtr<FSpatialWrapperInfo> Adapter)
 	{
-		ToBuild.Add(Pair.Value.Get());
-	}
+		ToBuild.Add(Adapter.Get());
+	});
 
 	// Initialize the initial set of wrappers. Must do this here so that meshes are loaded and TriangleCount() below is valid
 	FMeshSceneAdapterBuildOptions TempBuildOptions = BuildOptions;
@@ -2139,11 +2273,11 @@ void FMeshSceneAdapter::GetGeometryStatistics(FStatistics& StatsOut)
 {
 	StatsOut.UniqueMeshCount = 0;
 	StatsOut.UniqueMeshTriangleCount = 0;
-	for (TPair<void*, TSharedPtr<FSpatialWrapperInfo>> Pair : SpatialAdapters)
+	ForEachSpatialWrapper([&StatsOut](TSharedPtr<FSpatialWrapperInfo> Adapter)
 	{
 		StatsOut.UniqueMeshCount++;
-		StatsOut.UniqueMeshTriangleCount += (int64)Pair.Value->SpatialWrapper->GetTriangleCount();
-	}
+		StatsOut.UniqueMeshTriangleCount += (int64)Adapter->SpatialWrapper->GetTriangleCount();
+	});
 
 	StatsOut.InstanceMeshCount = 0;
 	StatsOut.InstanceMeshTriangleCount = 0;
@@ -2252,9 +2386,9 @@ void FMeshSceneAdapter::BuildSpatialEvaluationCache()
 	TRACE_CPUPROFILER_EVENT_SCOPE(MeshScene_BuildSpatialEvaluationCache);
 
 	// validate that internals are correctly configured
-	for (TPair<void*, TSharedPtr<FSpatialWrapperInfo>> Pair : SpatialAdapters)
+	ForEachSpatialWrapper([](TSharedPtr<FSpatialWrapperInfo> Adapter)
 	{
-		FSpatialWrapperInfo& WrapperInfo = *Pair.Value;
+		FSpatialWrapperInfo& WrapperInfo = *Adapter;
 		IMeshSpatialWrapper* SpatialWrapper = WrapperInfo.SpatialWrapper.Get();
 		for (FActorChildMesh* MeshInstance : WrapperInfo.ParentMeshes)
 		{
@@ -2263,7 +2397,7 @@ void FMeshSceneAdapter::BuildSpatialEvaluationCache()
 				UE_LOG(LogGeometry, Warning, TEXT("FMeshSceneAdapter::BuildSpatialEvaluationCache: broken MeshSpatial link found!"));
 			}
 		}
-	}
+	});
 
 	// build list of unique meshes we need to evaluate for spatial queries
 	SortedSpatials.Reset();
@@ -2378,7 +2512,8 @@ void FMeshSceneAdapter::FastUpdateTransforms(bool bRebuildSpatialCache)
 					ChildMesh->bIsNonUniformScaled = ChildMesh->WorldTransform.HasNonUniformScale();
 				}
 			}
-			else if (ChildMesh->ComponentType == EActorMeshComponentType::StaticMesh)
+			else if ((ChildMesh->ComponentType == EActorMeshComponentType::StaticMesh) ||
+				(ChildMesh->ComponentType == EActorMeshComponentType::SplineStaticMesh))
 			{
 				UStaticMeshComponent* StaticMeshComponent = Cast<UStaticMeshComponent>(ChildMesh->SourceComponent);
 				ChildMesh->WorldTransform.Append(StaticMeshComponent->GetComponentTransform());
diff --git a/Engine/Plugins/Runtime/MeshModelingToolset/Source/ModelingComponents/Public/Scene/MeshSceneAdapter.h b/Engine/Plugins/Runtime/MeshModelingToolset/Source/ModelingComponents/Public/Scene/MeshSceneAdapter.h
--- a/Engine/Plugins/Runtime/MeshModelingToolset/Source/ModelingComponents/Public/Scene/MeshSceneAdapter.h
+++ b/Engine/Plugins/Runtime/MeshModelingToolset/Source/ModelingComponents/Public/Scene/MeshSceneAdapter.h
@@ -197,6 +197,7 @@ enum class EActorMeshComponentType
 	StaticMesh,
 	InstancedStaticMesh,
 	HierarchicalInstancedStaticMesh,
+	SplineStaticMesh,
 
 	InternallyGeneratedComponent,
 
@@ -439,8 +440,13 @@ protected:
 	// Unique set of spatial data structure query interfaces, one for each Mesh object, which is identified by void* pointer
 	TMap<void*, TSharedPtr<FSpatialWrapperInfo>> SpatialAdapters;
 
+	// Spatial data structure query interfaces for spline meshes
+	TArray<TSharedPtr<FSpatialWrapperInfo>> SplineMeshSpatialAdapters;
+
 	UE_API void InitializeSpatialWrappers(const TArray<FActorAdapter*>& NewItemsToProcess);
 
+	UE_API void ForEachSpatialWrapper(TFunctionRef<void(TSharedPtr<FSpatialWrapperInfo>)> Func) const;
+
 	bool bSceneIsAllSolids = false;
 
 	UE_API void UpdateActorBounds(FActorAdapter& Actor);
