Skeleton Modifier: fixed wrong dependency test

#rnx

diff --git a/Engine/Plugins/Runtime/MeshModelingToolset/Source/SkeletalMeshModifiers/Private/SkeletonModifier.cpp b/Engine/Plugins/Runtime/MeshModelingToolset/Source/SkeletalMeshModifiers/Private/SkeletonModifier.cpp
--- a/Engine/Plugins/Runtime/MeshModelingToolset/Source/SkeletalMeshModifiers/Private/SkeletonModifier.cpp
+++ b/Engine/Plugins/Runtime/MeshModelingToolset/Source/SkeletalMeshModifiers/Private/SkeletonModifier.cpp
@@ -681,8 +681,8 @@ bool USkeletonModifier::CommitSkeletonToSkeletalMesh()
 			USkeleton* Skeleton = SkeletalMesh->GetSkeleton();
 			if (SkeletonModifications == ESkeletonModificationType::SimpleMerge)
 			{
-				constexpr bool bSkeletalMeshReferencesOnly = false;
-				if (!HasAnyOtherReferences(bSkeletalMeshReferencesOnly))
+				constexpr bool bSkeletalMeshReferencesOnly = true;
+				if (HasAnyOtherReferences(bSkeletalMeshReferencesOnly))
 				{
 					return Skeleton->MergeAllBonesToBoneTree(SkeletalMesh.Get());
 				}
