ModelingTools: Replace the sculpt and sculpt max brushes with ones that only use the stamp normal for application, which typically breaks meshes much less. Rename these old brushes to "inflate (stroke)" and "inflate max". Also expose brush update method for overriding in vertex sculpt tool, and add some comments describing brushes.

diff --git a/Engine/Plugins/Editor/ModelingToolsEditorMode/Source/ModelingToolsEditorMode/Private/ModelingToolsEditorModeStyle.cpp b/Engine/Plugins/Editor/ModelingToolsEditorMode/Source/ModelingToolsEditorMode/Private/ModelingToolsEditorModeStyle.cpp
--- a/Engine/Plugins/Editor/ModelingToolsEditorMode/Source/ModelingToolsEditorMode/Private/ModelingToolsEditorModeStyle.cpp
+++ b/Engine/Plugins/Editor/ModelingToolsEditorMode/Source/ModelingToolsEditorMode/Private/ModelingToolsEditorModeStyle.cpp
@@ -456,6 +456,8 @@ void FModelingToolsEditorModeStyle::Initialize()
 		StyleSet->Set("BrushTypeIcons.SculptV", new IMAGE_BRUSH_SVG("Icons/BrushIcons/Brush_SculptV", Icon120));
 		StyleSet->Set("BrushTypeIcons.SculptMx", new IMAGE_BRUSH_SVG("Icons/BrushIcons/Brush_SculptMx", Icon120));
 		StyleSet->Set("BrushTypeIcons.Inflate", new IMAGE_BRUSH_SVG("Icons/BrushIcons/Brush_Inflate", Icon120));
+		StyleSet->Set("BrushTypeIcons.InflateSt", new IMAGE_BRUSH_SVG("Icons/BrushIcons/Brush_Inflate", Icon120));
+		StyleSet->Set("BrushTypeIcons.InflateMax", new IMAGE_BRUSH_SVG("Icons/BrushIcons/Brush_Inflate", Icon120)); // TODO: Need a new icon for this one
 		StyleSet->Set("BrushTypeIcons.Pinch", new IMAGE_BRUSH_SVG("Icons/BrushIcons/Brush_Pinch", Icon120));
 		StyleSet->Set("BrushTypeIcons.Flatten", new IMAGE_BRUSH_SVG("Icons/BrushIcons/Brush_Flatten", Icon120));
 		StyleSet->Set("BrushTypeIcons.PlaneN", new IMAGE_BRUSH_SVG("Icons/BrushIcons/Brush_PlaneN", Icon120));
diff --git a/Engine/Plugins/Runtime/MeshModelingToolset/Source/MeshModelingTools/Private/MeshVertexSculptTool.cpp b/Engine/Plugins/Runtime/MeshModelingToolset/Source/MeshModelingTools/Private/MeshVertexSculptTool.cpp
--- a/Engine/Plugins/Runtime/MeshModelingToolset/Source/MeshModelingTools/Private/MeshVertexSculptTool.cpp
+++ b/Engine/Plugins/Runtime/MeshModelingToolset/Source/MeshModelingTools/Private/MeshVertexSculptTool.cpp
@@ -358,7 +358,7 @@ void UMeshVertexSculptTool::RegisterBrushes()
 		NewObject<UMoveBrushOpProps>(this));
 
 	RegisterBrushType((int32)EMeshVertexSculptBrushType::Offset, LOCTEXT("Offset", "SculptN"),
-		MakeUnique<FLambdaMeshSculptBrushOpFactory>([this]() { return MakeUnique<FSurfaceSculptBrushOp>(BaseMeshQueryFunc); }),
+		MakeUnique<TBasicMeshSculptBrushOpFactory<UE::Geometry::FSingleNormalSculptBrushOp>>(),
 		NewObject<UStandardSculptBrushOpProps>(this));
 
 	RegisterBrushType((int32)EMeshVertexSculptBrushType::SculptView, LOCTEXT("SculptView", "SculptV"),
@@ -366,13 +366,21 @@ void UMeshVertexSculptTool::RegisterBrushes()
 		NewObject<UViewAlignedSculptBrushOpProps>(this));
 
 	RegisterBrushType((int32)EMeshVertexSculptBrushType::SculptMax, LOCTEXT("SculptMax", "SculptMx"),
-		MakeUnique<FLambdaMeshSculptBrushOpFactory>([this]() { return MakeUnique<FSurfaceMaxSculptBrushOp>(BaseMeshQueryFunc); }),
+		MakeUnique<FLambdaMeshSculptBrushOpFactory>([this]() { return MakeUnique<UE::Geometry::FSingleNormalMaxSculptBrushOp>(BaseMeshQueryFunc); }),
 		NewObject<USculptMaxBrushOpProps>(this));
 
 	RegisterBrushType((int32)EMeshVertexSculptBrushType::Inflate, LOCTEXT("Inflate", "Inflate"),
 		MakeUnique<TBasicMeshSculptBrushOpFactory<FInflateBrushOp>>(),
 		NewObject<UInflateBrushOpProps>(this));
 
+	RegisterBrushType((int32)EMeshVertexSculptBrushType::InflateStroke, LOCTEXT("InflateStroke", "InflateSt"),
+		MakeUnique<FLambdaMeshSculptBrushOpFactory>([this]() { return MakeUnique<FSurfaceSculptBrushOp>(BaseMeshQueryFunc); }),
+		NewObject<USculptMaxBrushOpProps>(this));
+
+	RegisterBrushType((int32)EMeshVertexSculptBrushType::InflateMax, LOCTEXT("InflateMax", "InflateMax"),
+		MakeUnique<FLambdaMeshSculptBrushOpFactory>([this]() { return MakeUnique<FSurfaceMaxSculptBrushOp>(BaseMeshQueryFunc); }),
+		NewObject<USculptMaxBrushOpProps>(this));
+
 	RegisterBrushType((int32)EMeshVertexSculptBrushType::Pinch, LOCTEXT("Pinch", "Pinch"),
 		MakeUnique<TBasicMeshSculptBrushOpFactory<FPinchBrushOp>>(),
 		NewObject<UPinchBrushOpProps>(this));
@@ -1952,6 +1960,7 @@ void UMeshVertexSculptTool::UpdateBrushType(int32 BrushID)
 		SculptProperties->bCanFreezeTarget = 
 			BrushID == (int32)EMeshVertexSculptBrushType::Offset
 			|| BrushID == (int32)EMeshVertexSculptBrushType::SculptMax || BrushID == (int32)EMeshVertexSculptBrushType::SculptView
+			|| BrushID == (int32)EMeshVertexSculptBrushType::InflateStroke || BrushID == (int32)EMeshVertexSculptBrushType::InflateMax
 			|| BrushID == (int32)EMeshVertexSculptBrushType::Pinch;
 	}
 
diff --git a/Engine/Plugins/Runtime/MeshModelingToolset/Source/MeshModelingTools/Public/MeshVertexSculptTool.h b/Engine/Plugins/Runtime/MeshModelingToolset/Source/MeshModelingTools/Public/MeshVertexSculptTool.h
--- a/Engine/Plugins/Runtime/MeshModelingToolset/Source/MeshModelingTools/Public/MeshVertexSculptTool.h
+++ b/Engine/Plugins/Runtime/MeshModelingToolset/Source/MeshModelingTools/Public/MeshVertexSculptTool.h
@@ -72,18 +72,24 @@ enum class EMeshVertexSculptBrushType : uint8
 	/** Smooth mesh vertices but only in direction of normal (Ctrl to invert) */
 	SmoothFill UMETA(DisplayName = "SmoothFill"),
 
-	/** Displace vertices along the average surface normal (Ctrl to invert) */
+	/** Displace vertices along the pre-stroke surface normal (Ctrl to invert) */
 	Offset UMETA(DisplayName = "Sculpt (Normal)"),
 
 	/** Displace vertices towards the camera viewpoint (Ctrl to invert) */
 	SculptView UMETA(DisplayName = "Sculpt (Viewpoint)"),
 
-	/** Displaces vertices along the average surface normal to a maximum height based on the brush size (Ctrl to invert) */
+	/** Displaces vertices along the pre-stroke surface normal to a maximum height based on the brush size (Ctrl to invert) */
 	SculptMax UMETA(DisplayName = "Sculpt Max"),
 
 	/** Displace vertices along their vertex normals */
 	Inflate UMETA(DisplayName = "Inflate"),
 
+	/** Displace vertices along their vertex normals, operating on pre-stroke mesh. */
+	InflateStroke UMETA(DisplayName = "Inflate (Stroke)"),
+
+	/** Displace vertices along their vertex normals to a maximum distance based on the brush size (Ctrl to invert) */
+	InflateMax UMETA(DisplayName = "Inflate Max"),
+
 	/** Scale Brush will inflate or pinch radially from the center of the brush */
 	ScaleKelvin UMETA(DisplayName = "Kelvin Scale"),
 
@@ -347,8 +353,8 @@ protected:
 
 	UE_DEPRECATED(5.6, "Vertex sculpt now uses an integer for a brush ID to allow custom brush registration")
 	UE_API void UpdateBrushType(EMeshVertexSculptBrushType BrushType);
-private:
-	UE_API void UpdateBrushType(int32 BrushID);
+
+	UE_API virtual void UpdateBrushType(int32 BrushID);
 protected:
 
 	TUniquePtr<UE::Geometry::FPolygroupSet> ActiveGroupSet;
diff --git a/Engine/Plugins/Runtime/MeshModelingToolset/Source/MeshModelingTools/Public/Sculpting/MeshBrushOpBase.h b/Engine/Plugins/Runtime/MeshModelingToolset/Source/MeshModelingTools/Public/Sculpting/MeshBrushOpBase.h
--- a/Engine/Plugins/Runtime/MeshModelingToolset/Source/MeshModelingTools/Public/Sculpting/MeshBrushOpBase.h
+++ b/Engine/Plugins/Runtime/MeshModelingToolset/Source/MeshModelingTools/Public/Sculpting/MeshBrushOpBase.h
@@ -14,7 +14,9 @@ using UE::Geometry::FMatrix3d;
 
 enum class ESculptBrushOpTargetType : uint8
 {
+	// The "live" mesh that is updated by each stamp
 	SculptMesh,
+	// The "base" mesh that is typically updated after every stroke
 	TargetMesh,
 	ActivePlane
 };
diff --git a/Engine/Plugins/Runtime/MeshModelingToolset/Source/MeshModelingTools/Public/Sculpting/MeshInflateBrushOps.h b/Engine/Plugins/Runtime/MeshModelingToolset/Source/MeshModelingTools/Public/Sculpting/MeshInflateBrushOps.h
--- a/Engine/Plugins/Runtime/MeshModelingToolset/Source/MeshModelingTools/Public/Sculpting/MeshInflateBrushOps.h
+++ b/Engine/Plugins/Runtime/MeshModelingToolset/Source/MeshModelingTools/Public/Sculpting/MeshInflateBrushOps.h
@@ -29,6 +29,8 @@ public:
 	virtual bool SupportsStrengthPressure() override { return true; }
 };
 
+// Like FSurfaceSculptBrushOp (moves vertices along individual vertex normals), but operates on the
+//  sculpt mesh (i.e. the live mesh) instead of the pre-stroke mesh.
 class FInflateBrushOp : public FMeshSculptBrushOp
 {
 
diff --git a/Engine/Plugins/Runtime/MeshModelingToolset/Source/MeshModelingTools/Public/Sculpting/MeshSculptBrushOps.h b/Engine/Plugins/Runtime/MeshModelingToolset/Source/MeshModelingTools/Public/Sculpting/MeshSculptBrushOps.h
--- a/Engine/Plugins/Runtime/MeshModelingToolset/Source/MeshModelingTools/Public/Sculpting/MeshSculptBrushOps.h
+++ b/Engine/Plugins/Runtime/MeshModelingToolset/Source/MeshModelingTools/Public/Sculpting/MeshSculptBrushOps.h
@@ -29,6 +29,69 @@ public:
 	virtual bool SupportsStrengthPressure() override { return true; }
 };
 
+namespace UE::Geometry
+{
+// Brush that sculpts all vertices in one direction, set by the normal at the brush position.
+//  Operates on pre-stroke mesh.
+class FSingleNormalSculptBrushOp : public FMeshSculptBrushOp
+{
+public:
+	double BrushSpeedTuning = 6.0;
+
+	virtual ESculptBrushOpTargetType GetBrushTargetType() const override
+	{
+		// Get our stamp information from the pre-stroke mesh
+		return ESculptBrushOpTargetType::TargetMesh;
+	}
+	virtual bool SupportsVariableSpacing() const override
+	{
+		return true;
+	}
+
+	virtual bool SupportsStrokeType(EMeshSculptStrokeType StrokeType) const override
+	{
+		switch (StrokeType)
+		{
+		case EMeshSculptStrokeType::Airbrush:
+		case EMeshSculptStrokeType::Dots:
+		case EMeshSculptStrokeType::Spacing:
+			return true;
+		default:
+			return false;
+		}
+	}
+
+	virtual bool UsesAlpha() const override
+	{
+		return true;
+	}
+
+	virtual void ApplyStamp(const FDynamicMesh3* Mesh, const FSculptBrushStamp& Stamp, const TArray<int32>& Vertices, TArray<FVector3d>& NewPositionsOut) override
+	{
+		double UsePower = Stamp.Direction * Stamp.Power * Stamp.Radius * Stamp.DeltaTime * BrushSpeedTuning;
+
+		bool bHaveAlpha = Stamp.HasAlpha();
+		FVector3d OffsetDirection = Stamp.LocalFrame.Z();
+
+		ParallelFor(Vertices.Num(), [this, Mesh, &Stamp, &Vertices, &NewPositionsOut, UsePower, bHaveAlpha, &OffsetDirection](int32 k)
+		{
+			int32 VertIdx = Vertices[k];
+			FVector3d OrigPos = Mesh->GetVertex(VertIdx);
+
+			double Alpha = (bHaveAlpha) ? Stamp.StampAlphaFunc(Stamp, OrigPos) : 1.0;
+
+			FVector3d MoveVec = UsePower * OffsetDirection;
+			double Falloff = GetFalloff().Evaluate(Stamp, OrigPos) * Alpha;
+			FVector3d NewPos = OrigPos + Falloff * MoveVec;
+			NewPositionsOut[k] = NewPos;
+		});
+	}
+};
+}// end namespace UE::Geometry
+//~ TODO: The rest should be moved to UE::Geometry namespace as well
+
+// Uses the vertex normal at the "base" mesh (i.e. pre-stroke mesh) vertex to move vertices.
+//  Basically just like the inflate brush, but using the pre-stroke mesh.
 class FSurfaceSculptBrushOp : public FMeshSculptBrushOp
 {
 public:
@@ -73,7 +136,6 @@ public:
 	virtual void ApplyStamp(const FDynamicMesh3* Mesh, const FSculptBrushStamp& Stamp, const TArray<int32>& Vertices, TArray<FVector3d>& NewPositionsOut) override
 	{
 		double UsePower = Stamp.Direction * Stamp.Power * Stamp.Radius * Stamp.DeltaTime * BrushSpeedTuning;
-		double MaxOffset = Stamp.Radius;
 
 		bool bHaveAlpha = Stamp.HasAlpha();
 
@@ -126,6 +188,8 @@ public:
 };
 
 
+// Starts with the base mesh, and moves all vertices towards the camera (by asking the owning tool to align
+//  the stamp normal to view, and using that value).
 class FViewAlignedSculptBrushOp : public FMeshSculptBrushOp
 {
 public:
@@ -175,7 +239,6 @@ public:
 		FVector3d StampNormal = Stamp.LocalFrame.Z();
 
 		double UsePower = Stamp.Direction * Stamp.Power * Stamp.Radius * Stamp.DeltaTime * BrushSpeedTuning;
-		double MaxOffset = Stamp.Radius;
 
 		bool bHaveAlpha = Stamp.HasAlpha();
 
@@ -240,7 +303,69 @@ public:
 	virtual bool SupportsStrengthPressure() override { return true; }
 };
 
+namespace UE::Geometry
+{
+// Brush that sculpts all vertices in one direction, set by the normal at the brush position, 
+//  and capped by a max value. Operates on pre-stroke mesh.
+class FSingleNormalMaxSculptBrushOp : public FSingleNormalSculptBrushOp
+{
+public:
+	typedef TUniqueFunction<bool(int32, const FVector3d&, double, FVector3d&, FVector3d&)> NearestQueryFuncType;
+
+	// To cap the movement distance, we need to know the original position on the pre-sculpt mesh
+	NearestQueryFuncType BaseMeshNearestQueryFunc;
+
+	FSingleNormalMaxSculptBrushOp(NearestQueryFuncType QueryFunc)
+	{
+		BaseMeshNearestQueryFunc = MoveTemp(QueryFunc);
+	}
+
+	virtual void ApplyStamp(const FDynamicMesh3* Mesh, const FSculptBrushStamp& Stamp, const TArray<int32>& Vertices, TArray<FVector3d>& NewPositionsOut) override
+	{
+		double UsePower = Stamp.Direction * Stamp.Power * Stamp.Radius * Stamp.DeltaTime * BrushSpeedTuning;
+
+		USculptMaxBrushOpProps* Props = GetPropertySetAs<USculptMaxBrushOpProps>();
+		double MaxOffset = (Props->bUseFixedHeight) ? Props->FixedHeight : (Props->MaxHeight * Stamp.Radius);
+
+		bool bHaveAlpha = Stamp.HasAlpha();
+		FVector3d OffsetDirection = Stamp.LocalFrame.Z();
+
+		ParallelFor(Vertices.Num(), [this, Mesh, &Stamp, &Vertices, &NewPositionsOut, UsePower, MaxOffset, bHaveAlpha, &OffsetDirection](int32 k)
+		{
+			int32 VertIdx = Vertices[k];
+			FVector3d OrigPos = Mesh->GetVertex(VertIdx);
+
+			FVector3d BasePos, BaseNormal;
+			bool bFoundBasePos = BaseMeshNearestQueryFunc(VertIdx, OrigPos, 4.0 * Stamp.Radius, BasePos, BaseNormal);
+			if (bFoundBasePos == false)
+			{
+				NewPositionsOut[k] = OrigPos;
+			}
+			else
+			{
+				double Alpha = (bHaveAlpha) ? Stamp.StampAlphaFunc(Stamp, OrigPos) : 1.0;
+
+				FVector3d MoveVec = UsePower * OffsetDirection;
+				double Falloff = GetFalloff().Evaluate(Stamp, OrigPos) * Alpha;
+				FVector3d NewPos = OrigPos + Falloff * MoveVec;
+
+				FVector3d DeltaPos = NewPos - BasePos;
+				if (DeltaPos.SquaredLength() > MaxOffset * MaxOffset)
+				{
+					UE::Geometry::Normalize(DeltaPos);
+					NewPos = BasePos + MaxOffset * DeltaPos;
+				}
+
+				NewPositionsOut[k] = NewPos;
+			}
+		});
+	}
+};
+}//end UE::Geometry
 
+// Like FSurfaceSculptBrushOp, but caps the movement to some maximum distance. In other words, it is
+//  an inflate (movement along individual vertex normals) based off of base (i.e. pre-stroke) mesh
+//  normals, but capped in offset distance from the pre-stroke mesh.
 class FSurfaceMaxSculptBrushOp : public FMeshSculptBrushOp
 {
 public:
