Skeleton Modifier: do not recreate the bone tree if the skeleton has bones not referenced in the skeletal mesh

#rnx

diff --git a/Engine/Plugins/Runtime/MeshModelingToolset/Source/SkeletalMeshModifiers/Private/SkeletonModifier.cpp b/Engine/Plugins/Runtime/MeshModelingToolset/Source/SkeletalMeshModifiers/Private/SkeletonModifier.cpp
--- a/Engine/Plugins/Runtime/MeshModelingToolset/Source/SkeletalMeshModifiers/Private/SkeletonModifier.cpp
+++ b/Engine/Plugins/Runtime/MeshModelingToolset/Source/SkeletalMeshModifiers/Private/SkeletonModifier.cpp
@@ -164,6 +164,62 @@ namespace USkeletonModifierLocals
 		const FReferenceSkeleton& ReferenceSkeleton;
 	};
 
+	/**
+	* FSkeletonOnlyBones: helps track bones that are in the USkeleton ref skeleton but not in those of the skeletal mesh.
+	* This is important because the bone tree must not be recreated in this situation, but rather merged.
+	*/
+	
+	struct FSkeletonOnlyBones
+	{
+		FSkeletonOnlyBones(const USkeleton* InSkeleton)
+			: Skeleton(InSkeleton)
+		{
+			const FReferenceSkeleton& RefSkeleton = Skeleton->GetReferenceSkeleton();
+ 
+			SkeletonOnlyBones.Reserve(RefSkeleton.GetRawBoneNum());
+			for (const FMeshBoneInfo& BoneInfo: RefSkeleton.GetRawRefBoneInfo())
+			{
+				SkeletonOnlyBones.Add(BoneInfo.Name);
+			}
+ 
+			for (const FVirtualBone& VB : InSkeleton->GetVirtualBones())
+			{
+				SkeletonOnlyBones.Remove(VB.SourceBoneName);
+				SkeletonOnlyBones.Remove(VB.TargetBoneName);
+			}
+		}
+		
+		void Remove(const FReferenceSkeleton& InRefSkeleton)
+		{
+			for (const FMeshBoneInfo& BoneInfo: InRefSkeleton.GetRawRefBoneInfo())
+			{
+				SkeletonOnlyBones.Remove(BoneInfo.Name);
+			}
+ 
+			const FReferenceSkeleton& RefSkeleton = Skeleton->GetReferenceSkeleton();
+			for (int32 BoneIndex = RefSkeleton.GetRawBoneNum() - 1; BoneIndex >= 0; --BoneIndex)
+			{
+				const FName BoneName = RefSkeleton.GetBoneName(BoneIndex);
+				if (!SkeletonOnlyBones.Contains(BoneName))
+				{
+					const int32 ParentIndex = RefSkeleton.GetParentIndex(BoneIndex);
+					if (ParentIndex != INDEX_NONE)
+					{
+						SkeletonOnlyBones.Remove(RefSkeleton.GetBoneName(ParentIndex));
+					}
+				}
+			}
+		}
+ 
+		int32 NumBones() const
+		{
+			return SkeletonOnlyBones.Num();
+		}
+		
+	private:
+		const USkeleton* Skeleton;
+		TArray<FName> SkeletonOnlyBones;
+	};
 }
 
 FTransform FMirrorOptions::MirrorTransform(const FTransform& InTransform) const
@@ -627,6 +683,17 @@ bool USkeletonModifier::CommitSkeletonToSkeletalMesh()
 			NewSkeleton->SetPreviewMesh(SkeletalMesh.Get());
 		}
 	}
+
+	const bool bHasSkeletonOnlyBones = [Modifications, SkeletonModifications, SkeletalMeshPtr = SkeletalMesh.Get()]()
+	{
+		if (SkeletonModifications == ESkeletonModificationType::SimpleMerge && Modifications == ESkeletalMeshModificationType::BonesAdded)
+		{
+			USkeletonModifierLocals::FSkeletonOnlyBones SkeletonOnlyBones(SkeletalMeshPtr->GetSkeleton());
+			SkeletonOnlyBones.Remove(SkeletalMeshPtr->GetRefSkeleton());
+			return SkeletonOnlyBones.NumBones() > 0;
+		}
+		return false;
+	}();
 	
 	const TArray<FMeshBoneInfo>& BoneInfos = ReferenceSkeleton->GetRawRefBoneInfo();
 	
@@ -676,13 +743,15 @@ bool USkeletonModifier::CommitSkeletonToSkeletalMesh()
 	{
 		NotifyFromSkeletonChanges();
 		
-		auto UpdateSkeleton = [this, &SkeletonModifications, &RetargetingModes]()
+		auto UpdateSkeleton = [this, &SkeletonModifications, &RetargetingModes, bHasSkeletonOnlyBones]()
 		{
 			USkeleton* Skeleton = SkeletalMesh->GetSkeleton();
 			if (SkeletonModifications == ESkeletonModificationType::SimpleMerge)
 			{
 				constexpr bool bSkeletalMeshReferencesOnly = true;
-				if (HasAnyOtherReferences(bSkeletalMeshReferencesOnly))
+				// merge the skeleton if there are several references or if the USkeleton's ref skeleton has bones that did not exist in the
+				// non-modified ref skeleton.
+				if (HasAnyOtherReferences(bSkeletalMeshReferencesOnly) || bHasSkeletonOnlyBones)
 				{
 					return Skeleton->MergeAllBonesToBoneTree(SkeletalMesh.Get());
 				}
