ModelingTools: Add separate options for erase and smooth in paint tool, support pressure sensitivity, and allow subclasses to prefer dynamic mesh provider interface.

diff --git a/Engine/Plugins/Runtime/MeshModelingToolset/Source/MeshModelingTools/Private/MeshAttributePaintTool.cpp b/Engine/Plugins/Runtime/MeshModelingToolset/Source/MeshModelingTools/Private/MeshAttributePaintTool.cpp
--- a/Engine/Plugins/Runtime/MeshModelingToolset/Source/MeshModelingTools/Private/MeshAttributePaintTool.cpp
+++ b/Engine/Plugins/Runtime/MeshModelingToolset/Source/MeshModelingTools/Private/MeshAttributePaintTool.cpp
@@ -244,15 +244,15 @@ void UMeshAttributePaintEditActions::PostAction(EMeshAttributePaintToolActions A
 
 UMeshSurfacePointTool* UMeshAttributePaintToolBuilder::CreateNewTool(const FToolBuilderState& SceneState) const
 {
-	UMeshAttributePaintTool* SelectionTool = NewObject<UMeshAttributePaintTool>(SceneState.ToolManager);
-	SelectionTool->SetWorld(SceneState.World);
+	UMeshAttributePaintTool* NewTool = NewObject<UMeshAttributePaintTool>(SceneState.ToolManager);
+	NewTool->SetWorld(SceneState.World);
 
 	if (ColorMapFactory)
 	{
-		SelectionTool->SetColorMap(ColorMapFactory());
+		NewTool->SetColorMap(ColorMapFactory());
 	}
 
-	return SelectionTool;
+	return NewTool;
 }
 
 // TODO: May want a base class for brush tools with selection?
@@ -311,6 +311,10 @@ void UMeshAttributePaintTool::Setup()
 	AddToolPropertySource(BrushActionProps);
 
 	UDynamicMeshBrushTool::Setup();
+	if (ensure(BrushProperties))
+	{
+		BrushProperties->bToolSupportsPressureSensitivity = true;
+	}
 
 	// hide strength and falloff
 	BrushProperties->RestoreProperties(this);
@@ -382,11 +386,11 @@ void UMeshAttributePaintTool::Setup()
 
 	ColorMapper = MakeUnique<FFloatAttributeColorMapper>();
 
-	if (Cast<IMeshDescriptionProvider>(Target) && Cast<IMeshDescriptionCommitter>(Target))
+	if (bPreferMeshDescription && Cast<IMeshDescriptionProvider>(Target) && Cast<IMeshDescriptionCommitter>(Target))
 	{
-		EditedMesh = MakeUnique<FMeshDescription>();
-		*EditedMesh = *UE::ToolTarget::GetMeshDescription(Target);
+		EditedMesh = MakeUnique<FMeshDescription>();			
 		AttributeSource = MakeUnique<FMeshDescriptionVertexAttributeSource>(EditedMesh.Get());
+		*EditedMesh = *UE::ToolTarget::GetMeshDescription(Target);
 	}
 	else
 	{
@@ -397,15 +401,6 @@ void UMeshAttributePaintTool::Setup()
 		AttributeSource = MakeUnique< FDynamicMeshVertexAttributeSource>(EditedDynamicMesh.Get());
 	}
 
-	AttribProps->Initialize(AttributeSource->GetAttributeList(), true);
-
-	if (AttribProps->Attributes.Num() == 0)
-	{
-		GetToolManager()->DisplayMessage(
-			LOCTEXT("StartAttribPaintFailed", "No Float attributes exist for this mesh. Use the Attribute Editor to create one."),
-			EToolMessageLevel::UserWarning);
-	}
-
 	InitializeAttributes();
 	if (AttribProps->Attributes.Num() > 0)
 	{
@@ -423,8 +418,6 @@ void UMeshAttributePaintTool::Setup()
 	});
 }
 
-
-
 void UMeshAttributePaintTool::RegisterActions(FInteractiveToolActionSet& ActionSet)
 {
 	UDynamicMeshBrushTool::RegisterActions(ActionSet);
@@ -592,6 +585,20 @@ void UMeshAttributePaintTool::CalculateVertexROI(const FBrushStampData& Stamp, T
 
 void UMeshAttributePaintTool::InitializeAttributes()
 {
+	AttribProps->Initialize(AttributeSource->GetAttributeList(), true);
+
+	if (AttribProps->Attributes.Num() == 0)
+	{
+		GetToolManager()->DisplayMessage(
+			LOCTEXT("StartAttribPaintFailed", "No Float attributes exist for this mesh. Use the Attribute Editor to create one."),
+			EToolMessageLevel::UserWarning);
+	}
+	else
+	{
+		// Clear any existing error messages.
+		GetToolManager()->DisplayMessage(FText(), EToolMessageLevel::UserWarning);
+	}
+
 	AttributeBufferCount = AttributeSource->GetAttributeElementNum();
 	TArray<FName> AttributeNames = AttributeSource->GetAttributeList();
 
@@ -699,7 +706,9 @@ void UMeshAttributePaintTool::ApplyStamp(const FBrushStampData& Stamp)
 	FStampActionData ActionData;
 	CalculateVertexROI(Stamp, ActionData.ROIVertices);
 
-	if (BrushActionProps->BrushAction == EBrushActionMode::FloodFill)
+	switch (BrushActionProps->BrushAction)
+	{
+	case EBrushActionMode::FloodFill:
 	{
 		const bool bEmptyROIVertices = (ActionData.ROIVertices.Num() == 0);
 		if (bEmptyROIVertices) // append a vertex from the hit triangle to start the flood fill.
@@ -713,10 +722,19 @@ void UMeshAttributePaintTool::ApplyStamp(const FBrushStampData& Stamp)
 			}
 		}
 		ApplyStamp_FloodFill(Stamp, ActionData);
+		break;
 	}
-	else
-	{
-		ApplyStamp_Paint(Stamp, ActionData);
+	case EBrushActionMode::Paint:
+		ApplyStamp_Paint(Stamp, ActionData, bInSmoothStroke, bInRemoveStroke);
+		break;
+	case EBrushActionMode::Smooth:
+		ApplyStamp_Paint(Stamp, ActionData, true, bInRemoveStroke);
+		break;
+	case EBrushActionMode::Erase:
+		ApplyStamp_Paint(Stamp, ActionData, bInSmoothStroke, !bInRemoveStroke);
+		break;
+	default:
+		ensure(false);
 	}
 
 
@@ -750,8 +768,12 @@ void UMeshAttributePaintTool::ApplyStamp(const FBrushStampData& Stamp)
 	PreviewMesh->NotifyDeferredEditCompleted(UPreviewMesh::ERenderUpdateMode::FastUpdate, EMeshRenderAttributeFlags::VertexColors, false);
 }
 
-
 void UMeshAttributePaintTool::ApplyStamp_Paint(const FBrushStampData& Stamp, FStampActionData& ActionData)
+{
+	ApplyStamp_Paint(Stamp, ActionData, bInSmoothStroke, bInRemoveStroke);
+}
+
+void UMeshAttributePaintTool::ApplyStamp_Paint(const FBrushStampData& Stamp, FStampActionData& ActionData, bool bSmooth, bool bInvert)
 {
 	FTransform3d Transform(Cast<ISceneComponentBackedTarget>(Target)->GetWorldTransform());
 	FVector3d StampPosLocal = Transform.InverseTransformPosition((FVector3d)Stamp.WorldPosition);
@@ -764,7 +786,7 @@ void UMeshAttributePaintTool::ApplyStamp_Paint(const FBrushStampData& Stamp, FSt
 
 	const FDynamicMesh3* CurrentMesh = PreviewMesh->GetMesh();
 	FNonManifoldMappingSupport NonManifoldMappingSupport(*CurrentMesh);
-	if (bInSmoothStroke)
+	if (bSmooth)
 	{
 		float SmoothSpeed = 0.25f;
 
@@ -793,7 +815,6 @@ void UMeshAttributePaintTool::ApplyStamp_Paint(const FBrushStampData& Stamp, FSt
 	}
 	else
 	{
-		bool bInvert = bInRemoveStroke;
 		float Sign = (bInvert) ? -1.0f : 1.0f;
 		float UseStrength = Sign * BrushProperties->BrushStrength * CurrentValueRange.Length();
 		for (int32 k = 0; k < NumVertices; ++k)
@@ -1022,6 +1043,5 @@ bool UMeshAttributePaintTool::ShouldFilterTriangles() const
 
 
 
-
 #undef LOCTEXT_NAMESPACE
 
diff --git a/Engine/Plugins/Runtime/MeshModelingToolset/Source/MeshModelingTools/Public/MeshAttributePaintTool.h b/Engine/Plugins/Runtime/MeshModelingToolset/Source/MeshModelingTools/Public/MeshAttributePaintTool.h
--- a/Engine/Plugins/Runtime/MeshModelingToolset/Source/MeshModelingTools/Public/MeshAttributePaintTool.h
+++ b/Engine/Plugins/Runtime/MeshModelingToolset/Source/MeshModelingTools/Public/MeshAttributePaintTool.h
@@ -106,8 +106,14 @@ protected:
 UENUM()
 enum class EBrushActionMode
 {
+	// Clicking adds to the weight, Ctrl subtracts, Shift smooths
 	Paint,
-	FloodFill
+	// Clicking sets the weight for the entire mesh component. Ctrl unsets the weight for the entire component, Shift is unused.
+	FloodFill,
+	// Clicking subtracts from the weight, Ctrl adds, Shift smooths
+	Erase,
+	// Clicking smooths values, Ctrl and Shift are unused
+	Smooth
 };
 
 
@@ -133,7 +139,6 @@ public:
 
 	UPROPERTY(EditAnywhere, Category = Attribute)
 	EBrushActionMode BrushAction = EBrushActionMode::Paint;
-
 };
 
 
@@ -298,6 +303,10 @@ protected:
 
 	// Note: This is used if the target supports mesh description read/write; otherwise EditedDynamicMesh is used
 	TUniquePtr<FMeshDescription> EditedMesh;
+	// This is used if the target doesn't support mesh description read/write; otherwise EditedMesh is used
+	TUniquePtr<FDynamicMesh3> EditedDynamicMesh;
+	// Can be set by subclasses to determine which to use when both mesh description and dynamic mesh are options.
+	bool bPreferMeshDescription = true;
 
 	UE_API double CalculateBrushFalloff(double Distance);
 	UE::Geometry::TDynamicVerticesOctree3<FDynamicMesh3> VerticesOctree;
@@ -343,9 +352,8 @@ private:
 	TSet<int32> SelectionTids;
 	TSet<int32> SelectionVids;
 	UE_API bool ShouldFilterTriangles() const;
-
-	// Note: This is used if the target doesn't support mesh description read/write; otherwise EditedMesh is used
-	TUniquePtr<FDynamicMesh3> EditedDynamicMesh;
+	// bInvert is ignored if bSmooth is true
+	void ApplyStamp_Paint(const FBrushStampData& Stamp, FStampActionData& ActionData, bool bSmooth, bool bInvert);
 };
 
 
