Add utility for

diff --git a/Engine/Plugins/Runtime/MeshModelingToolset/Source/ModelingComponents/Private/SplineUtil.cpp b/Engine/Plugins/Runtime/MeshModelingToolset/Source/ModelingComponents/Private/SplineUtil.cpp
--- a/Engine/Plugins/Runtime/MeshModelingToolset/Source/ModelingComponents/Private/SplineUtil.cpp
+++ b/Engine/Plugins/Runtime/MeshModelingToolset/Source/ModelingComponents/Private/SplineUtil.cpp
@@ -7,6 +7,7 @@
 #include "SceneView.h"
 #include "Styling/StyleColors.h"
 #include "ToolContextInterfaces.h" // IToolsContextRenderAPI
+#include "DynamicMesh/DynamicMeshAABBTree3.h"
 
 UE::Geometry::SplineUtil::FDrawSplineSettings::FDrawSplineSettings()
 	: RegularColor(FStyleColors::White.GetSpecifiedColor().ToFColor(true))
@@ -174,3 +175,146 @@ void UE::Geometry::SplineUtil::DrawSpline(const USplineComponent& SplineComp, IT
 		OldKeyScale = NewKeyScale;
 	}
 }
+
+
+void UE::Geometry::SplineUtil::ProjectSplineToSurface(FInterpCurveVector& OutputSpline, 
+	const FInterpCurveVector& InputSpline,
+	const UE::Geometry::FDynamicMeshAABBTree3& SurfaceAABBTree, 
+	const FTransform& SplineTransform, 
+	const FTransform& MeshTransform, 
+	double RelativeErrorThreshold,
+	int32 MaxNewPoints)
+{
+	if (InputSpline.Points.Num() == 0)
+	{
+		return;
+	}
+
+	// Find the closest point on the surface given by SurfaceAABBTree to the spline point given by the parameter SplinePointT
+	auto GetNearestPointOnSurface = [&SplineTransform, &MeshTransform, &SurfaceAABBTree](const FInterpCurveVector& Spline, double SplinePointT, int32* NearestTriangleID = nullptr) -> FVector3d
+	{
+		const FVector3d SplinePoint = Spline.Eval(SplinePointT, FVector3d(0));
+		const FVector3d WorldSplinePoint = SplineTransform.TransformPosition(SplinePoint);
+		const FVector3d MeshSpaceSplinePoint = MeshTransform.InverseTransformPosition(WorldSplinePoint);
+
+		const FVector3d MeshSpaceNearestPoint = SurfaceAABBTree.FindNearestPoint(MeshSpaceSplinePoint);
+		const FVector3d WorldNearestPoint = MeshTransform.TransformPosition(MeshSpaceNearestPoint);
+		const FVector3d SplineNearestPoint = SplineTransform.InverseTransformPosition(WorldNearestPoint);
+
+		if (NearestTriangleID != nullptr)
+		{
+			double DistToTriSq;
+			*NearestTriangleID = SurfaceAABBTree.FindNearestTriangle(MeshSpaceSplinePoint, DistToTriSq);
+		}
+
+		return SplineNearestPoint;
+	};
+
+	// Find the closest point on the surface and return a new spline point there
+	auto ProjectToSurface = [&SplineTransform, &MeshTransform, &SurfaceAABBTree, GetNearestPointOnSurface](const FInterpCurveVector& Spline, double SplinePointT, EInterpCurveMode InterpMode) -> FInterpCurvePointVector
+	{
+		int32 NearestTriangleID;
+		const FVector3d SplineNearestPoint = GetNearestPointOnSurface(Spline, SplinePointT, &NearestTriangleID);
+
+		FInterpCurvePointVector ProjectedPoint;
+		ProjectedPoint.InVal = SplinePointT;
+		ProjectedPoint.OutVal = SplineNearestPoint;
+		ProjectedPoint.InterpMode = InterpMode;
+
+		// Spline point tangent = spline derivative projected onto the tangent plane of the surface
+		const FVector3d SplineDerivative = Spline.EvalDerivative(SplinePointT, FVector3d(0));
+		const FVector3d TriangleNormal = SurfaceAABBTree.GetMesh()->GetTriNormal(NearestTriangleID);
+		const FVector3d Tangent = SplineDerivative - SplineDerivative.Dot(TriangleNormal) * TriangleNormal;
+		ProjectedPoint.ArriveTangent = 0.5 * Tangent;
+		ProjectedPoint.LeaveTangent = 0.5 * Tangent;
+
+		return ProjectedPoint;
+	};
+
+
+	// Sample the full spline and find the furthest point from the surface
+	auto ApproximationError = [&SplineTransform, &MeshTransform, &SurfaceAABBTree, GetNearestPointOnSurface](const FInterpCurveVector& Spline, int32 NumSamplesPerSegment = 10) -> TPair<double, double>
+	{
+		NumSamplesPerSegment = FMath::Max(NumSamplesPerSegment, 1);
+
+		double MaxError = -1.0;
+		double MaxErrorT = -1.0;
+
+		for (int32 PointIndex = 0; PointIndex < Spline.Points.Num() - 1; ++PointIndex)
+		{
+			const int32 NextPointIndex = PointIndex + 1;
+
+			const double MinT = Spline.Points[PointIndex].InVal;
+			const double MaxT = Spline.Points[NextPointIndex].InVal;
+			const double DeltaT = (MaxT - MinT) / (NumSamplesPerSegment + 1);
+
+			for (int32 SampleIndex = 0; SampleIndex < NumSamplesPerSegment; ++SampleIndex)
+			{
+				const double T = MinT + DeltaT * (SampleIndex + 1);
+				const FVector SplinePoint = Spline.Eval(T, FVector(0));
+				const FVector3d SplineNearestPoint = GetNearestPointOnSurface(Spline, T);
+				const double Dist = FVector3d::Distance(SplineNearestPoint, SplinePoint);
+
+				if (MaxError < Dist)
+				{
+					MaxErrorT = T;
+					MaxError = Dist;
+				}
+			}
+		}
+
+		return { MaxErrorT, MaxError };
+	};
+
+
+	// Initialize OutputSpline using intput points
+
+	const double MaxT = InputSpline.Points.Last().InVal;
+	double AvgSegmentLength = 0.0;
+	for (int32 PointIndex = 0; PointIndex < InputSpline.Points.Num(); ++PointIndex)
+	{
+		const FInterpCurvePointVector& Point = InputSpline.Points[PointIndex];
+		const FInterpCurvePointVector ProjectedPoint = ProjectToSurface(InputSpline, Point.InVal, Point.InterpMode);
+		OutputSpline.Points.Add(ProjectedPoint);
+
+		if (PointIndex > 0)
+		{
+			AvgSegmentLength += FVector3d::Distance(InputSpline.Points[PointIndex].OutVal, InputSpline.Points[PointIndex - 1].OutVal);
+		}
+	}
+
+	if (InputSpline.Points.Num() > 1)
+	{
+		AvgSegmentLength /= (double)(InputSpline.Points.Num() - 1);
+	}
+
+	// Now iteratively insert points at the max-error sample along the spline
+
+	TPair<double, double> TAndError = ApproximationError(OutputSpline);
+	double MaxErrorT = TAndError.Key;
+	double MaxError = TAndError.Value;
+
+	const double ErrorThreshold = 0.1 * AvgSegmentLength;
+
+	for (int32 InsertIteration = 0; InsertIteration < MaxNewPoints; ++InsertIteration)
+	{
+		if (MaxError < ErrorThreshold)
+		{
+			break;
+		}
+
+		const FInterpCurvePointVector ProjectedPoint = ProjectToSurface(OutputSpline, MaxErrorT, EInterpCurveMode::CIM_CurveAuto);
+
+		const int32 NewPointIndex = OutputSpline.AddPoint(MaxErrorT, ProjectedPoint.OutVal);
+		OutputSpline.Points[NewPointIndex] = ProjectedPoint;
+
+		TAndError = ApproximationError(OutputSpline);
+		MaxErrorT = TAndError.Key;
+		MaxError = TAndError.Value;
+
+		if (MaxErrorT <= 0.0 || MaxErrorT >= OutputSpline.Points.Last().InVal)
+		{
+			break;
+		}
+	}
+}
diff --git a/Engine/Plugins/Runtime/MeshModelingToolset/Source/ModelingComponents/Public/SplineUtil.h b/Engine/Plugins/Runtime/MeshModelingToolset/Source/ModelingComponents/Public/SplineUtil.h
--- a/Engine/Plugins/Runtime/MeshModelingToolset/Source/ModelingComponents/Public/SplineUtil.h
+++ b/Engine/Plugins/Runtime/MeshModelingToolset/Source/ModelingComponents/Public/SplineUtil.h
@@ -12,6 +12,11 @@ class IToolsContextRenderAPI;
 namespace UE {
 namespace Geometry {
 
+class FDynamicMesh3;
+template<typename T> class TMeshAABBTree3;
+typedef TMeshAABBTree3<FDynamicMesh3> FDynamicMeshAABBTree3;
+
+
 namespace SplineUtil {
 
 	// Used for DrawSpline()
@@ -33,6 +38,18 @@ namespace SplineUtil {
 
 	MODELINGCOMPONENTS_API void DrawSpline(const USplineComponent& SplineComp, IToolsContextRenderAPI& RenderAPI, const FDrawSplineSettings& Settings);
 
+	/**
+	* Iteratively subdivide and project a spline so that it lies on the given surface.
+	* NOTE: Point tangents for newly added points are computed using the spline derivative and local surface tangent plane. Rotations and Scales are not handled by this function.
+	*/
+	MODELINGCOMPONENTS_API void ProjectSplineToSurface(FInterpCurveVector& OutputSpline,
+		const FInterpCurveVector& InputSpline,
+		const UE::Geometry::FDynamicMeshAABBTree3& SurfaceAABBTree, 
+		const FTransform& SplineTransform, 
+		const FTransform& MeshTransform, 
+		double RelativeErrorThreshold = 0.1,
+		int32 MaxNewPoints = 100);
+
 }
 }
 }
