ModelingTools: Added symmetry support to the Paint Vertex tool.

diff --git a/Engine/Plugins/Runtime/MeshModelingToolset/Source/MeshModelingTools/Private/MeshVertexPaintTool.cpp b/Engine/Plugins/Runtime/MeshModelingToolset/Source/MeshModelingTools/Private/MeshVertexPaintTool.cpp
--- a/Engine/Plugins/Runtime/MeshModelingToolset/Source/MeshModelingTools/Private/MeshVertexPaintTool.cpp
+++ b/Engine/Plugins/Runtime/MeshModelingToolset/Source/MeshModelingTools/Private/MeshVertexPaintTool.cpp
@@ -11,6 +11,7 @@
 #include "Async/Async.h"
 #include "SceneView.h"
 #include "ToolSetupUtil.h"
+#include "ToolDataVisualizer.h"
 #include "ModelingToolTargetUtil.h"
 #include "MeshWeights.h"
 #include "DynamicMesh/MeshNormals.h"
@@ -23,6 +24,7 @@
 #include "Selections/MeshConnectedComponents.h"
 #include "Selections/MeshFaceSelection.h"
 #include "Selections/MeshVertexSelection.h"
+#include "Parameterization/MeshPlanarSymmetry.h"
 #include "Polygroups/PolygroupUtil.h"
 #include "Polygon2.h"
 #include "DynamicMesh/Operations/SplitAttributeWelder.h"
@@ -233,6 +235,10 @@ void UMeshVertexPaintTool::Setup()
 		AABBTree.SetMesh(Mesh, true);
 	});
 
+	TFuture<void> InitializeSymmetry = Async(VertexPaintToolAsyncExecTarget, [Mesh, this]()
+	{
+		TryToInitializeSymmetry();
+	});
 
 	// initialize render decomposition
 	TUniquePtr<FMeshRenderDecomposition> Decomp = MakeUnique<FMeshRenderDecomposition>();
@@ -251,7 +257,7 @@ void UMeshVertexPaintTool::Setup()
 	PolyLassoMechanic->OnDrawPolyLassoFinished.AddUObject(this, &UMeshVertexPaintTool::OnPolyLassoFinished);
 
 	PolygroupLayerProperties = NewObject<UPolygroupLayersProperties>(this);
-	PolygroupLayerProperties->RestoreProperties(this, TEXT("MeshVertexPaintTool"));
+	PolygroupLayerProperties->RestoreProperties(this, GetPropertyCacheIdentifier());
 	PolygroupLayerProperties->InitializeGroupLayers(GetSculptMesh());
 	PolygroupLayerProperties->WatchProperty(PolygroupLayerProperties->ActiveGroupLayer, [&](FName) { OnSelectedGroupLayerChanged(); });
 	UpdateActiveGroupLayer();
@@ -284,6 +290,11 @@ void UMeshVertexPaintTool::Setup()
 	UMeshSculptToolBase::BrushProperties->BrushSize.bToolSupportsPressureSensitivity = true;
 	CalculateBrushRadius();
 
+	SymmetryProperties = NewObject<UMeshSymmetryProperties>(this);
+	SymmetryProperties->RestoreProperties(this, GetPropertyCacheIdentifier());
+	SymmetryProperties->bSymmetryCanBeEnabled = false;
+	AddToolPropertySource(SymmetryProperties);
+
 	PaintBrushOpProperties = NewObject<UVertexColorPaintBrushOpProps>(this);
 	RegisterBrushType((int32)EMeshVertexPaintBrushType::Paint, LOCTEXT("Paint", "Paint"),
 		MakeUnique<FLambdaMeshSculptBrushOpFactory>([this]() { return MakeUnique<UE::Geometry::FVertexPaintBrushOp>(); }),
@@ -339,7 +350,7 @@ void UMeshVertexPaintTool::Setup()
 	if (ensure(MeshElementsDisplay->Settings))
 	{
 		MeshElementsDisplay->Settings->bShowColorSeams = false;		// default color seams to false
-		MeshElementsDisplay->Settings->RestoreProperties(this, TEXT("MeshVertexPaintTool"));
+		MeshElementsDisplay->Settings->RestoreProperties(this, GetPropertyCacheIdentifier());
 		AddToolPropertySource(MeshElementsDisplay->Settings);
 	}
 	MeshElementsDisplay->SetMeshAccessFunction([this](UMeshElementsVisualizer::ProcessDynamicMeshFunc ProcessFunc) 
@@ -416,6 +427,14 @@ void UMeshVertexPaintTool::Setup()
 	PrecomputeFuture.Wait();
 	OctreeFuture.Wait();
 	AABBTreeFuture.Wait();
+	InitializeSymmetry.Wait();
+
+	// update symmetry state based on validity, and then update internal apply-symmetry state
+	SymmetryProperties->bSymmetryCanBeEnabled = bMeshSymmetryIsValid;
+	bApplySymmetry = bMeshSymmetryIsValid && SymmetryProperties->bEnableSymmetry;
+
+	SymmetryProperties->WatchProperty(SymmetryProperties->bEnableSymmetry,
+		[this](bool bNewValue) { bApplySymmetry = bMeshSymmetryIsValid && bNewValue; });
 }
 
 void UMeshVertexPaintTool::Shutdown(EToolShutdownType ShutdownType)
@@ -427,13 +446,14 @@ void UMeshVertexPaintTool::Shutdown(EToolShutdownType ShutdownType)
 
 	if (ensure(MeshElementsDisplay->Settings))
 	{
-		MeshElementsDisplay->Settings->SaveProperties(this, TEXT("MeshVertexPaintTool"));
+		MeshElementsDisplay->Settings->SaveProperties(this, GetPropertyCacheIdentifier());
 	}
 	MeshElementsDisplay->Disconnect();
 
 	BasicProperties->SaveProperties(this);
 	FilterProperties->SaveProperties(this);
-	PolygroupLayerProperties->SaveProperties(this, TEXT("MeshVertexPaintTool"));
+	PolygroupLayerProperties->SaveProperties(this, GetPropertyCacheIdentifier());
+	SymmetryProperties->SaveProperties(this, GetPropertyCacheIdentifier());
 
 	if (PreviewMeshActor != nullptr)
 	{
@@ -577,6 +597,14 @@ void UMeshVertexPaintTool::OnBeginStroke(const FRay& WorldRay)
 	LastStamp.Power = GetActiveBrushStrength();
 	LastStamp.TimeStamp = FDateTime::Now();
 
+	// If applying symmetry, make sure the stamp is on the "positive" side. 
+	if (bApplySymmetry)
+	{
+		LastStamp.LocalFrame = Symmetry->GetPositiveSideFrame(LastStamp.LocalFrame);
+		LastStamp.WorldFrame = LastStamp.LocalFrame;
+		LastStamp.WorldFrame.Transform(CurTargetTransform);
+	}
+
 	FSculptBrushOptions SculptOptions;
 	//SculptOptions.bPreserveUVFlow = false; // FilterProperties->bPreserveUVFlow;
 	SculptOptions.ConstantReferencePlane = GetCurrentStrokeReferencePlane();
@@ -637,12 +665,12 @@ void UMeshVertexPaintTool::UpdateROI(const FSculptBrushStamp& BrushStamp)
 	//  unselected triangles in ApplyVisibilityFilter if needed. But seems odd not to.
 	bool bShouldFilterTrianglesBySelection = ShouldFilterTrianglesBySelection();
 
-	// With Lazy Brush, GetBrushTriangleID() may not necessarily return a triangle that contains BrushPos,
+	// With Symmetry or Lazy Brush, GetBrushTriangleID() may not necessarily return a triangle that contains BrushPos,
 	// so in that case we will try to find the triangle containing BrushPos. Note that this may not
 	// be correct in some geometric cases...potentially we could check if the initial CenterTID 
 	// contains the BrushPos point first
 	int32 CenterTID = GetBrushTriangleID();
-	if (BrushProperties->Lazyness > 0)
+	if (bApplySymmetry || BrushProperties->Lazyness > 0)
 	{
 		double NearDistSqr;
 		IMeshSpatial::FQueryOptions QueryOptions;
@@ -682,9 +710,6 @@ void UMeshVertexPaintTool::UpdateROI(const FSculptBrushStamp& BrushStamp)
 		return false;
 	};
 
-	bool bFill = (BasicProperties->SubToolType == EMeshVertexPaintInteractionType::Fill);
-	bool bGroupFill = (BasicProperties->SubToolType == EMeshVertexPaintInteractionType::GroupFill);
-
 	if (bVolumetric)
 	{
 		Octree.RangeQuery(BrushBox,
@@ -721,14 +746,14 @@ void UMeshVertexPaintTool::UpdateROI(const FSculptBrushStamp& BrushStamp)
 		}
 	}
 
-	if (bFill)
+	auto FillROI = [this, bShouldFilterTrianglesBySelection, CheckEdgeCriteria, Mesh](TSet<int32>& ResultROI)
 	{
 		TArray<int32> StartROI;
-		for (int32 tid : TriangleROI)
+		for (int32 tid : ResultROI)
 		{
 			StartROI.Add(tid);
 		}
-		FMeshConnectedComponents::GrowToConnectedTriangles(Mesh, StartROI, TriangleROI, &TempROIBuffer,
+		FMeshConnectedComponents::GrowToConnectedTriangles(Mesh, StartROI, ResultROI, &TempROIBuffer,
 														   [&](int t1, int t2)
 		{
 			if (bShouldFilterTrianglesBySelection && (!SelectionTids.Contains(t1) || !SelectionTids.Contains(t2)))
@@ -737,17 +762,18 @@ void UMeshVertexPaintTool::UpdateROI(const FSculptBrushStamp& BrushStamp)
 			}
 			return CheckEdgeCriteria(t1, t2);
 		});
-	}
-	else if (bGroupFill)
+	};
+
+	auto GroupFillROI = [this, bShouldFilterTrianglesBySelection, CheckEdgeCriteria, Mesh](TSet<int>& ResultROI)
 	{
 		TArray<int32> StartROI;
 		TSet<int32> FillGroups;
-		for (int32 tid : TriangleROI)
+		for (int32 tid : ResultROI)
 		{
 			StartROI.Add(tid);
 			FillGroups.Add(ActiveGroupSet->GetGroup(tid));
 		}
-		FMeshConnectedComponents::GrowToConnectedTriangles(Mesh, StartROI, TriangleROI, &TempROIBuffer,
+		FMeshConnectedComponents::GrowToConnectedTriangles(Mesh, StartROI, ResultROI, &TempROIBuffer,
 			[&](int t1, int t2)
 		{
 			if (bShouldFilterTrianglesBySelection && (!SelectionTids.Contains(t1) || !SelectionTids.Contains(t2)))
@@ -756,16 +782,67 @@ void UMeshVertexPaintTool::UpdateROI(const FSculptBrushStamp& BrushStamp)
 			}
 			return (FillGroups.Contains(ActiveGroupSet->GetGroup(t2)));
 		});
-	}
+	};
 
+	auto PostProcessROI = [this, &FillROI, &GroupFillROI](TSet<int>& ResultROI)
+	{
+		if (BasicProperties->SubToolType == EMeshVertexPaintInteractionType::Fill)
+		{
+			FillROI(ResultROI);
+		}
+		else if (BasicProperties->SubToolType == EMeshVertexPaintInteractionType::GroupFill)
+		{
+			GroupFillROI(ResultROI);
+		}
 
-	// apply visibility filter
-	if (FilterProperties->VisibilityFilter != EMeshVertexPaintVisibilityType::None)
+		// apply visibility filter
+		if (FilterProperties->VisibilityFilter != EMeshVertexPaintVisibilityType::None)
+		{
+			TArray<int32> ResultBuffer;
+			ApplyVisibilityFilter(ResultROI, TempROIBuffer, ResultBuffer);
+		}
+	};
+
+	if (bApplySymmetry)
 	{
-		TArray<int32> ResultBuffer;
-		ApplyVisibilityFilter(TriangleROI, TempROIBuffer, ResultBuffer);
+		TempVertexSet.Reset();
+		for (int32 tid : TriangleROI)
+		{
+			FIndex3i Tri = Mesh->GetTriangle(tid);
+			TempVertexSet.Add(Tri.A);  TempVertexSet.Add(Tri.B);  TempVertexSet.Add(Tri.C);
+		}
+		VertexROI.SetNum(0, EAllowShrinking::No);
+		BufferUtil::AppendElements(VertexROI, TempVertexSet);
+		
+		// Find symmetric Vertex ROI. This will overlap with VertexROI in many cases.
+		SymmetricVertexROI.Reset();
+		Symmetry->GetMirrorVertexROI(VertexROI, SymmetricVertexROI, true);
+		SymmetricTriangleROI.Reset();
+		// expand the Triangle ROI to include the symmetric vertex one-rings
+		for (int32 VertexID : SymmetricVertexROI)
+		{
+			if (Mesh->IsVertex(VertexID))
+			{
+				Mesh->EnumerateVertexTriangles(VertexID, [&](int32 tid)
+				{
+					SymmetricTriangleROI.Add(tid);
+				});
+			}
+		}
+
+		PostProcessROI(SymmetricTriangleROI);
+
+		// construct ROI triangle and group buffers
+		SymmetricROITriangleBuffer.Reserve(SymmetricTriangleROI.Num());
+		SymmetricROITriangleBuffer.SetNum(0, EAllowShrinking::No);
+		for (int32 tid : SymmetricTriangleROI)
+		{
+			SymmetricROITriangleBuffer.Add(tid);
+		}
 	}
 
+	PostProcessROI(TriangleROI);
+
 	// construct ROI vertex set
 	TempVertexSet.Reset();
 	for (int32 tid : TriangleROI)
@@ -816,6 +893,14 @@ bool UMeshVertexPaintTool::UpdateStampPosition(const FRay& WorldRay)
 	CurrentStamp.LocalFrame = GetBrushFrameLocal();
 	CurrentStamp.Power = GetActiveBrushStrength();
 
+	// If applying symmetry, make sure the stamp is on the "positive" side. 
+	if (bApplySymmetry)
+	{
+		CurrentStamp.LocalFrame = Symmetry->GetPositiveSideFrame(CurrentStamp.LocalFrame);
+		CurrentStamp.WorldFrame = CurrentStamp.LocalFrame;
+		CurrentStamp.WorldFrame.Transform(CurTargetTransform);
+	}
+
 	CurrentStamp.PrevLocalFrame = LastStamp.LocalFrame;
 	CurrentStamp.PrevWorldFrame = LastStamp.WorldFrame;
 
@@ -854,6 +939,37 @@ bool UMeshVertexPaintTool::ApplyStamp()
 
 	bool bUpdated = SyncMeshWithColorBuffer(Mesh);
 
+	// TODO: Unify the symmetry and main stamp passes.
+	if (bApplySymmetry)
+	{
+		// convert triangle ROI to element IDs, populating ROIElementSet, ROIElementBuffer, ROIColorBuffer
+		{
+			TRACE_CPUPROFILER_EVENT_SCOPE(MeshVPaint_ApplyStamp_ElementROI);
+			InitializeElementROIFromTriangleROI(SymmetricROITriangleBuffer, true);
+		}
+
+		// Temporarily update CurrentStamp to point to the negative side
+		FFrame3d PosLocalFrame = CurrentStamp.LocalFrame;
+		FFrame3d PosWorldFrame = CurrentStamp.WorldFrame;
+		CurrentStamp.LocalFrame.Origin = Symmetry->GetMirroredPosition(CurrentStamp.LocalFrame.Origin);
+		CurrentStamp.LocalFrame.Rotation = Symmetry->GetMirroredOrientation(CurrentStamp.LocalFrame.Rotation);
+		CurrentStamp.WorldFrame = CurrentStamp.LocalFrame;
+		CurrentStamp.WorldFrame.Transform(CurTargetTransform);
+
+		{
+			TRACE_CPUPROFILER_EVENT_SCOPE(MeshVPaint_ApplyStamp_ApplyStampToVertexColors);
+			VertexColorBrushOp->ApplyStampToVertexColors(Mesh, GetActiveColorOverlay(),
+				StrokeInitialColorBuffer, StrokeAccumColorBuffer,
+				CurrentStamp, ROIElementBuffer, ROIColorBuffer);
+		}
+
+		// Restore CurrentStamp
+		CurrentStamp.LocalFrame = PosLocalFrame;
+		CurrentStamp.WorldFrame = PosWorldFrame;
+
+		bUpdated = SyncMeshWithColorBuffer(Mesh) || bUpdated;
+	}
+
 	LastStamp = CurrentStamp;
 	LastStamp.TimeStamp = FDateTime::Now();
 	return bUpdated;
@@ -1294,6 +1410,21 @@ void UMeshVertexPaintTool::DrawHUD(FCanvas* Canvas, IToolsContextRenderAPI* Rend
 
 }
 
+void UMeshVertexPaintTool::Render(IToolsContextRenderAPI* RenderAPI)
+{
+	UMeshSculptToolBase::Render(RenderAPI);
+
+	// draw a dot for the symmetric brush stamp position
+	if (bApplySymmetry)
+	{
+		FToolDataVisualizer Visualizer;
+		Visualizer.BeginFrame(RenderAPI);
+		FVector3d MirrorPoint = CurTargetTransform.TransformPosition(
+			Symmetry->GetMirroredPosition(HoverStamp.LocalFrame.Origin));
+		Visualizer.DrawPoint(MirrorPoint, FLinearColor(1.0, 0.1, 0.1, 1), 5.0f, false);
+		Visualizer.EndFrame();
+	}
+}
 
 void UMeshVertexPaintTool::OnTick(float DeltaTime)
 {
@@ -2178,8 +2309,33 @@ void UMeshVertexPaintTool::ApplyAction(EMeshVertexPaintToolActions ActionType)
 	}
 }
 
+FString UMeshVertexPaintTool::GetPropertyCacheIdentifier() const
+{
+	return TEXT("MeshVertexPaintTool");
+}
 
 
+void UMeshVertexPaintTool::TryToInitializeSymmetry()
+{
+	// Attempt to find symmetry, favoring the X axis, then Y axis, if a single symmetry plane is not immediately found
+	// Uses local mesh surface (angle sum, normal) to help disambiguate final matches, but does not require exact topology matches across the plane
+
+	FAxisAlignedBox3d Bounds = GetSculptMesh()->GetBounds(true);
+
+	TArray<FVector3d> PreferAxes;
+	PreferAxes.Add(this->InitialTargetTransform.GetRotation().AxisX());
+	PreferAxes.Add(this->InitialTargetTransform.GetRotation().AxisY());
+
+	FMeshPlanarSymmetry FindSymmetry;
+	FFrame3d FoundPlane;
+	if (FindSymmetry.FindPlaneAndInitialize(GetSculptMesh(), Bounds, FoundPlane, PreferAxes))
+	{
+		Symmetry = MakePimpl<FMeshPlanarSymmetry>();
+		*Symmetry = MoveTemp(FindSymmetry);
+		bMeshSymmetryIsValid = true;
+	}
+}
+
 
 
 #undef LOCTEXT_NAMESPACE
diff --git a/Engine/Plugins/Runtime/MeshModelingToolset/Source/MeshModelingTools/Public/MeshVertexPaintTool.h b/Engine/Plugins/Runtime/MeshModelingToolset/Source/MeshModelingTools/Public/MeshVertexPaintTool.h
--- a/Engine/Plugins/Runtime/MeshModelingToolset/Source/MeshModelingTools/Public/MeshVertexPaintTool.h
+++ b/Engine/Plugins/Runtime/MeshModelingToolset/Source/MeshModelingTools/Public/MeshVertexPaintTool.h
@@ -17,6 +17,9 @@
 #include "Sculpting/MeshSculptToolBase.h"
 #include "Sculpting/MeshBrushOpBase.h"
 
+// TODO: Refactor UMeshSymmetryProperties out of MeshVertexSculptTool
+#include "MeshVertexSculptTool.h" // For UMeshSymmetryProperties
+
 #include "DynamicMesh/DynamicMeshAABBTree3.h"
 #include "DynamicMesh/DynamicMeshOctree3.h"
 #include "DynamicMesh/MeshNormals.h"
@@ -503,6 +506,7 @@ public:
 
 	UE_API virtual void OnTick(float DeltaTime) override;
 	UE_API virtual void DrawHUD(FCanvas* Canvas, IToolsContextRenderAPI* RenderAPI) override;
+	UE_API virtual void Render(IToolsContextRenderAPI* RenderAPI) override;
 
 	virtual bool HasCancel() const override { return true; }
 	virtual bool HasAccept() const override { return true; }
@@ -534,6 +538,8 @@ public:
 	UPROPERTY()
 	TObjectPtr<UVertexPaintBrushFilterProperties> FilterProperties;
 
+	UPROPERTY()
+	TObjectPtr<UMeshSymmetryProperties> SymmetryProperties;
 
 private:
 	// This will be of type UVertexPaintBrushOpProps, we keep a ref so we can change active color on pick
@@ -572,7 +578,7 @@ protected:
 	UE_API virtual TUniquePtr<FMeshSculptBrushOp>& GetActiveBrushOp();
 	// end UMeshSculptToolBase API
 
-
+	UE_API virtual FString GetPropertyCacheIdentifier() const;
 
 	//
 	// Action support
@@ -725,6 +731,15 @@ private:
 	TOptional<UE::Geometry::FGeometrySelection> GeometrySelection;
 	TSet<int32> SelectionTids;
 	UE_API bool ShouldFilterTrianglesBySelection() const;
+
+	TPimplPtr<UE::Geometry::FMeshPlanarSymmetry> Symmetry;
+	bool bMeshSymmetryIsValid = false;
+	UE_API void TryToInitializeSymmetry();
+
+	bool bApplySymmetry = false;
+	TArray<int> SymmetricVertexROI;
+	TSet<int> SymmetricTriangleROI;
+	TArray<int32> SymmetricROITriangleBuffer;
 };
 
 
