ModelingTools: Fixed PaintVertexTool bug with ApplyVisibilityFilter updating TriangleROI directly.

diff --git a/Engine/Plugins/Runtime/MeshModelingToolset/Source/MeshModelingTools/Private/MeshVertexPaintTool.cpp b/Engine/Plugins/Runtime/MeshModelingToolset/Source/MeshModelingTools/Private/MeshVertexPaintTool.cpp
--- a/Engine/Plugins/Runtime/MeshModelingToolset/Source/MeshModelingTools/Private/MeshVertexPaintTool.cpp
+++ b/Engine/Plugins/Runtime/MeshModelingToolset/Source/MeshModelingTools/Private/MeshVertexPaintTool.cpp
@@ -1219,7 +1219,7 @@ void UMeshVertexPaintTool::ApplyVisibilityFilter(TSet<int32>& Triangles, TArray<
 	Triangles.Reset();
 	for (int32 tid : OutputBuffer)
 	{
-		TriangleROI.Add(tid);
+		Triangles.Add(tid);
 	}
 }
 
