Add 'vertex color sculpt material' to have both sculpt shading and vertex colors + expose this as an option in the attribute paint tool

#virtualized

diff --git a/Engine/Plugins/Runtime/MeshModelingToolset/Source/MeshModelingTools/Private/MeshAttributePaintTool.cpp b/Engine/Plugins/Runtime/MeshModelingToolset/Source/MeshModelingTools/Private/MeshAttributePaintTool.cpp
--- a/Engine/Plugins/Runtime/MeshModelingToolset/Source/MeshModelingTools/Private/MeshAttributePaintTool.cpp
+++ b/Engine/Plugins/Runtime/MeshModelingToolset/Source/MeshModelingTools/Private/MeshAttributePaintTool.cpp
@@ -1,6 +1,8 @@
 // Copyright Epic Games, Inc. All Rights Reserved.
 
 #include "MeshAttributePaintTool.h"
+#include "Engine/Engine.h" // for GEngine->VertexColorViewModeMaterial_ColorOnly
+#include "Materials/Material.h"
 #include "InteractiveToolManager.h"
 #include "ToolBuilderUtil.h"
 #include "Drawing/MeshDebugDrawing.h"
@@ -310,6 +312,10 @@ void UMeshAttributePaintTool::Setup()
 	BrushActionProps->RestoreProperties(this);
 	AddToolPropertySource(BrushActionProps);
 
+	ViewProperties = NewObject<UMeshAttributePaintToolVisualizationProperties>(this);
+	ViewProperties->RestoreProperties(this);
+	AddToolPropertySource(ViewProperties);
+
 	UDynamicMeshBrushTool::Setup();
 	if (ensure(BrushProperties))
 	{
@@ -371,11 +377,7 @@ void UMeshAttributePaintTool::Setup()
 	// build octree
 	VerticesOctree.Initialize(PreviewMesh->GetMesh(), true);
 
-	UMaterialInterface* VtxColorMaterial = ToolSetupUtil::GetVertexColorMaterial(GetToolManager());
-	if (VtxColorMaterial != nullptr)
-	{
-		PreviewMesh->SetOverrideRenderMaterial(VtxColorMaterial);
-	}
+	UpdateMaterialMode(ViewProperties->MaterialMode);
 
 	RecalculateBrushRadius();
 
@@ -416,6 +418,11 @@ void UMeshAttributePaintTool::Setup()
 	{
 		PreviewMesh->SetSecondaryBuffersVisibility(!BrushActionProps->bIsolateGeometrySelection);
 	});
+
+	ViewProperties->WatchProperty(ViewProperties->MaterialMode,
+		[this](EMeshAttributePaintMaterialMode NewMode) { UpdateMaterialMode(NewMode); });
+	ViewProperties->WatchProperty(ViewProperties->bFlatShading,
+		[this](bool bNewValue) { UpdateFlatShadingSetting(bNewValue); });
 }
 
 void UMeshAttributePaintTool::RegisterActions(FInteractiveToolActionSet& ActionSet)
@@ -768,6 +775,41 @@ void UMeshAttributePaintTool::ApplyStamp(const FBrushStampData& Stamp)
 	PreviewMesh->NotifyDeferredEditCompleted(UPreviewMesh::ERenderUpdateMode::FastUpdate, EMeshRenderAttributeFlags::VertexColors, false);
 }
 
+void UMeshAttributePaintTool::UpdateMaterialMode(EMeshAttributePaintMaterialMode MaterialMode)
+{
+	if (MaterialMode == EMeshAttributePaintMaterialMode::Shaded)
+	{
+		UMaterialInterface* SculptMaterial = ToolSetupUtil::GetVertexColorSculptMaterial(GetToolManager());
+		if (SculptMaterial != nullptr)
+		{
+			ActiveOverrideMaterial = UMaterialInstanceDynamic::Create(SculptMaterial, this);
+		}
+	}
+	else if (MaterialMode == EMeshAttributePaintMaterialMode::ColorOnly && 
+			GEngine && GEngine->VertexColorViewModeMaterial_ColorOnly)
+	{
+		// use the engine color-only material, which has no shading
+		ActiveOverrideMaterial = nullptr;
+		PreviewMesh->SetOverrideRenderMaterial(GEngine->VertexColorViewModeMaterial_ColorOnly);
+	}
+
+	if (ActiveOverrideMaterial != nullptr)
+	{
+		PreviewMesh->SetOverrideRenderMaterial(ActiveOverrideMaterial);
+		ActiveOverrideMaterial->SetScalarParameterValue(TEXT("FlatShading"), (ViewProperties->bFlatShading) ? 1.0f : 0.0f);
+	}
+
+	PreviewMesh->SetShadowsEnabled(false);
+}
+
+void UMeshAttributePaintTool::UpdateFlatShadingSetting(bool bNewValue)
+{
+	if (ActiveOverrideMaterial != nullptr)
+	{
+		ActiveOverrideMaterial->SetScalarParameterValue(TEXT("FlatShading"), (bNewValue) ? 1.0f : 0.0f);
+	}
+}
+
 void UMeshAttributePaintTool::ApplyStamp_Paint(const FBrushStampData& Stamp, FStampActionData& ActionData)
 {
 	ApplyStamp_Paint(Stamp, ActionData, bInSmoothStroke, bInRemoveStroke);
@@ -929,6 +971,7 @@ void UMeshAttributePaintTool::OnShutdown(EToolShutdownType ShutdownType)
 {
 	BrushProperties->SaveProperties(this);
 	BrushActionProps->SaveProperties(this);
+	ViewProperties->SaveProperties(this);
 
 	StoreCurrentAttribute();
 
diff --git a/Engine/Plugins/Runtime/MeshModelingToolset/Source/MeshModelingTools/Public/MeshAttributePaintTool.h b/Engine/Plugins/Runtime/MeshModelingToolset/Source/MeshModelingTools/Public/MeshAttributePaintTool.h
--- a/Engine/Plugins/Runtime/MeshModelingToolset/Source/MeshModelingTools/Public/MeshAttributePaintTool.h
+++ b/Engine/Plugins/Runtime/MeshModelingToolset/Source/MeshModelingTools/Public/MeshAttributePaintTool.h
@@ -103,6 +103,14 @@ protected:
 
 
 
+UENUM()
+enum class EMeshAttributePaintMaterialMode
+{
+	Shaded,
+	ColorOnly
+};
+
+
 UENUM()
 enum class EBrushActionMode
 {
@@ -142,7 +150,18 @@ public:
 };
 
 
+UCLASS(MinimalAPI)
+class UMeshAttributePaintToolVisualizationProperties : public UInteractiveToolPropertySet
+{
+	GENERATED_BODY()
+public:
+
+	UPROPERTY(EditAnywhere, Category = Rendering)
+	EMeshAttributePaintMaterialMode MaterialMode = EMeshAttributePaintMaterialMode::Shaded;
 
+	UPROPERTY(EditAnywhere, Category = Rendering, meta = (EditCondition ="MaterialMode==EMeshAttributePaintMaterialMode::Shaded"))
+	bool bFlatShading = true;
+};
 
 
 UCLASS(MinimalAPI)
@@ -287,11 +306,17 @@ protected:
 	UPROPERTY()
 	TObjectPtr<UMeshAttributePaintToolProperties> AttribProps;
 
+	UPROPERTY()
+	TObjectPtr<UMeshAttributePaintToolVisualizationProperties> ViewProperties;
+
 	TValueWatcher<int32> SelectedAttributeWatcher;
 
 	//UPROPERTY()
 	//UMeshAttributePaintEditActions* AttributeEditActions;
 
+	UPROPERTY()
+	TObjectPtr<UMaterialInstanceDynamic> ActiveOverrideMaterial;
+
 protected:
 	UWorld* TargetWorld;
 
@@ -354,6 +379,9 @@ private:
 	UE_API bool ShouldFilterTriangles() const;
 	// bInvert is ignored if bSmooth is true
 	void ApplyStamp_Paint(const FBrushStampData& Stamp, FStampActionData& ActionData, bool bSmooth, bool bInvert);
+
+	void UpdateMaterialMode(EMeshAttributePaintMaterialMode MaterialMode);
+	void UpdateFlatShadingSetting(bool bNewValue);
 };
 
 
diff --git a/Engine/Plugins/Runtime/MeshModelingToolset/Source/ModelingComponents/Private/ToolSetupUtil.cpp b/Engine/Plugins/Runtime/MeshModelingToolset/Source/ModelingComponents/Private/ToolSetupUtil.cpp
--- a/Engine/Plugins/Runtime/MeshModelingToolset/Source/ModelingComponents/Private/ToolSetupUtil.cpp
+++ b/Engine/Plugins/Runtime/MeshModelingToolset/Source/ModelingComponents/Private/ToolSetupUtil.cpp
@@ -50,6 +50,16 @@ UMaterialInstanceDynamic* ToolSetupUtil::GetVertexColorMaterial(UInteractiveTool
 	return nullptr;
 }
 
+UMaterialInterface* ToolSetupUtil::GetVertexColorSculptMaterial(UInteractiveToolManager* ToolManager)
+{
+	UMaterialInterface* Material = LoadObject<UMaterial>(nullptr, TEXT("/MeshModelingToolsetExp/Materials/MeshVertexColorSculptMaterial"));
+	if (Material == nullptr && ToolManager != nullptr)
+	{
+		return ToolManager->GetContextQueriesAPI()->GetStandardMaterial(EStandardToolContextMaterials::VertexColorMaterial);
+	}
+	return Material;
+}
+
 
 UMaterialInterface* ToolSetupUtil::GetDefaultWorkingMaterial(UInteractiveToolManager* ToolManager)
 {
diff --git a/Engine/Plugins/Runtime/MeshModelingToolset/Source/ModelingComponents/Public/ToolSetupUtil.h b/Engine/Plugins/Runtime/MeshModelingToolset/Source/ModelingComponents/Public/ToolSetupUtil.h
--- a/Engine/Plugins/Runtime/MeshModelingToolset/Source/ModelingComponents/Public/ToolSetupUtil.h
+++ b/Engine/Plugins/Runtime/MeshModelingToolset/Source/ModelingComponents/Public/ToolSetupUtil.h
@@ -37,6 +37,11 @@ namespace ToolSetupUtil
 	 */
 	MODELINGCOMPONENTS_API UMaterialInstanceDynamic* GetVertexColorMaterial(UInteractiveToolManager* ToolManager, bool bTwoSided = false);
 
+	/**
+	 * @return a vertex color material with sculpt-material shading
+	 */
+	MODELINGCOMPONENTS_API UMaterialInterface* GetVertexColorSculptMaterial(UInteractiveToolManager* ToolManager);
+
 
 	/**
 	 * @return default material to use for "Working"/In-Progress animations
