Move materials from MeshModelingToolsetExp to EngineMaterials

Moves 2 materials:
1) LineSetComponentMaterial
2) PointSetComponentMaterial

Also leaves object redirectors in the original location for backwards compatibility.

This change is a prerequisite for a refactor of the spline component's scene proxy to accelerate rendering of very large splines.

#virtualized

diff --git a/Engine/Plugins/Runtime/MeshModelingToolset/Source/ModelingComponents/Private/ToolSetupUtil.cpp b/Engine/Plugins/Runtime/MeshModelingToolset/Source/ModelingComponents/Private/ToolSetupUtil.cpp
--- a/Engine/Plugins/Runtime/MeshModelingToolset/Source/ModelingComponents/Private/ToolSetupUtil.cpp
+++ b/Engine/Plugins/Runtime/MeshModelingToolset/Source/ModelingComponents/Private/ToolSetupUtil.cpp
@@ -335,7 +335,7 @@ UMaterialInterface* ToolSetupUtil::GetDefaultEditVolumeMaterial()
 UMaterialInterface* ToolSetupUtil::GetDefaultPointComponentMaterial(UInteractiveToolManager* ToolManager, bool bDepthTested)
 {
 	UMaterialInterface* Material = (bDepthTested) ?
-		LoadObject<UMaterial>(nullptr, TEXT("/MeshModelingToolsetExp/Materials/PointSetComponentMaterial"))
+		LoadObject<UMaterial>(nullptr, TEXT("/Engine/EngineMaterials/PointSetComponentMaterial"))
 		: LoadObject<UMaterial>(nullptr, TEXT("/MeshModelingToolsetExp/Materials/PointSetOverlaidComponentMaterial"));
 	if (Material == nullptr && ToolManager != nullptr)
 	{
@@ -359,7 +359,7 @@ UMaterialInterface* ToolSetupUtil::GetRoundPointComponentMaterial(UInteractiveTo
 UMaterialInterface* ToolSetupUtil::GetDefaultLineComponentMaterial(UInteractiveToolManager* ToolManager, bool bDepthTested)
 {
 	UMaterialInterface* Material = (bDepthTested) ?
-		LoadObject<UMaterial>(nullptr, TEXT("/MeshModelingToolsetExp/Materials/LineSetComponentMaterial"))
+		LoadObject<UMaterial>(nullptr, TEXT("/Engine/EngineMaterials/LineSetComponentMaterial"))
 		: LoadObject<UMaterial>(nullptr, TEXT("/MeshModelingToolsetExp/Materials/LineSetOverlaidComponentMaterial"));
 	if (Material == nullptr && ToolManager != nullptr)
 	{
