Tools : Add volumetric brush stamp indicator and update the skin weight paint tool to use it

diff --git a/Engine/Plugins/Runtime/MeshModelingToolset/Source/MeshModelingToolsEditorOnly/Private/SkeletalMesh/SkinWeightsPaintTool.cpp b/Engine/Plugins/Runtime/MeshModelingToolset/Source/MeshModelingToolsEditorOnly/Private/SkeletalMesh/SkinWeightsPaintTool.cpp
--- a/Engine/Plugins/Runtime/MeshModelingToolset/Source/MeshModelingToolsEditorOnly/Private/SkeletalMesh/SkinWeightsPaintTool.cpp
+++ b/Engine/Plugins/Runtime/MeshModelingToolset/Source/MeshModelingToolsEditorOnly/Private/SkeletalMesh/SkinWeightsPaintTool.cpp
@@ -22,7 +22,9 @@
 #include "Preferences/PersonaOptions.h"
 #include "PreviewProfileController.h"
 #include "BaseGizmos/BrushStampIndicator.h"
+#include "BaseGizmos/VolumetricBrushStampIndicator.h"
 
+#include "InteractiveGizmoManager.h"
 #include "InteractiveToolManager.h"
 #include "ToolSetupUtil.h"
 #include "ToolTargetManager.h"
@@ -1981,6 +1983,55 @@ void USkinWeightsPaintTool::Render(IToolsContextRenderAPI* RenderAPI)
 	}
 }
 
+namespace SkinWeightsPaintTool::Private
+{
+	const FString BaseVolumetricBrushIndicatorGizmoType = TEXT("VolumetricBrushIndicatorGizmoType");
+}
+
+void USkinWeightsPaintTool::SetupBrushStampIndicator()
+{
+	if (!BrushStampIndicator)
+	{
+		// register and spawn brush indicator gizmo
+		GetToolManager()->GetPairedGizmoManager()->RegisterGizmoType(SkinWeightsPaintTool::Private::BaseVolumetricBrushIndicatorGizmoType, NewObject<UVolumetricBrushStampIndicatorBuilder>());
+		BrushStampIndicator = GetToolManager()->GetPairedGizmoManager()->CreateGizmo<UVolumetricBrushStampIndicator>(SkinWeightsPaintTool::Private::BaseVolumetricBrushIndicatorGizmoType, FString(), this);
+		if (UVolumetricBrushStampIndicator* VolumetricBrushIndicator = Cast<UVolumetricBrushStampIndicator>(BrushStampIndicator))
+		{
+			VolumetricBrushIndicator->MakeBrushIndicatorMesh(this, GetTargetWorld());
+		}
+	}
+}
+
+void USkinWeightsPaintTool::UpdateBrushStampIndicator()
+{
+	Super::UpdateBrushStampIndicator();
+}
+
+void USkinWeightsPaintTool::ShutdownBrushStampIndicator()
+{
+	if (BrushStampIndicator)
+	{
+		GetToolManager()->GetPairedGizmoManager()->DestroyGizmo(BrushStampIndicator);
+		BrushStampIndicator = nullptr;
+		GetToolManager()->GetPairedGizmoManager()->DeregisterGizmoType(SkinWeightsPaintTool::Private::BaseVolumetricBrushIndicatorGizmoType);
+	}
+}
+
+void USkinWeightsPaintTool::UpdateBrushIndicators()
+{
+	const bool bUseVolumnetricBrush = WeightToolProperties? (WeightToolProperties->GetBrushConfig().FalloffMode == EWeightBrushFalloffMode::Volume): false;
+	const bool bIsInBrushMode = WeightToolProperties ? (WeightToolProperties->EditingMode == EWeightEditMode::Brush) : false;
+
+	if (BrushStampIndicator)
+	{
+		BrushStampIndicator->bVisible = (bIsInBrushMode && !bUseVolumnetricBrush);
+		if (UVolumetricBrushStampIndicator* VolumetricBrushIndicator = Cast<UVolumetricBrushStampIndicator>(BrushStampIndicator))
+		{
+			VolumetricBrushIndicator->SetVolumetricBrushVisible(bIsInBrushMode && bUseVolumnetricBrush);
+		}
+	}
+}
+
 FBox USkinWeightsPaintTool::GetWorldSpaceFocusBox()
 {
 	if (!WeightToolProperties)
@@ -2210,6 +2261,7 @@ void USkinWeightsPaintToolProperties::SetFalloffMode(EWeightBrushFalloffMode InF
 	GetBrushConfig().FalloffMode = InFalloffMode;
 	SaveConfig();
 
+	WeightTool->UpdateBrushIndicators();
 	WeightTool->SetFocusInViewport();
 }
 
@@ -4090,8 +4142,7 @@ void USkinWeightsPaintTool::ToggleEditingMode()
 {
 	Weights.Deformer.SetAllVerticesToBeUpdated();
 
-	// toggle brush mode
-	SetBrushEnabled(WeightToolProperties->EditingMode == EWeightEditMode::Brush);
+	UpdateBrushIndicators();
 
 	// toggle which mesh we're selecting and what components (vert/edge/face)
 	UpdateSelectorState();
diff --git a/Engine/Plugins/Runtime/MeshModelingToolset/Source/MeshModelingToolsEditorOnly/Public/SkeletalMesh/SkinWeightsPaintTool.h b/Engine/Plugins/Runtime/MeshModelingToolset/Source/MeshModelingToolsEditorOnly/Public/SkeletalMesh/SkinWeightsPaintTool.h
--- a/Engine/Plugins/Runtime/MeshModelingToolset/Source/MeshModelingToolsEditorOnly/Public/SkeletalMesh/SkinWeightsPaintTool.h
+++ b/Engine/Plugins/Runtime/MeshModelingToolset/Source/MeshModelingToolsEditorOnly/Public/SkeletalMesh/SkinWeightsPaintTool.h
@@ -30,6 +30,7 @@ class USkinWeightsPaintTool;
 class UPolygonSelectionMechanic;
 class UPersonaEditorModeManagerContext;
 class FEditorViewportClient;
+class UVolumetricBrushStampIndicator;
 
 namespace UE::Geometry 
 {
@@ -653,6 +654,9 @@ public:
 	UE_API virtual void OnEndDrag(const FRay& Ray) override;
 	UE_API virtual bool OnUpdateHover(const FInputDeviceRay& DevicePos) override;
 	UE_API virtual double EstimateMaximumTargetDimension() override;
+	UE_API virtual void SetupBrushStampIndicator() override;
+	UE_API virtual void UpdateBrushStampIndicator() override;
+	UE_API virtual void ShutdownBrushStampIndicator() override;
 
 	UE_API void Init(const FToolBuilderState& InSceneState);
 	
@@ -788,6 +792,9 @@ public:
 		const FDynamicMesh3& InDynamicMesh,
 		const FMeshDescription& InMeshDescription);
 
+	// Update the brush indicators visibility state
+	UE_API void UpdateBrushIndicators();
+
 	// called whenever the selection is modified
 	DECLARE_MULTICAST_DELEGATE(FOnSelectionChanged);
 	FOnSelectionChanged OnSelectionChanged;
diff --git a/Engine/Plugins/Runtime/MeshModelingToolset/Source/ModelingComponents/Private/BaseGizmos/VolumetricBrushStampIndicator.cpp b/Engine/Plugins/Runtime/MeshModelingToolset/Source/ModelingComponents/Private/BaseGizmos/VolumetricBrushStampIndicator.cpp
new file mode 100644
--- /dev/null
+++ b/Engine/Plugins/Runtime/MeshModelingToolset/Source/ModelingComponents/Private/BaseGizmos/VolumetricBrushStampIndicator.cpp
@@ -0,0 +1,79 @@
+// Copyright Epic Games, Inc. All Rights Reserved.
+
+#include "BaseGizmos/VolumetricBrushStampIndicator.h"
+
+#include "Components/PrimitiveComponent.h"
+#include "Generators/SphereGenerator.h"
+#include "InteractiveGizmoManager.h"
+#include "InteractiveTool.h"
+#include "PreviewMesh.h"
+#include "ToolDataVisualizer.h"
+#include "ToolSetupUtil.h"
+
+#include UE_INLINE_GENERATED_CPP_BY_NAME(VolumetricBrushStampIndicator)
+
+UInteractiveGizmo* UVolumetricBrushStampIndicatorBuilder::BuildGizmo(const FToolBuilderState& SceneState) const
+{
+	UVolumetricBrushStampIndicator* NewGizmo = NewObject<UVolumetricBrushStampIndicator>(SceneState.GizmoManager);
+	return NewGizmo;
+}
+
+void UVolumetricBrushStampIndicator::MakeBrushIndicatorMesh(UInteractiveTool* ParentTool, UWorld* World)
+{
+	if (ensure(ParentTool))
+	{
+		SphereMesh = NewObject<UPreviewMesh>(ParentTool);
+		if (SphereMesh)
+		{
+			SphereMesh->CreateInWorld(World, FTransform::Identity);
+			UE::Geometry::FSphereGenerator SphereGen;
+			SphereGen.NumPhi = SphereGen.NumTheta = 32;
+			SphereGen.Generate();
+			FDynamicMesh3 Mesh(&SphereGen);
+			SphereMesh->UpdatePreview(&Mesh);
+
+			UMaterialInstanceDynamic* BrushIndicatorMaterial = ToolSetupUtil::GetDefaultBrushVolumeMaterial(ParentTool->GetToolManager());
+			if (BrushIndicatorMaterial)
+			{
+				SphereMesh->SetMaterial(BrushIndicatorMaterial);
+			}
+
+			// make sure raytracing is disabled on the brush indicator
+			if (UDynamicMeshComponent* DynamicComponent = Cast<UDynamicMeshComponent>(SphereMesh->GetRootComponent()))
+			{
+				DynamicComponent->SetEnableRaytracing(false);
+				SphereMesh->SetShadowsEnabled(false);
+			}
+			AttachedComponent = SphereMesh->GetRootComponent();
+
+			// Todo : set a "fully transparent" material for wireframe			
+		}
+	}
+}
+
+void UVolumetricBrushStampIndicator::Shutdown()
+{
+	if (SphereMesh)
+	{
+		SphereMesh->Disconnect();
+	}
+	SphereMesh = nullptr;
+}
+
+
+void UVolumetricBrushStampIndicator::SetVolumetricBrushVisible(bool bInVisible)
+{
+	if (SphereMesh)
+	{
+		SphereMesh->SetVisible(bInVisible);
+	}
+}
+
+bool UVolumetricBrushStampIndicator::IsVolumetricBrushVisible()
+{
+	if (SphereMesh)
+	{
+		return SphereMesh->IsVisible();
+	};
+	return false;
+}
diff --git a/Engine/Plugins/Runtime/MeshModelingToolset/Source/ModelingComponents/Public/BaseGizmos/VolumetricBrushStampIndicator.h b/Engine/Plugins/Runtime/MeshModelingToolset/Source/ModelingComponents/Public/BaseGizmos/VolumetricBrushStampIndicator.h
new file mode 100644
--- /dev/null
+++ b/Engine/Plugins/Runtime/MeshModelingToolset/Source/ModelingComponents/Public/BaseGizmos/VolumetricBrushStampIndicator.h
@@ -0,0 +1,46 @@
+// Copyright Epic Games, Inc. All Rights Reserved.
+
+#pragma once
+
+#include "InteractiveGizmo.h"
+#include "InteractiveGizmoBuilder.h"
+#include "BaseGizmos/BrushStampIndicator.h"
+
+#include "VolumetricBrushStampIndicator.generated.h"
+
+class UObject;
+class UInteractiveTool;
+class UPrimitiveComponent;
+class UPreviewMesh;
+class UWorld;
+
+UCLASS(MinimalAPI)
+class UVolumetricBrushStampIndicatorBuilder : public UInteractiveGizmoBuilder
+{
+	GENERATED_BODY()
+public:
+	MODELINGCOMPONENTS_API virtual UInteractiveGizmo* BuildGizmo(const FToolBuilderState& SceneState) const override;
+};
+
+/*
+ * UVolumetricBrushStampIndicator is a simple 3D brush indicator that can be volumetric 
+ */
+UCLASS(Transient, MinimalAPI)
+class UVolumetricBrushStampIndicator : public UBrushStampIndicator
+{
+	GENERATED_BODY()
+
+public:
+	MODELINGCOMPONENTS_API virtual void Shutdown() override;
+
+	MODELINGCOMPONENTS_API void SetVolumetricBrushVisible(bool bInVisible);
+
+	MODELINGCOMPONENTS_API bool IsVolumetricBrushVisible();
+
+	MODELINGCOMPONENTS_API void MakeBrushIndicatorMesh(UInteractiveTool* ParentTool, UWorld* World);
+
+private:
+	UPROPERTY()
+	TObjectPtr<UPreviewMesh> SphereMesh;
+};
+
