fix static mesh tool target 'IMeshDescriptionCommitter' interface not finding the right material ID mapping method -- same function is also on the Provider, so needs to be pulled to the top level for both interfaces to find the method

diff --git a/Engine/Plugins/Runtime/MeshModelingToolset/Source/ModelingComponentsEditorOnly/Public/ToolTargets/StaticMeshComponentToolTarget.h b/Engine/Plugins/Runtime/MeshModelingToolset/Source/ModelingComponentsEditorOnly/Public/ToolTargets/StaticMeshComponentToolTarget.h
--- a/Engine/Plugins/Runtime/MeshModelingToolset/Source/ModelingComponentsEditorOnly/Public/ToolTargets/StaticMeshComponentToolTarget.h
+++ b/Engine/Plugins/Runtime/MeshModelingToolset/Source/ModelingComponentsEditorOnly/Public/ToolTargets/StaticMeshComponentToolTarget.h
@@ -107,6 +107,14 @@ class UStaticMeshComponentToolTarget :
 	// IMeshDescriptionCommitter implementation
 	UE_API virtual void CommitMeshDescription(const FCommitter& Committer, const FCommitMeshParameters& CommitParams = FCommitMeshParameters()) override;
 	using IMeshDescriptionCommitter::CommitMeshDescription; // unhide the other overload
+
+	// Note: This method is on both the IMeshDescriptionCommitter and IMeshDescriptionProvider interface.
+	// The UStaticMeshComponentReadOnlyToolTarget base class only implements the Provider; we lift it to here
+	// so that the IMeshDescriptionCommitter interface will also find/call it.
+	virtual TArray<int32> GetPolygonGroupToMaterialIndexMap() const override
+	{
+		return UStaticMeshComponentReadOnlyToolTarget::GetPolygonGroupToMaterialIndexMap();
+	}
 };
 
 
diff --git a/Engine/Plugins/Runtime/MeshModelingToolset/Source/ModelingComponentsEditorOnly/Public/ToolTargets/StaticMeshToolTarget.h b/Engine/Plugins/Runtime/MeshModelingToolset/Source/ModelingComponentsEditorOnly/Public/ToolTargets/StaticMeshToolTarget.h
--- a/Engine/Plugins/Runtime/MeshModelingToolset/Source/ModelingComponentsEditorOnly/Public/ToolTargets/StaticMeshToolTarget.h
+++ b/Engine/Plugins/Runtime/MeshModelingToolset/Source/ModelingComponentsEditorOnly/Public/ToolTargets/StaticMeshToolTarget.h
@@ -139,6 +139,14 @@ public:
 	UE_API virtual void CommitDynamicMesh(const UE::Geometry::FDynamicMesh3& Mesh, const FDynamicMeshCommitInfo& CommitInfo) override;
 	using IDynamicMeshCommitter::CommitDynamicMesh; // unhide the other overload
 
+	// Note: This method is on both the IMeshDescriptionCommitter and IMeshDescriptionProvider interface.
+	// The UStaticMeshReadOnlyToolTarget base class only implements the Provider; we lift it to here
+	// so that the IMeshDescriptionCommitter interface will also find/call it.
+	virtual TArray<int32> GetPolygonGroupToMaterialIndexMap() const override
+	{
+		return UStaticMeshReadOnlyToolTarget::GetPolygonGroupToMaterialIndexMap();
+	}
+
 protected:
 	friend class UStaticMeshComponentToolTarget;
 
