ModelingTools: Added BrushValue property to the AttributePaintTool to allow users to specify the target value that is painted down rather than always accumulating to 1.0f.

diff --git a/Engine/Plugins/Runtime/MeshModelingToolset/Source/MeshModelingTools/Private/MeshAttributePaintTool.cpp b/Engine/Plugins/Runtime/MeshModelingToolset/Source/MeshModelingTools/Private/MeshAttributePaintTool.cpp
--- a/Engine/Plugins/Runtime/MeshModelingToolset/Source/MeshModelingTools/Private/MeshAttributePaintTool.cpp
+++ b/Engine/Plugins/Runtime/MeshModelingToolset/Source/MeshModelingTools/Private/MeshAttributePaintTool.cpp
@@ -732,13 +732,13 @@ void UMeshAttributePaintTool::ApplyStamp(const FBrushStampData& Stamp)
 		break;
 	}
 	case EBrushActionMode::Paint:
-		ApplyStamp_Paint(Stamp, ActionData, bInSmoothStroke, bInRemoveStroke);
+		ApplyStamp_Paint(Stamp, ActionData, {BrushActionProps->BrushValue, bInSmoothStroke, bInRemoveStroke});
 		break;
 	case EBrushActionMode::Smooth:
-		ApplyStamp_Paint(Stamp, ActionData, true, bInRemoveStroke);
+		ApplyStamp_Paint(Stamp, ActionData, {BrushActionProps->BrushValue, true, bInRemoveStroke});
 		break;
 	case EBrushActionMode::Erase:
-		ApplyStamp_Paint(Stamp, ActionData, bInSmoothStroke, !bInRemoveStroke);
+		ApplyStamp_Paint(Stamp, ActionData, {BrushActionProps->BrushValue, bInSmoothStroke, !bInRemoveStroke});
 		break;
 	default:
 		ensure(false);
@@ -812,10 +812,10 @@ void UMeshAttributePaintTool::UpdateFlatShadingSetting(bool bNewValue)
 
 void UMeshAttributePaintTool::ApplyStamp_Paint(const FBrushStampData& Stamp, FStampActionData& ActionData)
 {
-	ApplyStamp_Paint(Stamp, ActionData, bInSmoothStroke, bInRemoveStroke);
+	ApplyStamp_Paint(Stamp, ActionData, {1.0f, bInSmoothStroke, bInRemoveStroke});
 }
 
-void UMeshAttributePaintTool::ApplyStamp_Paint(const FBrushStampData& Stamp, FStampActionData& ActionData, bool bSmooth, bool bInvert)
+void UMeshAttributePaintTool::ApplyStamp_Paint(const FBrushStampData& Stamp, FStampActionData& ActionData, const FBrushActionData& BrushData)
 {
 	FTransform3d Transform(Cast<ISceneComponentBackedTarget>(Target)->GetWorldTransform());
 	FVector3d StampPosLocal = Transform.InverseTransformPosition((FVector3d)Stamp.WorldPosition);
@@ -828,9 +828,9 @@ void UMeshAttributePaintTool::ApplyStamp_Paint(const FBrushStampData& Stamp, FSt
 
 	const FDynamicMesh3* CurrentMesh = PreviewMesh->GetMesh();
 	FNonManifoldMappingSupport NonManifoldMappingSupport(*CurrentMesh);
-	if (bSmooth)
+	if (BrushData.bSmooth)
 	{
-		float SmoothSpeed = 0.25f;
+		constexpr float SmoothSpeed = 0.25f;
 
 		for (int32 k = 0; k < NumVertices; ++k)
 		{
@@ -857,16 +857,17 @@ void UMeshAttributePaintTool::ApplyStamp_Paint(const FBrushStampData& Stamp, FSt
 	}
 	else
 	{
-		float Sign = (bInvert) ? -1.0f : 1.0f;
-		float UseStrength = Sign * BrushProperties->BrushStrength * CurrentValueRange.Length();
+		const float TargetValue = BrushData.bErase ? CurrentValueRange.Min : BrushData.BrushValue;
+		const float UseStrength = BrushProperties->BrushStrength * CurrentValueRange.Length();
 		for (int32 k = 0; k < NumVertices; ++k)
 		{
-			int32 vid = ActionData.ROIVertices[k];
-			int32 srcvid =  NonManifoldMappingSupport.GetOriginalNonManifoldVertexID(vid);
-			FVector3d Position = CurrentMesh->GetVertex(vid);
-			float Falloff = (float)CalculateBrushFalloff(Distance(Position, StampPosLocal));
+			const int32 vid = ActionData.ROIVertices[k];
+			const int32 srcvid =  NonManifoldMappingSupport.GetOriginalNonManifoldVertexID(vid);
+			const FVector3d Position = CurrentMesh->GetVertex(vid);
+			const float Falloff = (float)CalculateBrushFalloff(Distance(Position, StampPosLocal));
 			ActionData.ROIBefore[k] = AttribData.CurrentValues[srcvid];
-			ActionData.ROIAfter[k] = CurrentValueRange.Clamp(ActionData.ROIBefore[k] + UseStrength*Falloff);
+			const float ValueAfter = FMath::Lerp(ActionData.ROIBefore[k], TargetValue, UseStrength*Falloff);
+			ActionData.ROIAfter[k] = CurrentValueRange.Clamp(ValueAfter);
 		}
 	}
 }
diff --git a/Engine/Plugins/Runtime/MeshModelingToolset/Source/MeshModelingTools/Public/MeshAttributePaintTool.h b/Engine/Plugins/Runtime/MeshModelingToolset/Source/MeshModelingTools/Public/MeshAttributePaintTool.h
--- a/Engine/Plugins/Runtime/MeshModelingToolset/Source/MeshModelingTools/Public/MeshAttributePaintTool.h
+++ b/Engine/Plugins/Runtime/MeshModelingToolset/Source/MeshModelingTools/Public/MeshAttributePaintTool.h
@@ -147,6 +147,9 @@ public:
 
 	UPROPERTY(EditAnywhere, Category = Attribute)
 	EBrushActionMode BrushAction = EBrushActionMode::Paint;
+
+	UPROPERTY(EditAnywhere, Category = Attribute, meta=(UIMin=0.0, UIMax=1.0, ModelingQuickEdit, ModelingQuickSettings=150))
+	float BrushValue = 1.0f;
 };
 
 
@@ -373,12 +376,19 @@ protected:
 	UE_API void ExternalUpdateValues(int32 AttribIndex, const TArray<int32>& VertexIndices, const TArray<float>& NewValues);
 
 private:
+	struct FBrushActionData
+	{
+		float BrushValue = 1.0f;
+		bool bSmooth = false;
+		bool bErase = false;
+	};
+	
 	TOptional<UE::Geometry::FGeometrySelection> GeometrySelection;
 	TSet<int32> SelectionTids;
 	TSet<int32> SelectionVids;
 	UE_API bool ShouldFilterTriangles() const;
 	// bInvert is ignored if bSmooth is true
-	void ApplyStamp_Paint(const FBrushStampData& Stamp, FStampActionData& ActionData, bool bSmooth, bool bInvert);
+	void ApplyStamp_Paint(const FBrushStampData& Stamp, FStampActionData& ActionData, const FBrushActionData& BrushData);
 
 	void UpdateMaterialMode(EMeshAttributePaintMaterialMode MaterialMode);
 	void UpdateFlatShadingSetting(bool bNewValue);
