ModelingTools: Remove obsolete comment

diff --git a/Engine/Plugins/Runtime/MeshModelingToolset/Source/MeshModelingTools/Public/MeshAttributePaintTool.h b/Engine/Plugins/Runtime/MeshModelingToolset/Source/MeshModelingTools/Public/MeshAttributePaintTool.h
--- a/Engine/Plugins/Runtime/MeshModelingToolset/Source/MeshModelingTools/Public/MeshAttributePaintTool.h
+++ b/Engine/Plugins/Runtime/MeshModelingToolset/Source/MeshModelingTools/Public/MeshAttributePaintTool.h
@@ -387,7 +387,7 @@ private:
 	TSet<int32> SelectionTids;
 	TSet<int32> SelectionVids;
 	UE_API bool ShouldFilterTriangles() const;
-	// bInvert is ignored if bSmooth is true
+	
 	void ApplyStamp_Paint(const FBrushStampData& Stamp, FStampActionData& ActionData, const FBrushActionData& BrushData);
 
 	void UpdateMaterialMode(EMeshAttributePaintMaterialMode MaterialMode);
