ModelingTools: Bake RS on startup of the AttributePaintTool to fix skewed stamps for meshes with non-uniform scale.

diff --git a/Engine/Plugins/Runtime/MeshModelingToolset/Source/MeshModelingTools/Private/DynamicMeshBrushTool.cpp b/Engine/Plugins/Runtime/MeshModelingToolset/Source/MeshModelingTools/Private/DynamicMeshBrushTool.cpp
--- a/Engine/Plugins/Runtime/MeshModelingToolset/Source/MeshModelingTools/Private/DynamicMeshBrushTool.cpp
+++ b/Engine/Plugins/Runtime/MeshModelingToolset/Source/MeshModelingTools/Private/DynamicMeshBrushTool.cpp
@@ -27,8 +27,10 @@ void UDynamicMeshBrushTool::Setup()
 	PreviewMesh = NewObject<UPreviewMesh>(this);
 	PreviewMesh->bBuildSpatialDataStructure = true;
 	PreviewMesh->CreateInWorld(GetTargetWorld(), FTransform::Identity);
-	FTransformSRT3d LocalToWorldTransform = UE::ToolTarget::GetLocalToWorldTransform(Target);
-	PreviewMesh->SetTransform((FTransform)LocalToWorldTransform);
+	BakedTargetTransform = UE::ToolTarget::GetLocalToWorldTransform(Target);
+	CurrentTargetTransform = BakedTargetTransform;
+	PreviewMesh->SetTransform((FTransform)BakedTargetTransform);
+
 	ToolSetupUtil::ApplyRenderingConfigurationToPreview(PreviewMesh, Target);
 
 	FComponentMaterialSet MaterialSet = UE::ToolTarget::GetMaterialSet(Target);
@@ -36,12 +38,26 @@ void UDynamicMeshBrushTool::Setup()
 
 	PreviewMesh->ReplaceMesh(UE::ToolTarget::GetDynamicMeshCopy(Target));
 
+	if (bBakeRotateScale)
+	{
+		if (UDynamicMeshComponent* MeshComponent = Cast<UDynamicMeshComponent>(PreviewMesh->GetRootComponent()))
+		{
+			// clamp scaling because if we allow zero-scale we cannot invert this transform on Accept
+			BakedTargetTransform.ClampMinimumScale(0.01);
+			FVector3d Translation = BakedTargetTransform.GetTranslation();
+			BakedTargetTransform.SetTranslation(FVector3d::Zero());
+			MeshComponent->ApplyTransform(BakedTargetTransform, false);
+			CurrentTargetTransform = FTransformSRT3d(Translation);
+			PreviewMesh->SetTransform(CurrentTargetTransform);
+		}
+	}
+
 	OnBaseMeshComponentChangedHandle = PreviewMesh->GetOnMeshChanged().Add(
 		FSimpleMulticastDelegate::FDelegate::CreateUObject(this, &UDynamicMeshBrushTool::OnBaseMeshComponentChanged));
 
 	// call this here so that base tool can estimate target dimension
 	InputMeshBoundsLocal = PreviewMesh->GetPreviewDynamicMesh()->GetBounds();
-	double ScaledDim = LocalToWorldTransform.TransformVector(FVector::OneVector).Size();
+	double ScaledDim = CurrentTargetTransform.TransformVector(FVector::OneVector).Size();
 	this->WorldToLocalScale = FMathd::Sqrt3 / FMathd::Max(FMathf::ZeroTolerance, ScaledDim);
 	UBaseBrushTool::Setup();
 
diff --git a/Engine/Plugins/Runtime/MeshModelingToolset/Source/MeshModelingTools/Private/MeshAttributePaintTool.cpp b/Engine/Plugins/Runtime/MeshModelingToolset/Source/MeshModelingTools/Private/MeshAttributePaintTool.cpp
--- a/Engine/Plugins/Runtime/MeshModelingToolset/Source/MeshModelingTools/Private/MeshAttributePaintTool.cpp
+++ b/Engine/Plugins/Runtime/MeshModelingToolset/Source/MeshModelingTools/Private/MeshAttributePaintTool.cpp
@@ -316,6 +316,8 @@ void UMeshAttributePaintTool::Setup()
 	ViewProperties->RestoreProperties(this);
 	AddToolPropertySource(ViewProperties);
 
+	// Bake the rotate and scale into the mesh to avoid skewed stamps.
+	bBakeRotateScale = true;
 	UDynamicMeshBrushTool::Setup();
 	if (ensure(BrushProperties))
 	{
@@ -556,7 +558,7 @@ bool UMeshAttributePaintTool::HitTest(const FRay& Ray, FHitResult& OutHit)
 
 void UMeshAttributePaintTool::CalculateVertexROI(const FBrushStampData& Stamp, TArray<int>& VertexROI)
 {
-	FTransform3d Transform(Cast<ISceneComponentBackedTarget>(Target)->GetWorldTransform());
+	FTransform3d Transform(PreviewMesh->GetTransform());
 	FVector3d StampPosLocal = Transform.InverseTransformPosition((FVector3d)Stamp.WorldPosition);
 
 	float Radius = GetCurrentBrushRadiusLocal();
@@ -817,7 +819,7 @@ void UMeshAttributePaintTool::ApplyStamp_Paint(const FBrushStampData& Stamp, FSt
 
 void UMeshAttributePaintTool::ApplyStamp_Paint(const FBrushStampData& Stamp, FStampActionData& ActionData, const FBrushActionData& BrushData)
 {
-	FTransform3d Transform(Cast<ISceneComponentBackedTarget>(Target)->GetWorldTransform());
+	FTransform3d Transform(PreviewMesh->GetTransform());
 	FVector3d StampPosLocal = Transform.InverseTransformPosition((FVector3d)Stamp.WorldPosition);
 
 	int32 NumVertices = ActionData.ROIVertices.Num();
@@ -864,7 +866,7 @@ void UMeshAttributePaintTool::ApplyStamp_Paint(const FBrushStampData& Stamp, FSt
 			const int32 vid = ActionData.ROIVertices[k];
 			const int32 srcvid =  NonManifoldMappingSupport.GetOriginalNonManifoldVertexID(vid);
 			const FVector3d Position = CurrentMesh->GetVertex(vid);
-			const float Falloff = (float)CalculateBrushFalloff(Distance(Position, StampPosLocal));
+			const float Falloff = static_cast<float>(CalculateBrushFalloff(Distance(Position, StampPosLocal)));
 			ActionData.ROIBefore[k] = AttribData.CurrentValues[srcvid];
 			const float ValueAfter = FMath::Lerp(ActionData.ROIBefore[k], TargetValue, UseStrength*Falloff);
 			ActionData.ROIAfter[k] = CurrentValueRange.Clamp(ValueAfter);
diff --git a/Engine/Plugins/Runtime/MeshModelingToolset/Source/MeshModelingTools/Public/DynamicMeshBrushTool.h b/Engine/Plugins/Runtime/MeshModelingToolset/Source/MeshModelingTools/Public/DynamicMeshBrushTool.h
--- a/Engine/Plugins/Runtime/MeshModelingToolset/Source/MeshModelingTools/Public/DynamicMeshBrushTool.h
+++ b/Engine/Plugins/Runtime/MeshModelingToolset/Source/MeshModelingTools/Public/DynamicMeshBrushTool.h
@@ -55,6 +55,13 @@ protected:
 
 	UE::Geometry::FAxisAlignedBox3d InputMeshBoundsLocal;
 
+	// When true during Setup(), Rotation and Scale will be baked into the PreviewMesh.
+	// The rotate/scale portion of the transform will be stored in BakedTargetTransform.
+	// The translate portion of the transform will be stored in CurrentTargetTransform.
+	bool bBakeRotateScale = false;
+	UE::Geometry::FTransformSRT3d BakedTargetTransform;
+	UE::Geometry::FTransformSRT3d CurrentTargetTransform;
+
 	//
 	// UBaseBrushTool private interface
 	//
