Remove old mesh validity check from simplify operator

diff --git a/Engine/Plugins/Runtime/MeshModelingToolset/Source/ModelingOperatorsEditorOnly/Private/CleaningOps/SimplifyMeshOp.cpp b/Engine/Plugins/Runtime/MeshModelingToolset/Source/ModelingOperatorsEditorOnly/Private/CleaningOps/SimplifyMeshOp.cpp
--- a/Engine/Plugins/Runtime/MeshModelingToolset/Source/ModelingOperatorsEditorOnly/Private/CleaningOps/SimplifyMeshOp.cpp
+++ b/Engine/Plugins/Runtime/MeshModelingToolset/Source/ModelingOperatorsEditorOnly/Private/CleaningOps/SimplifyMeshOp.cpp
@@ -194,19 +194,6 @@ bool FSimplifyMeshOp::ComputeStandardSimplifier(IMeshReduction* MeshReduction, c
 			return false;
 		}
 
-		// TODO: This CheckValidity has been here a long time, but it's not clear why we need it ...
-		// If MergeConcidentMeshEdges is creating invalid meshes, we probably need to fix it rather than workaround it here.
-		// For now, have changed the fail mode to Ensure so we can see if/when this actually happens.
-		if (OutResult.CheckValidity(true, EValidityCheckFailMode::Ensure) == false)
-		{
-			bFailedModifyNeedsRegen = true;
-		}
-
-		if (Progress && Progress->Cancelled())
-		{
-			return false;
-		}
-
 		// in the fallback case where merge edges failed, give up and reset it to what it was before the attempted merger (w/ split UV boundaries everywhere, oh well)
 		if (bFailedModifyNeedsRegen)
 		{
