make the static mesh -> mesh description 'with build scale' path respect the normal/tangent generation build settings, matching how the UStaticMeshComponentReadOnlyToolTarget::GetMeshDescriptionCopy path works

diff --git a/Engine/Plugins/Runtime/MeshModelingToolset/Source/ModelingComponentsEditorOnly/Private/ToolTargets/StaticMeshComponentToolTarget.cpp b/Engine/Plugins/Runtime/MeshModelingToolset/Source/ModelingComponentsEditorOnly/Private/ToolTargets/StaticMeshComponentToolTarget.cpp
--- a/Engine/Plugins/Runtime/MeshModelingToolset/Source/ModelingComponentsEditorOnly/Private/ToolTargets/StaticMeshComponentToolTarget.cpp
+++ b/Engine/Plugins/Runtime/MeshModelingToolset/Source/ModelingComponentsEditorOnly/Private/ToolTargets/StaticMeshComponentToolTarget.cpp
@@ -133,8 +133,6 @@ const FMeshDescription* UStaticMeshComponentReadOnlyToolTarget::GetMeshDescripti
 {
 	if (ensure(IsValid()))
 	{
-		ensure(GetMeshParams.bWantMeshTangents == false);		// cannot support in this path because we have to modify the FoundMeshDescription pointer!
-
 		UStaticMesh* StaticMesh = Cast<UStaticMeshComponent>(Component)->GetStaticMesh();
 		EMeshLODIdentifier UseLOD = EditingLOD;
 		if (StaticMesh && GetMeshParams.bHaveRequestLOD)
diff --git a/Engine/Plugins/Runtime/MeshModelingToolset/Source/ModelingComponentsEditorOnly/Private/ToolTargets/StaticMeshToolTarget.cpp b/Engine/Plugins/Runtime/MeshModelingToolset/Source/ModelingComponentsEditorOnly/Private/ToolTargets/StaticMeshToolTarget.cpp
--- a/Engine/Plugins/Runtime/MeshModelingToolset/Source/ModelingComponentsEditorOnly/Private/ToolTargets/StaticMeshToolTarget.cpp
+++ b/Engine/Plugins/Runtime/MeshModelingToolset/Source/ModelingComponentsEditorOnly/Private/ToolTargets/StaticMeshToolTarget.cpp
@@ -305,7 +305,7 @@ TArray<int32> UStaticMeshReadOnlyToolTarget::MapSectionToMaterialID(const UStati
 	return UE::Conversion::GetPolygonGroupToMaterialIndexMap(Mesh, LODType, UseLOD);
 }
 
-const FMeshDescription* UStaticMeshReadOnlyToolTarget::GetMeshDescriptionWithScaleApplied(const UStaticMesh* StaticMesh, int32 UseLOD, FMeshDescriptionCache& CachedMeshDescriptions)
+const FMeshDescription* UStaticMeshReadOnlyToolTarget::GetMeshDescriptionWithScaleApplied(const UStaticMesh* StaticMesh, int32 UseLOD, FMeshDescriptionCache& CachedMeshDescriptions, const FGetMeshDescriptionOptions& Options)
 {
 	static FMeshDescription EmptyMeshDescription;
 	static bool bFirst = true;
@@ -322,20 +322,27 @@ const FMeshDescription* UStaticMeshReadOnlyToolTarget::GetMeshDescriptionWithSca
 	}
 
 	FMeshDescription* FoundMeshDescription = nullptr;
-	FVector BuildScale = FVector::One();
+	TOptional<FMeshBuildSettings> FoundBuildSettings;
 	if (UseLOD == (int32)EMeshLODIdentifier::HiResSource)
 	{
 		FoundMeshDescription = StaticMesh->GetHiResMeshDescription();
-		BuildScale = StaticMesh->GetHiResSourceModel().BuildSettings.BuildScale3D;
+		FoundBuildSettings = StaticMesh->GetHiResSourceModel().BuildSettings;
 	}
 	else
 	{
 		FoundMeshDescription = StaticMesh->GetMeshDescription(UseLOD);
 		if (FoundMeshDescription)
 		{
-			BuildScale = StaticMesh->GetSourceModel(UseLOD).BuildSettings.BuildScale3D;
+			FoundBuildSettings = StaticMesh->GetSourceModel(UseLOD).BuildSettings;
 		}
 	}
+	
+	// Apply allow-recompute overrides, if present
+	if (FoundBuildSettings)
+	{
+		FoundBuildSettings->bRecomputeNormals &= Options.bAllowRecomputeNormals;
+		FoundBuildSettings->bRecomputeTangents &= Options.bAllowRecomputeTangents;
+	}
 
 	if (CachedMeshDescriptions.Contains(UseLOD))
 	{
@@ -349,15 +356,30 @@ const FMeshDescription* UStaticMeshReadOnlyToolTarget::GetMeshDescriptionWithSca
 		}
 	}
 
-	if (FoundMeshDescription && !BuildScale.Equals(FVector::OneVector))
+	if (FoundMeshDescription)
 	{
-		FCachedMeshDescription& Cache = CachedMeshDescriptions.Emplace(UseLOD);
-		Cache.Source = FoundMeshDescription;
-		Cache.Copy = MakeUnique<FMeshDescription>(*FoundMeshDescription);
-		FTransform ScaleTransform = FTransform::Identity;
-		ScaleTransform.SetScale3D(BuildScale);
-		FStaticMeshOperations::ApplyTransform(*Cache.Copy, ScaleTransform, true);
-		FoundMeshDescription = Cache.Copy.Get();
+		bool bNeedsBuildScale = !FoundBuildSettings->BuildScale3D.Equals(FVector::OneVector);
+		bool bNeedsNormalTangent = FoundBuildSettings->bRecomputeNormals || FoundBuildSettings->bRecomputeTangents;
+
+		if (bNeedsBuildScale || bNeedsNormalTangent)
+		{
+			FCachedMeshDescription& Cache = CachedMeshDescriptions.Emplace(UseLOD);
+			Cache.Source = FoundMeshDescription;
+			Cache.Copy = MakeUnique<FMeshDescription>(*FoundMeshDescription);
+
+			if (bNeedsBuildScale)
+			{
+				FTransform ScaleTransform = FTransform::Identity;
+				ScaleTransform.SetScale3D(FoundBuildSettings->BuildScale3D);
+				FStaticMeshOperations::ApplyTransform(*Cache.Copy, ScaleTransform, true);
+			}
+			if (bNeedsNormalTangent)
+			{
+				UE::MeshDescription::InitializeAutoGeneratedAttributes(*Cache.Copy, FoundBuildSettings.GetPtrOrNull());
+			}
+
+			FoundMeshDescription = Cache.Copy.Get();
+		}	
 	}
 
 	return (FoundMeshDescription != nullptr) ? FoundMeshDescription : &EmptyMeshDescription;
diff --git a/Engine/Plugins/Runtime/MeshModelingToolset/Source/ModelingComponentsEditorOnly/Public/ToolTargets/StaticMeshToolTarget.h b/Engine/Plugins/Runtime/MeshModelingToolset/Source/ModelingComponentsEditorOnly/Public/ToolTargets/StaticMeshToolTarget.h
--- a/Engine/Plugins/Runtime/MeshModelingToolset/Source/ModelingComponentsEditorOnly/Public/ToolTargets/StaticMeshToolTarget.h
+++ b/Engine/Plugins/Runtime/MeshModelingToolset/Source/ModelingComponentsEditorOnly/Public/ToolTargets/StaticMeshToolTarget.h
@@ -85,8 +85,18 @@ public:
 		TUniquePtr<FMeshDescription> Copy;
 	};
 	using FMeshDescriptionCache = TMap<int32, FCachedMeshDescription>;
+
 	// helper to get the mesh description with build scale applied, using/updating the cache as needed
-	static UE_API const FMeshDescription* GetMeshDescriptionWithScaleApplied(const UStaticMesh* Mesh, int32 UseLOD, FMeshDescriptionCache& Cache);
+	struct FGetMeshDescriptionOptions
+	{
+		bool bAllowRecomputeNormals = true;
+		bool bAllowRecomputeTangents = true;
+		static FGetMeshDescriptionOptions Default()
+		{
+			return FGetMeshDescriptionOptions{};
+		}
+	};
+	static UE_API const FMeshDescription* GetMeshDescriptionWithScaleApplied(const UStaticMesh* Mesh, int32 UseLOD, FMeshDescriptionCache& Cache, const FGetMeshDescriptionOptions& Options = FGetMeshDescriptionOptions::Default());
 
 protected:
 	TWeakObjectPtr<UStaticMesh> StaticMesh = nullptr;
