extract the editor-preferred delete actors method w/ non-editor fallback as a general utility function in modeling components, so it is easier to use in more places

diff --git a/Engine/Plugins/Runtime/MeshModelingToolset/Source/ModelingComponents/Private/ModelingToolTargetUtil.cpp b/Engine/Plugins/Runtime/MeshModelingToolset/Source/ModelingComponents/Private/ModelingToolTargetUtil.cpp
--- a/Engine/Plugins/Runtime/MeshModelingToolset/Source/ModelingComponents/Private/ModelingToolTargetUtil.cpp
+++ b/Engine/Plugins/Runtime/MeshModelingToolset/Source/ModelingComponents/Private/ModelingToolTargetUtil.cpp
@@ -29,7 +29,12 @@
 #if WITH_EDITOR
 #include "StaticMeshCompiler.h"
 #include "SkinnedAssetCompiler.h"
+
+#include "Editor/UnrealEdEngine.h"
+#include "Selection.h"
+#include "UnrealEdGlobals.h"
 #endif
+#include "Engine/World.h"
 
 #include "DynamicMesh/NonManifoldMappingSupport.h"
 
@@ -56,6 +61,30 @@ AActor* UE::ToolTarget::GetTargetActor(UToolTarget* Target)
 	return nullptr;
 }
 
+void UE::ToolTarget::SafeDeleteActor(AActor* Actor)
+{
+	if (!Actor)
+	{
+		// no actor to delete
+		return;
+	}
+	if (UWorld* ActorWorld = Actor->GetWorld())
+	{
+#if WITH_EDITOR
+		// Attempt to delete in the standard editor way
+		if (GIsEditor && GUnrealEd)
+		{
+			GUnrealEd->DeleteActors(TArray{ Actor }, ActorWorld, GUnrealEd->GetSelectedActors()->GetElementSelectionSet());
+		}
+		else
+#endif
+		// Alternate non-editor deletion path
+		{
+			ActorWorld->DestroyActor(Actor);
+		}
+	}
+}
+
 UPrimitiveComponent* UE::ToolTarget::GetTargetComponent(UToolTarget* Target)
 {
 	IPrimitiveComponentBackedTarget* TargetComponent = Cast<IPrimitiveComponentBackedTarget>(Target);
diff --git a/Engine/Plugins/Runtime/MeshModelingToolset/Source/ModelingComponents/Private/PropertySets/OnAcceptProperties.cpp b/Engine/Plugins/Runtime/MeshModelingToolset/Source/ModelingComponents/Private/PropertySets/OnAcceptProperties.cpp
--- a/Engine/Plugins/Runtime/MeshModelingToolset/Source/ModelingComponents/Private/PropertySets/OnAcceptProperties.cpp
+++ b/Engine/Plugins/Runtime/MeshModelingToolset/Source/ModelingComponents/Private/PropertySets/OnAcceptProperties.cpp
@@ -6,12 +6,7 @@
 #include "GameFramework/Actor.h"
 #include "UObject/UObjectGlobals.h"
 #include "InteractiveToolManager.h"
-
-#if WITH_EDITOR
-#include "Editor/UnrealEdEngine.h"
-#include "Selection.h"
-#include "UnrealEdGlobals.h"
-#endif
+#include "ModelingToolTargetUtil.h"
 
 #include UE_INLINE_GENERATED_CPP_BY_NAME(OnAcceptProperties)
 
@@ -58,19 +53,7 @@ void UOnAcceptHandleSourcesPropertiesBase::ApplyMethod(const TArray<AActor*>& Ac
 
 			if (bDelete)
 			{
-				if (UWorld* ActorWorld = Actor->GetWorld())
-				{
-#if WITH_EDITOR
-					if (GIsEditor && GUnrealEd)
-					{
-						GUnrealEd->DeleteActors(TArray{Actor}, ActorWorld, GUnrealEd->GetSelectedActors()->GetElementSelectionSet());
-					}
-					else
-#endif
-					{
-						ActorWorld->DestroyActor(Actor);
-					}
-				}
+				UE::ToolTarget::SafeDeleteActor(Actor);
 			}
 			else
 			{
diff --git a/Engine/Plugins/Runtime/MeshModelingToolset/Source/ModelingComponents/Public/ModelingToolTargetUtil.h b/Engine/Plugins/Runtime/MeshModelingToolset/Source/ModelingComponents/Public/ModelingToolTargetUtil.h
--- a/Engine/Plugins/Runtime/MeshModelingToolset/Source/ModelingComponents/Public/ModelingToolTargetUtil.h
+++ b/Engine/Plugins/Runtime/MeshModelingToolset/Source/ModelingComponents/Public/ModelingToolTargetUtil.h
@@ -308,6 +308,10 @@ MODELINGCOMPONENTS_API UStaticMesh* GetStaticMeshFromTargetIfAvailable(UToolTarg
 /** @return SkeletalMesh from a tool target */
 MODELINGCOMPONENTS_API USkeletalMesh* GetSkeletalMeshFromTargetIfAvailable(UToolTarget* Target);
 
+/**
+ * Delete an AActor in the correct way for the editor, with a fallback for the non-editor case
+ */
+MODELINGCOMPONENTS_API void SafeDeleteActor(AActor* TargetActor);
 
 }  // end namespace ToolTarget
 }  // end namespace UE
\ No newline at end of file
