[MeshModelingToolset] minor refactor to extract get available lod function to a shared header

diff --git a/Engine/Plugins/Runtime/MeshModelingToolset/Source/ModelingComponentsEditorOnly/Private/ModelingToolTargetEditorOnlyUtil.cpp b/Engine/Plugins/Runtime/MeshModelingToolset/Source/ModelingComponentsEditorOnly/Private/ModelingToolTargetEditorOnlyUtil.cpp
new file mode 100644
--- /dev/null
+++ b/Engine/Plugins/Runtime/MeshModelingToolset/Source/ModelingComponentsEditorOnly/Private/ModelingToolTargetEditorOnlyUtil.cpp
@@ -0,0 +1,32 @@
+// Copyright Epic Games, Inc. All Rights Reserved.
+
+#include "ModelingToolTargetEditorOnlyUtil.h"
+#include "Engine/SkeletalMesh.h"
+
+#define LOCTEXT_NAMESPACE "ModelingToolTargetEditorOnlyUtil"
+
+TArray<EMeshLODIdentifier> UE::ToolTarget::GetAvailableLODs(const USkeletalMesh* SkeletalMesh, bool bSkipAutoGenerated)
+{
+	TArray<EMeshLODIdentifier> AvailableLODs;
+
+	if (!SkeletalMesh)
+	{
+		return AvailableLODs;
+	}
+
+	const int32 MaxLODs = FMath::Min(8, SkeletalMesh->GetNumSourceModels());
+
+	for (int32 LODIndex = 0; LODIndex < MaxLODs; LODIndex++)
+	{
+		if (bSkipAutoGenerated && !SkeletalMesh->HasMeshDescription(LODIndex))
+		{
+			continue;
+		}
+
+		AvailableLODs.Add(static_cast<EMeshLODIdentifier>(LODIndex));
+	}
+
+	return AvailableLODs;
+}
+
+#undef LOCTEXT_NAMESPACE
\ No newline at end of file
diff --git a/Engine/Plugins/Runtime/MeshModelingToolset/Source/ModelingComponentsEditorOnly/Private/ToolTargets/SkeletalMeshToolTarget.cpp b/Engine/Plugins/Runtime/MeshModelingToolset/Source/ModelingComponentsEditorOnly/Private/ToolTargets/SkeletalMeshToolTarget.cpp
--- a/Engine/Plugins/Runtime/MeshModelingToolset/Source/ModelingComponentsEditorOnly/Private/ToolTargets/SkeletalMeshToolTarget.cpp
+++ b/Engine/Plugins/Runtime/MeshModelingToolset/Source/ModelingComponentsEditorOnly/Private/ToolTargets/SkeletalMeshToolTarget.cpp
@@ -9,6 +9,7 @@
 #include "Rendering/SkeletalMeshModel.h"
 #include "Widgets/Notifications/SNotificationList.h"
 #include "Framework/Notifications/NotificationManager.h"
+#include "ModelingToolTargetEditorOnlyUtil.h"
 #include "ModelingToolTargetUtil.h"
 #include "SkeletalMeshAttributes.h"
 #include "Misc/SlowTask.h"
@@ -56,26 +57,7 @@ bool USkeletalMeshReadOnlyToolTarget::IsValid(const USkeletalMesh* SkeletalMeshI
 
 TArray<EMeshLODIdentifier> USkeletalMeshReadOnlyToolTarget::GetAvailableLODs(const USkeletalMesh* SkeletalMesh, bool bSkipAutoGenerated)
 {
-	TArray<EMeshLODIdentifier> AvailableLODs;
-
-	if (!SkeletalMesh)
-	{
-		return AvailableLODs;
-	}
-
-	const int32 MaxLODs = FMath::Min(8, SkeletalMesh->GetNumSourceModels());
-
-	for (int32 LODIndex = 0; LODIndex < MaxLODs; LODIndex++)
-	{
-		if (bSkipAutoGenerated && !SkeletalMesh->HasMeshDescription(LODIndex))
-		{
-			continue;
-		}
-
-		AvailableLODs.Add(static_cast<EMeshLODIdentifier>(LODIndex));
-	}
-
-	return AvailableLODs;
+	return UE::ToolTarget::GetAvailableLODs(SkeletalMesh, bSkipAutoGenerated);
 }
 
 int32 USkeletalMeshReadOnlyToolTarget::GetNumMaterials() const
diff --git a/Engine/Plugins/Runtime/MeshModelingToolset/Source/ModelingComponentsEditorOnly/Public/ModelingToolTargetEditorOnlyUtil.h b/Engine/Plugins/Runtime/MeshModelingToolset/Source/ModelingComponentsEditorOnly/Public/ModelingToolTargetEditorOnlyUtil.h
new file mode 100644
--- /dev/null
+++ b/Engine/Plugins/Runtime/MeshModelingToolset/Source/ModelingComponentsEditorOnly/Public/ModelingToolTargetEditorOnlyUtil.h
@@ -0,0 +1,27 @@
+// Copyright Epic Games, Inc. All Rights Reserved.
+
+#pragma once
+
+#include "CoreMinimal.h"
+#include "TargetInterfaces/MeshTargetInterfaceTypes.h"
+
+class USkeletalMesh;
+//
+// UE::ToolTarget:: namespace contains utility/helper functions for interacting with UToolTargets.
+// Generally these are meant to be used by UInteractiveTools to handle standard tasks that would
+// otherwise require each Tool to figure out things like which ToolTargetInterface to cast to, etc.
+// Using these functions ideally avoids all the boilerplate inherent in the ToolTarget system.
+// 
+// However, the cost is that it is not necessarily the most efficient, as each one of these functions
+// may potentially do many repeated Cast<>'s internally. So, use sparingly, or cache the outputs.
+//
+namespace UE
+{
+namespace ToolTarget
+{
+
+/** @return available LODs on the Skeletal Mesh  */
+MODELINGCOMPONENTSEDITORONLY_API TArray<EMeshLODIdentifier> GetAvailableLODs(const USkeletalMesh* SkeletalMesh, bool bSkipAutoGenerated);
+	
+}  // end namespace ToolTarget
+}  // end namespace UE
\ No newline at end of file
