[MeshModelingToolSet] Added tool target interface for skeleton edit, which would allow skeleleton edit tool to edit both skeletal mesh and dynamic mesh

diff --git a/Engine/Plugins/Runtime/MeshModelingToolset/Source/MeshModelingToolsEditorOnly/Private/SkeletalMesh/SkeletonEditingTool.cpp b/Engine/Plugins/Runtime/MeshModelingToolset/Source/MeshModelingToolsEditorOnly/Private/SkeletalMesh/SkeletonEditingTool.cpp
--- a/Engine/Plugins/Runtime/MeshModelingToolset/Source/MeshModelingToolsEditorOnly/Private/SkeletalMesh/SkeletonEditingTool.cpp
+++ b/Engine/Plugins/Runtime/MeshModelingToolset/Source/MeshModelingToolsEditorOnly/Private/SkeletalMesh/SkeletonEditingTool.cpp
@@ -23,6 +23,9 @@
 
 #include "SkeletalMesh/SkeletonTransformProxy.h"
 #include "TargetInterfaces/PrimitiveComponentBackedTarget.h"
+#include "TargetInterfaces/SkeletonProvider.h"
+#include "TargetInterfaces/SkeletonCommitter.h"
+#include "TargetInterfaces/SkeletalMeshBackedTarget.h"
 
 #include UE_INLINE_GENERATED_CPP_BY_NAME(SkeletonEditingTool)
 
@@ -72,7 +75,11 @@ bool USkeletonEditingToolBuilder::CanBuildTool(const FToolBuilderState& SceneSta
 
 const FToolTargetTypeRequirements& USkeletonEditingToolBuilder::GetTargetRequirements() const
 {
-	static const FToolTargetTypeRequirements TypeRequirements(UPrimitiveComponentBackedTarget::StaticClass());
+	static const FToolTargetTypeRequirements TypeRequirements({
+			UPrimitiveComponentBackedTarget::StaticClass(),
+			USkeletonProvider::StaticClass(),
+			USkeletonCommitter::StaticClass()
+			});
 	return TypeRequirements;
 }
 
@@ -105,18 +112,17 @@ void USkeletonEditingTool::Setup()
 {
 	Super::Setup();
 
-	const IPrimitiveComponentBackedTarget* TargetComponent = Cast<IPrimitiveComponentBackedTarget>(Target);
-	USkeletalMeshComponent* Component = Cast<USkeletalMeshComponent>(TargetComponent->GetOwnerComponent());
+	const IPrimitiveComponentBackedTarget* PrimitiveComponentTarget = Cast<IPrimitiveComponentBackedTarget>(Target);
+	ISkeletonCommitter* SkeletonCommitter = CastChecked<ISkeletonCommitter>(Target);
+	UPrimitiveComponent* Component = Cast<UPrimitiveComponent>(PrimitiveComponentTarget->GetOwnerComponent());
 
-	USkeletalMesh* SkeletalMesh = Component ? Component->GetSkeletalMeshAsset() : nullptr;
-	if (!SkeletalMesh)
+	if (ISkeletalMeshBackedTarget* SkeletalMeshBackedTarget = Cast<ISkeletalMeshBackedTarget>(Target))
 	{
-		return;
+		// Used to get bone draw size, optional
+		WeakMesh = SkeletalMeshBackedTarget->GetSkeletalMesh();
 	}
-
-	WeakMesh = SkeletalMesh;
-
-	SetupModifier(SkeletalMesh);
+	
+	SetupModifier(SkeletonCommitter);
 	SetupPreviewMesh();
 	SetupProperties();
 	SetupBehaviors();
@@ -131,17 +137,10 @@ void USkeletonEditingTool::Setup()
 	}
 }
 
-void USkeletonEditingTool::SetupModifier(USkeletalMesh* InSkeletalMesh)
+void USkeletonEditingTool::SetupModifier(ISkeletonCommitter* InCommitter)
 {
-	if (InSkeletalMesh)
-	{
-		Modifier = NewObject<USkeletonModifier>(this);
-		Modifier->SetSkeletalMesh(InSkeletalMesh);
-	}
-	else
-	{
-		Modifier = nullptr;
-	}
+	Modifier = NewObject<USkeletonModifier>(this);
+	InCommitter->SetupSkeletonModifier(Modifier);
 }
 
 void USkeletonEditingTool::SetupPreviewMesh()
@@ -209,7 +208,7 @@ void USkeletonEditingTool::SetupBehaviors()
 	AddInputBehavior(ClickDragBehavior);
 }
 
-void USkeletonEditingTool::SetupGizmo(USkeletalMeshComponent* InComponent)
+void USkeletonEditingTool::SetupGizmo(USceneComponent* InComponent)
 {
 	if (!GizmoContext.IsValid() || !InComponent)
 	{
@@ -410,7 +409,8 @@ void USkeletonEditingTool::Shutdown(EToolShutdownType ShutdownType)
 	if (ShutdownType == EToolShutdownType::Accept)
 	{
 		GetToolManager()->BeginUndoTransaction(LOCTEXT("SkeletonEditingTool", "Commit Skeleton Editing"));
-		Modifier->CommitSkeletonToSkeletalMesh();
+		ISkeletonCommitter* SkeletonCommitter = CastChecked<ISkeletonCommitter>(Target);
+		SkeletonCommitter->CommitSkeletonModifier(Modifier);
 		GetToolManager()->EndUndoTransaction();
 
 		// to force to refresh the tree
diff --git a/Engine/Plugins/Runtime/MeshModelingToolset/Source/MeshModelingToolsEditorOnly/Public/SkeletalMesh/SkeletonEditingTool.h b/Engine/Plugins/Runtime/MeshModelingToolset/Source/MeshModelingToolsEditorOnly/Public/SkeletalMesh/SkeletonEditingTool.h
--- a/Engine/Plugins/Runtime/MeshModelingToolset/Source/MeshModelingToolsEditorOnly/Public/SkeletalMesh/SkeletonEditingTool.h
+++ b/Engine/Plugins/Runtime/MeshModelingToolset/Source/MeshModelingToolsEditorOnly/Public/SkeletalMesh/SkeletonEditingTool.h
@@ -27,6 +27,7 @@ class USkeletonTransformProxy;
 class USkeletalMeshGizmoContextObjectBase;
 class USkeletalMeshGizmoWrapperBase;
 class UPolygonSelectionMechanic;
+class ISkeletonCommitter;
 
 namespace SkeletonEditingTool
 {
@@ -258,11 +259,11 @@ private:
 	UE_API void NormalizeSelection();
 
 	// setup
-	UE_API void SetupModifier(USkeletalMesh* InSkeletalMesh);
+	UE_API void SetupModifier(ISkeletonCommitter* InCommitter);
 	UE_API void SetupPreviewMesh();
 	UE_API void SetupProperties();
 	UE_API void SetupBehaviors();
-	UE_API void SetupGizmo(USkeletalMeshComponent* InComponent);
+	UE_API void SetupGizmo(USceneComponent* InComponent);
 	UE_API void SetupWatchers();
 	UE_API void SetupComponentsSelection();
 	
diff --git a/Engine/Plugins/Runtime/MeshModelingToolset/Source/ModelingComponentsEditorOnly/ModelingComponentsEditorOnly.Build.cs b/Engine/Plugins/Runtime/MeshModelingToolset/Source/ModelingComponentsEditorOnly/ModelingComponentsEditorOnly.Build.cs
--- a/Engine/Plugins/Runtime/MeshModelingToolset/Source/ModelingComponentsEditorOnly/ModelingComponentsEditorOnly.Build.cs
+++ b/Engine/Plugins/Runtime/MeshModelingToolset/Source/ModelingComponentsEditorOnly/ModelingComponentsEditorOnly.Build.cs
@@ -39,7 +39,8 @@ public class ModelingComponentsEditorOnly : ModuleRules
                 "RHI",
 				"AssetTools",
 				"UnrealEd", 
-				"SkeletalMeshUtilitiesCommon" // required for asset factories
+				"SkeletalMeshUtilitiesCommon", // required for asset factories,
+				"SkeletalMeshModifiers"
 			}
 			);
 
diff --git a/Engine/Plugins/Runtime/MeshModelingToolset/Source/ModelingComponentsEditorOnly/Private/ToolTargets/DynamicMeshComponentToolTarget.cpp b/Engine/Plugins/Runtime/MeshModelingToolset/Source/ModelingComponentsEditorOnly/Private/ToolTargets/DynamicMeshComponentToolTarget.cpp
--- a/Engine/Plugins/Runtime/MeshModelingToolset/Source/ModelingComponentsEditorOnly/Private/ToolTargets/DynamicMeshComponentToolTarget.cpp
+++ b/Engine/Plugins/Runtime/MeshModelingToolset/Source/ModelingComponentsEditorOnly/Private/ToolTargets/DynamicMeshComponentToolTarget.cpp
@@ -12,6 +12,8 @@
 
 #include "Misc/ITransaction.h"
 #include "ScopedTransaction.h"
+#include "DynamicMesh/DynamicBoneAttribute.h"
+#include "SkeletalMeshModifiers/Public/SkeletonModifier.h"
 
 #include UE_INLINE_GENERATED_CPP_BY_NAME(DynamicMeshComponentToolTarget)
 
@@ -219,6 +221,50 @@ IInterface_CollisionDataProvider* UDynamicMeshComponentToolTarget::GetComplexCol
 	return IsValid() ? Cast<UDynamicMeshComponent>(Component) : nullptr;
 }
 
+FReferenceSkeleton UDynamicMeshComponentToolTarget::GetSkeleton() const
+{
+	if (!IsValid())
+	{
+		return FReferenceSkeleton();
+	}
+	
+	const FDynamicMesh3& EditingMesh = *(Cast<UDynamicMeshComponent>(Component)->GetMesh());
+	
+	FReferenceSkeleton RefSkeleton;
+
+	{
+		FReferenceSkeletonModifier RefSkeletonModifier(RefSkeleton, nullptr);
+
+		using namespace UE::Geometry;
+		const FDynamicMeshBoneNameAttribute* BoneNameAttrib = EditingMesh.Attributes()->GetBoneNames();
+		const FDynamicMeshBoneParentIndexAttribute* BoneParentIndexAttrib = EditingMesh.Attributes()->GetBoneParentIndices();
+		const FDynamicMeshBonePoseAttribute* BonePoses = EditingMesh.Attributes()->GetBonePoses();
+		const int32 NumBones = EditingMesh.Attributes()->GetNumBones();
+		
+		for (int32 BoneIndex = 0; BoneIndex < NumBones; ++BoneIndex)
+		{
+			const int32 ParentBoneIndex = BoneParentIndexAttrib->GetValue(BoneIndex);
+			FName BoneName = BoneNameAttrib->GetValue(BoneIndex);
+			FMeshBoneInfo BoneInfo(BoneName, BoneName.ToString(), ParentBoneIndex);
+
+			FTransform Transform = BonePoses->GetValue(BoneIndex);
+			RefSkeletonModifier.Add(BoneInfo, Transform);
+		}
+	}
+	
+	return RefSkeleton;
+}
+
+void UDynamicMeshComponentToolTarget::SetupSkeletonModifier(USkeletonModifier* InModifier)
+{
+	InModifier->SetDynamicMesh(GetDynamicMeshContainer());
+}
+
+void UDynamicMeshComponentToolTarget::CommitSkeletonModifier(USkeletonModifier* InModifier)
+{
+	InModifier->CommitSkeletonToDynamicMesh();
+}
+
 
 // Factory
 
diff --git a/Engine/Plugins/Runtime/MeshModelingToolset/Source/ModelingComponentsEditorOnly/Private/ToolTargets/SkeletalMeshComponentToolTarget.cpp b/Engine/Plugins/Runtime/MeshModelingToolset/Source/ModelingComponentsEditorOnly/Private/ToolTargets/SkeletalMeshComponentToolTarget.cpp
--- a/Engine/Plugins/Runtime/MeshModelingToolset/Source/ModelingComponentsEditorOnly/Private/ToolTargets/SkeletalMeshComponentToolTarget.cpp
+++ b/Engine/Plugins/Runtime/MeshModelingToolset/Source/ModelingComponentsEditorOnly/Private/ToolTargets/SkeletalMeshComponentToolTarget.cpp
@@ -12,6 +12,7 @@
 #include "Framework/Notifications/NotificationManager.h"
 #include "UObject/Package.h"
 #include "Widgets/Notifications/SNotificationList.h"
+#include "SkeletalMeshModifiers/Public/SkeletonModifier.h"
 
 #include UE_INLINE_GENERATED_CPP_BY_NAME(SkeletalMeshComponentToolTarget)
 
@@ -162,6 +163,11 @@ USkeletalMesh* USkeletalMeshComponentReadOnlyToolTarget::GetSkeletalMesh() const
 	return IsValid() ? Cast<USkeletalMeshComponent>(Component)->GetSkeletalMeshAsset() : nullptr;
 }
 
+FReferenceSkeleton USkeletalMeshComponentReadOnlyToolTarget::GetSkeleton() const
+{
+	return IsValid() ? GetSkeletalMesh()->GetRefSkeleton() : FReferenceSkeleton();
+}
+
 
 //
 // USkeletalMeshComponentToolTarget
@@ -197,6 +203,16 @@ void USkeletalMeshComponentToolTarget::CommitDynamicMesh(const FDynamicMesh3& Me
 	}
 }
 
+void USkeletalMeshComponentToolTarget::SetupSkeletonModifier(USkeletonModifier* InModifier)
+{
+	InModifier->SetSkeletalMesh(GetSkeletalMesh());
+}
+
+void USkeletalMeshComponentToolTarget::CommitSkeletonModifier(USkeletonModifier* InModifier)
+{
+	InModifier->CommitSkeletonToSkeletalMesh();
+}
+
 
 // Factory
 
diff --git a/Engine/Plugins/Runtime/MeshModelingToolset/Source/ModelingComponentsEditorOnly/Public/TargetInterfaces/SkeletonCommitter.h b/Engine/Plugins/Runtime/MeshModelingToolset/Source/ModelingComponentsEditorOnly/Public/TargetInterfaces/SkeletonCommitter.h
new file mode 100644
--- /dev/null
+++ b/Engine/Plugins/Runtime/MeshModelingToolset/Source/ModelingComponentsEditorOnly/Public/TargetInterfaces/SkeletonCommitter.h
@@ -0,0 +1,26 @@
+﻿// Copyright Epic Games, Inc. All Rights Reserved.
+
+#pragma once
+
+#include "CoreMinimal.h"
+#include "UObject/Interface.h"
+
+#include "SkeletonCommitter.generated.h"
+
+
+class USkeletonModifier;
+
+UINTERFACE(MinimalAPI)
+class USkeletonCommitter: public UInterface
+{
+	GENERATED_BODY()
+};
+
+class ISkeletonCommitter
+{
+	GENERATED_BODY()
+
+public:
+	virtual void SetupSkeletonModifier(USkeletonModifier* InModifier) = 0;
+	virtual void CommitSkeletonModifier(USkeletonModifier* InModifier) = 0;
+};
\ No newline at end of file
diff --git a/Engine/Plugins/Runtime/MeshModelingToolset/Source/ModelingComponentsEditorOnly/Public/TargetInterfaces/SkeletonProvider.h b/Engine/Plugins/Runtime/MeshModelingToolset/Source/ModelingComponentsEditorOnly/Public/TargetInterfaces/SkeletonProvider.h
new file mode 100644
--- /dev/null
+++ b/Engine/Plugins/Runtime/MeshModelingToolset/Source/ModelingComponentsEditorOnly/Public/TargetInterfaces/SkeletonProvider.h
@@ -0,0 +1,25 @@
+﻿// Copyright Epic Games, Inc. All Rights Reserved.
+
+#pragma once
+
+#include "CoreMinimal.h"
+#include "UObject/Interface.h"
+
+#include "SkeletonProvider.generated.h"
+
+
+struct FReferenceSkeleton;
+
+UINTERFACE(MinimalAPI)
+class USkeletonProvider : public UInterface
+{
+	GENERATED_BODY()
+};
+
+class ISkeletonProvider
+{
+	GENERATED_BODY()
+
+public:
+	virtual FReferenceSkeleton GetSkeleton() const = 0;
+};
\ No newline at end of file
diff --git a/Engine/Plugins/Runtime/MeshModelingToolset/Source/ModelingComponentsEditorOnly/Public/ToolTargets/DynamicMeshComponentToolTarget.h b/Engine/Plugins/Runtime/MeshModelingToolset/Source/ModelingComponentsEditorOnly/Public/ToolTargets/DynamicMeshComponentToolTarget.h
--- a/Engine/Plugins/Runtime/MeshModelingToolset/Source/ModelingComponentsEditorOnly/Public/ToolTargets/DynamicMeshComponentToolTarget.h
+++ b/Engine/Plugins/Runtime/MeshModelingToolset/Source/ModelingComponentsEditorOnly/Public/ToolTargets/DynamicMeshComponentToolTarget.h
@@ -11,6 +11,8 @@
 #include "TargetInterfaces/MeshDescriptionProvider.h"
 #include "TargetInterfaces/DynamicMeshSource.h"
 #include "TargetInterfaces/PhysicsDataSource.h"
+#include "TargetInterfaces/SkeletonCommitter.h"
+#include "TargetInterfaces/SkeletonProvider.h"
 #include "ToolTargets/PrimitiveComponentToolTarget.h"
 
 #include "DynamicMeshComponentToolTarget.generated.h"
@@ -31,7 +33,9 @@ class UDynamicMeshComponentToolTarget :
 	public IDynamicMeshCommitter,
 	public IMaterialProvider,
 	public IPersistentDynamicMeshSource,
-	public IPhysicsDataSource
+	public IPhysicsDataSource,
+	public ISkeletonProvider,
+	public ISkeletonCommitter
 {
 	GENERATED_BODY()
 
@@ -71,6 +75,14 @@ public:
 	UE_API virtual UBodySetup* GetBodySetup() const override;
 	UE_API virtual IInterface_CollisionDataProvider* GetComplexCollisionProvider() const override;
 
+	// ISkeletonProvider implementation
+	UE_API virtual FReferenceSkeleton GetSkeleton() const override;
+	
+	// ISkeletonComitter implementation
+	UE_API virtual void SetupSkeletonModifier(USkeletonModifier* InModifier) override;
+	UE_API virtual void CommitSkeletonModifier(USkeletonModifier* InModifier) override;
+	
+
 	// Rest provided by parent class
 
 protected:
diff --git a/Engine/Plugins/Runtime/MeshModelingToolset/Source/ModelingComponentsEditorOnly/Public/ToolTargets/SkeletalMeshComponentToolTarget.h b/Engine/Plugins/Runtime/MeshModelingToolset/Source/ModelingComponentsEditorOnly/Public/ToolTargets/SkeletalMeshComponentToolTarget.h
--- a/Engine/Plugins/Runtime/MeshModelingToolset/Source/ModelingComponentsEditorOnly/Public/ToolTargets/SkeletalMeshComponentToolTarget.h
+++ b/Engine/Plugins/Runtime/MeshModelingToolset/Source/ModelingComponentsEditorOnly/Public/ToolTargets/SkeletalMeshComponentToolTarget.h
@@ -10,6 +10,8 @@
 #include "TargetInterfaces/MeshDescriptionCommitter.h"
 #include "TargetInterfaces/MeshDescriptionProvider.h"
 #include "TargetInterfaces/SkeletalMeshBackedTarget.h"
+#include "TargetInterfaces/SkeletonCommitter.h"
+#include "TargetInterfaces/SkeletonProvider.h"
 #include "ToolTargets/PrimitiveComponentToolTarget.h"
 
 #include "SkeletalMeshComponentToolTarget.generated.h"
@@ -30,7 +32,8 @@ class USkeletalMeshComponentReadOnlyToolTarget :
 	public IMeshDescriptionProvider,
 	public IDynamicMeshProvider, 
 	public IMaterialProvider,
-	public ISkeletalMeshBackedTarget
+	public ISkeletalMeshBackedTarget,
+	public ISkeletonProvider
 {
 	GENERATED_BODY()
 
@@ -60,6 +63,9 @@ public:
 	// ISkeletalMeshBackedTarget implementation
 	UE_API virtual USkeletalMesh* GetSkeletalMesh() const override;
 
+	// ISkeletonProvider implementation
+	UE_API virtual FReferenceSkeleton GetSkeleton() const override;
+	
 protected:
 	// The LOD to edit if requested to get/commit at EMeshLODIdentifier::Default or if no specific LOD is requested.
 	int32 DefaultEditingLOD = 0;	
@@ -78,7 +84,8 @@ UCLASS(MinimalAPI, Transient)
 class USkeletalMeshComponentToolTarget :
 	public USkeletalMeshComponentReadOnlyToolTarget,
 	public IMeshDescriptionCommitter,
-	public IDynamicMeshCommitter
+	public IDynamicMeshCommitter,
+	public ISkeletonCommitter
 {
 	GENERATED_BODY()
 
@@ -91,6 +98,10 @@ public:
 	UE_API virtual void CommitDynamicMesh(const UE::Geometry::FDynamicMesh3& Mesh, const FDynamicMeshCommitInfo& CommitInfo) override;
 	using IDynamicMeshCommitter::CommitDynamicMesh; // unhide the other overload
 
+	// ISkeletonCommitter
+	UE_API virtual void SetupSkeletonModifier(USkeletonModifier* InModifier) override;
+	UE_API virtual void CommitSkeletonModifier(USkeletonModifier* InModifier) override;
+	
 protected:
 	// So that the tool target factory can poke into Component.
 	friend class USkeletalMeshComponentToolTargetFactory;
diff --git a/Engine/Plugins/Runtime/MeshModelingToolset/Source/SkeletalMeshModifiers/Private/SkeletonModifier.cpp b/Engine/Plugins/Runtime/MeshModelingToolset/Source/SkeletalMeshModifiers/Private/SkeletonModifier.cpp
--- a/Engine/Plugins/Runtime/MeshModelingToolset/Source/SkeletalMeshModifiers/Private/SkeletonModifier.cpp
+++ b/Engine/Plugins/Runtime/MeshModelingToolset/Source/SkeletalMeshModifiers/Private/SkeletonModifier.cpp
@@ -23,9 +23,14 @@
 #include "Misc/MessageDialog.h"
 #include "Widgets/Layout/SBox.h"
 #include "Widgets/Text/STextBlock.h"
+#include "Runtime/GeometryFramework/Public/UDynamicMesh.h"
+#include "DynamicMesh/DynamicBoneAttribute.h"
+#include "DynamicMesh/DynamicVertexSkinWeightsAttribute.h"
+#include "ScopedTransaction.h"
 
 #include UE_INLINE_GENERATED_CPP_BY_NAME(SkeletonModifier)
 
+#define LOCTEXT_NAMESPACE "SkeletonModifier"
 
 namespace USkeletonModifierLocals
 {
@@ -348,6 +353,25 @@ void USkeletonModifier::ExternalUpdate(const FReferenceSkeleton& InRefSkeleton,
 	BoneIndexTracker = InIndexTracker;
 }
 
+void USkeletonModifier::ExternalUpdate(const FReferenceSkeleton& InRefSkeleton)
+{
+	if (!ReferenceSkeleton)
+	{
+		return;
+	}
+	
+	*ReferenceSkeleton = InRefSkeleton;
+	
+	TransformComposer.Reset(new FTransformComposer(*ReferenceSkeleton));
+
+	int32 NumBones = InRefSkeleton.GetRawBoneNum();
+	BoneIndexTracker.Reset(NumBones);
+	for (int32 Index = 0; Index < NumBones; Index++)
+	{
+		BoneIndexTracker.Add(Index);
+	}
+}
+
 bool USkeletonModifier::SetSkeletalMesh(USkeletalMesh* InSkeletalMesh)
 {
 	SkeletalMesh = nullptr;
@@ -412,6 +436,162 @@ bool USkeletonModifier::SetSkeletalMesh(USkeletalMesh* InSkeletalMesh)
 #endif
 }
 
+bool USkeletonModifier::SetDynamicMesh(UDynamicMesh* InDynamicMesh)
+{
+#if WITH_EDITORONLY_DATA
+	
+	if (!InDynamicMesh)
+	{
+		return false;
+	}
+	
+	WeakDynamicMesh = InDynamicMesh;
+
+	ReferenceSkeleton.Reset();
+	TransformComposer.Reset();
+	BoneIndexTracker.Reset();
+
+	// store reference skeleton to edit
+	ReferenceSkeleton = MakeUnique<FReferenceSkeleton>();
+	{
+		const FDynamicMesh3& Mesh = InDynamicMesh->GetMeshRef();
+		
+		const int32 NumBones = Mesh.Attributes()->GetNumBones();
+		FReferenceSkeletonModifier RefSkeletonModifier(*ReferenceSkeleton, nullptr);
+
+		using namespace UE::Geometry;
+		const FDynamicMeshBoneNameAttribute* BoneNameAttrib = Mesh.Attributes()->GetBoneNames();
+		const FDynamicMeshBoneParentIndexAttribute* BoneParentIndexAttrib = Mesh.Attributes()->GetBoneParentIndices();
+		const FDynamicMeshBonePoseAttribute* BonePoses = Mesh.Attributes()->GetBonePoses();
+
+		for (int32 BoneIndex = 0; BoneIndex < NumBones; ++BoneIndex)
+		{
+			const int32 ParentBoneIndex = BoneParentIndexAttrib->GetValue(BoneIndex);
+			FName BoneName = BoneNameAttrib->GetValue(BoneIndex);
+			FMeshBoneInfo BoneInfo(BoneName, BoneName.ToString(), ParentBoneIndex);
+
+			FTransform Transform = BonePoses->GetValue(BoneIndex);
+			RefSkeletonModifier.Add(BoneInfo, Transform);
+		}	
+	}
+	
+
+	TransformComposer.Reset(new FTransformComposer(*ReferenceSkeleton));
+
+	// store initial bones indices to track for changes
+	const int32 NumBones = ReferenceSkeleton->GetRawBoneNum();
+	BoneIndexTracker.Reserve(NumBones);
+	for (int32 Index = 0; Index < NumBones; Index++)
+	{
+		BoneIndexTracker.Add(Index);
+	}
+	return true;
+	
+#else
+	ensureMsgf(false, TEXT("Skeleton Modifier is an editor only feature."));
+	return false;
+#endif
+
+}
+
+bool USkeletonModifier::CommitSkeletonToDynamicMesh()
+{
+	if (!WeakDynamicMesh.IsValid())
+	{
+		return false;
+	}
+	
+	UDynamicMesh* DynamicMesh = WeakDynamicMesh.Get();
+	
+	TSharedPtr<FDynamicMesh3> OldMeshShared = MakeShared<FDynamicMesh3>();
+	DynamicMesh->ProcessMesh([&](const FDynamicMesh3& ReadMesh) { *OldMeshShared = ReadMesh; });
+
+	DynamicMesh->EditMesh([&](FDynamicMesh3& Mesh)
+	{
+		const int32 NumBones = ReferenceSkeleton->GetNum();
+		
+		using namespace UE::Geometry;
+		Mesh.Attributes()->EnableBones(NumBones);
+		
+		FDynamicMeshBoneNameAttribute* BoneNameAttrib = Mesh.Attributes()->GetBoneNames();
+		FDynamicMeshBoneParentIndexAttribute* BoneParentIndexAttrib = Mesh.Attributes()->GetBoneParentIndices();
+		FDynamicMeshBonePoseAttribute* BonePoses = Mesh.Attributes()->GetBonePoses();
+
+		for (int32 BoneIndex = 0; BoneIndex < NumBones; ++BoneIndex)
+		{
+			FName BoneName = ReferenceSkeleton->GetBoneName(BoneIndex);
+			int32 ParentIndex = ReferenceSkeleton->GetParentIndex(BoneIndex);
+			const FTransform& RefBonePose = ReferenceSkeleton->GetRefBonePose()[BoneIndex];
+
+			BoneNameAttrib->SetValue(BoneIndex, BoneName);
+			BoneParentIndexAttrib->SetValue(BoneIndex, ParentIndex);
+			BonePoses->SetValue(BoneIndex, RefBonePose);
+		}	
+	});
+
+	// Update skin weights
+	DynamicMesh->EditMesh([&](FDynamicMesh3& Mesh)
+	{
+		using namespace UE::Geometry;
+		using namespace UE::AnimationCore;
+		
+		FBoneWeightsSettings BoneSettings; BoneSettings.SetNormalizeType(EBoneWeightNormalizeType::None);
+		
+		TArray<FName> SkinWeightProfiles;
+		Mesh.Attributes()->GetSkinWeightsAttributes().GenerateKeyArray(SkinWeightProfiles);
+
+		for (const FName SkinWeightProfile: SkinWeightProfiles)
+		{
+			FDynamicMeshVertexSkinWeightsAttribute* SkinWeights = Mesh.Attributes()->GetSkinWeightsAttribute(SkinWeightProfile);
+
+			for (int32 VertexID = 0 ; VertexID < Mesh.VertexCount(); VertexID++ )
+			{
+				FBoneWeights BoneWeights;
+				SkinWeights->GetValue(VertexID, BoneWeights);
+				
+				if (const int32 NumBoneWeights = BoneWeights.Num())
+				{
+					TArray<FBoneWeight> NewWeights;
+					for (int32 Idx = 0; Idx < NumBoneWeights; ++Idx)
+					{
+						const FBoneWeight& OldBoneWeight = BoneWeights[Idx];
+						const int32 BoneIndex = OldBoneWeight.GetBoneIndex();
+
+						int32 NewBoneIndex = 0;
+						if (ensure(BoneIndexTracker.IsValidIndex(BoneIndex)))
+						{
+							NewBoneIndex = BoneIndexTracker[BoneIndex];							
+						}
+						else
+						{
+							UE_LOG(LogAnimation, Warning, TEXT("Skeleton Modifier - Commit: Invalid bone index provided (%d); falling back to 0 as bone index."), BoneIndex);
+						}
+
+						if (NewBoneIndex != INDEX_NONE)
+						{
+							NewWeights.Add(FBoneWeight(NewBoneIndex, OldBoneWeight.GetRawWeight()));
+						}
+					}
+					SkinWeights->SetValue(VertexID, FBoneWeights::Create(NewWeights, BoneSettings));
+				}
+			}
+		}	
+	});
+
+	
+	TSharedPtr<FDynamicMesh3> NewMeshShared = MakeShared<FDynamicMesh3>();
+	DynamicMesh->ProcessMesh([&](const FDynamicMesh3& ReadMesh) { *NewMeshShared = ReadMesh; });
+	
+	TUniquePtr<FMeshReplacementChange> ReplaceChange = MakeUnique<FMeshReplacementChange>(OldMeshShared, NewMeshShared);
+
+	FScopedTransaction Transaction(LOCTEXT("SkeletonModifierCommitDynamicMesh", "Update Skeleton"));
+	check(GUndo != nullptr);
+	DynamicMesh->Modify();
+	GUndo->StoreUndo(DynamicMesh, MoveTemp(ReplaceChange));
+	
+	return true;
+}
+
 bool USkeletonModifier::IsReferenceSkeletonValid(const bool bLog) const
 {
 	if (!ReferenceSkeleton.IsValid())
@@ -2170,3 +2350,4 @@ const FTransform& USkeletonModifier::GetTransform(const int32 InBoneIndex, const
 }
 
 
+#undef LOCTEXT_NAMESPACE
\ No newline at end of file
diff --git a/Engine/Plugins/Runtime/MeshModelingToolset/Source/SkeletalMeshModifiers/Public/SkeletonModifier.h b/Engine/Plugins/Runtime/MeshModelingToolset/Source/SkeletalMeshModifiers/Public/SkeletonModifier.h
--- a/Engine/Plugins/Runtime/MeshModelingToolset/Source/SkeletalMeshModifiers/Public/SkeletonModifier.h
+++ b/Engine/Plugins/Runtime/MeshModelingToolset/Source/SkeletalMeshModifiers/Public/SkeletonModifier.h
@@ -18,6 +18,7 @@ class FName;
 struct FMeshDescription;
 struct FReferenceSkeleton;
 struct FTransformComposer;
+class UDynamicMesh;
 
 enum class ESkeletalMeshModificationType : uint32
 {
@@ -131,7 +132,24 @@ public:
 
 	UFUNCTION(BlueprintCallable, Category = "Mesh")
 	UE_API bool SetSkeletalMesh(USkeletalMesh* InSkeletalMesh);
+	
+	/**
+	 * Applies the skeleton modifications to the skeletal mesh.
+	 * @return true if commit succeeded.
+	 */
+	UFUNCTION(BlueprintCallable, Category = "Mesh")
+	UE_API bool CommitSkeletonToSkeletalMesh();
+
+
+	UE_API bool SetDynamicMesh(UDynamicMesh* InDynamicMesh);
+	
+	/**
+	 * Applies the skeleton modifications to the dynamic mesh.
+	 * @return true if commit succeeded.
+	 */
+	UE_API bool CommitSkeletonToDynamicMesh();
 
+	
 	/** Creates a new bone in the skeleton hierarchy at desired transform
 	 *  @param InBoneName The new bone's name.
 	 *  @param InParentName The new bone parent's name. 
@@ -233,14 +251,10 @@ public:
 	UFUNCTION(BlueprintCallable, Category = "Skeleton")
 	UE_API TArray<FName> GetAllBoneNames() const;
 	
-	/**
-	 * Actually applies the skeleton modifications to the skeletal mesh.
-	 * @return true if commit succeeded.
-	 */
-	UFUNCTION(BlueprintCallable, Category = "Mesh")
-	UE_API bool CommitSkeletonToSkeletalMesh();
+
 
 	UE_API void ExternalUpdate(const FReferenceSkeleton& InRefSkeleton, const TArray<int32>& InIndexTracker);
+	UE_API void ExternalUpdate(const FReferenceSkeleton& InRefSkeleton);
 	
 	UE_API FName GetUniqueName(const FName InBoneName, const TArray<FName>& InBoneNames = {}) const;
 	
@@ -283,6 +297,10 @@ private:
 	TWeakObjectPtr<USkeletalMesh> SkeletalMesh = nullptr;
 	
 	TUniquePtr<FMeshDescription> MeshDescription;
+
+	UPROPERTY(transient)
+	TWeakObjectPtr<UDynamicMesh> WeakDynamicMesh = nullptr;
+
 	TUniquePtr<FReferenceSkeleton> ReferenceSkeleton;
 	TUniquePtr<FTransformComposer> TransformComposer;
 	TArray<int32> BoneIndexTracker;
diff --git a/Engine/Plugins/Runtime/MeshModelingToolset/Source/SkeletalMeshModifiers/SkeletalMeshModifiers.build.cs b/Engine/Plugins/Runtime/MeshModelingToolset/Source/SkeletalMeshModifiers/SkeletalMeshModifiers.build.cs
--- a/Engine/Plugins/Runtime/MeshModelingToolset/Source/SkeletalMeshModifiers/SkeletalMeshModifiers.build.cs
+++ b/Engine/Plugins/Runtime/MeshModelingToolset/Source/SkeletalMeshModifiers/SkeletalMeshModifiers.build.cs
@@ -30,7 +30,9 @@ public class SkeletalMeshModifiers : ModuleRules
 				"SlateCore",
 				"Slate",
 				"UnrealEd",
-				"ToolWidgets"
+				"ToolWidgets",
+				"GeometryFramework",
+				"GeometryCore"
 			}
 			);
 
