[MeshModelingToolSet] minor fixes to skeleton modifiier dyna mesh path

diff --git a/Engine/Plugins/Runtime/MeshModelingToolset/Source/ModelingComponentsEditorOnly/Public/ToolTargets/DynamicMeshComponentToolTarget.h b/Engine/Plugins/Runtime/MeshModelingToolset/Source/ModelingComponentsEditorOnly/Public/ToolTargets/DynamicMeshComponentToolTarget.h
--- a/Engine/Plugins/Runtime/MeshModelingToolset/Source/ModelingComponentsEditorOnly/Public/ToolTargets/DynamicMeshComponentToolTarget.h
+++ b/Engine/Plugins/Runtime/MeshModelingToolset/Source/ModelingComponentsEditorOnly/Public/ToolTargets/DynamicMeshComponentToolTarget.h
@@ -78,7 +78,7 @@ public:
 	// ISkeletonProvider implementation
 	UE_API virtual FReferenceSkeleton GetSkeleton() const override;
 	
-	// ISkeletonComitter implementation
+	// ISkeletonCommitter implementation
 	UE_API virtual void SetupSkeletonModifier(USkeletonModifier* InModifier) override;
 	UE_API virtual void CommitSkeletonModifier(USkeletonModifier* InModifier) override;
 	
diff --git a/Engine/Plugins/Runtime/MeshModelingToolset/Source/SkeletalMeshModifiers/Private/SkeletonModifier.cpp b/Engine/Plugins/Runtime/MeshModelingToolset/Source/SkeletalMeshModifiers/Private/SkeletonModifier.cpp
--- a/Engine/Plugins/Runtime/MeshModelingToolset/Source/SkeletalMeshModifiers/Private/SkeletonModifier.cpp
+++ b/Engine/Plugins/Runtime/MeshModelingToolset/Source/SkeletalMeshModifiers/Private/SkeletonModifier.cpp
@@ -544,7 +544,7 @@ bool USkeletonModifier::CommitSkeletonToDynamicMesh()
 		{
 			FDynamicMeshVertexSkinWeightsAttribute* SkinWeights = Mesh.Attributes()->GetSkinWeightsAttribute(SkinWeightProfile);
 
-			for (int32 VertexID = 0 ; VertexID < Mesh.VertexCount(); VertexID++ )
+			for (int32 VertexID : Mesh.VertexIndicesItr())
 			{
 				FBoneWeights BoneWeights;
 				SkinWeights->GetValue(VertexID, BoneWeights);
