Dataflow Skin Weight Paint tool : add support for backface hit test

diff --git a/Engine/Plugins/Experimental/Dataflow/Source/DataflowEditor/Private/DataflowEditorTools/DataflowEditorSkinWeightsPaintTool.cpp b/Engine/Plugins/Experimental/Dataflow/Source/DataflowEditor/Private/DataflowEditorTools/DataflowEditorSkinWeightsPaintTool.cpp
--- a/Engine/Plugins/Experimental/Dataflow/Source/DataflowEditor/Private/DataflowEditorTools/DataflowEditorSkinWeightsPaintTool.cpp
+++ b/Engine/Plugins/Experimental/Dataflow/Source/DataflowEditor/Private/DataflowEditorTools/DataflowEditorSkinWeightsPaintTool.cpp
@@ -214,6 +214,12 @@ void UDataflowEditorSkinWeightsPaintTool::Setup()
 		}
 	}
 	USkinWeightsPaintTool::Setup();
+
+	// enable back face option 
+	if (WeightToolProperties)
+	{
+		WeightToolProperties->bShowHitBackFaces = true;
+	}
 }
 
 void UDataflowEditorSkinWeightsPaintTool::OnShutdown(EToolShutdownType ShutdownType)
diff --git a/Engine/Plugins/Runtime/MeshModelingToolset/Source/MeshModelingToolsEditorOnly/Private/SkeletalMesh/SkinWeightsPaintTool.cpp b/Engine/Plugins/Runtime/MeshModelingToolset/Source/MeshModelingToolsEditorOnly/Private/SkeletalMesh/SkinWeightsPaintTool.cpp
--- a/Engine/Plugins/Runtime/MeshModelingToolset/Source/MeshModelingToolsEditorOnly/Private/SkeletalMesh/SkinWeightsPaintTool.cpp
+++ b/Engine/Plugins/Runtime/MeshModelingToolset/Source/MeshModelingToolsEditorOnly/Private/SkeletalMesh/SkinWeightsPaintTool.cpp
@@ -2334,16 +2334,26 @@ bool USkinWeightsPaintTool::HitTest(const FRay& Ray, FHitResult& OutHit)
 	FViewCameraState StateOut;
 	GetToolManager()->GetContextQueriesAPI()->GetCurrentViewState(StateOut);
 	FVector3d LocalEyePosition(CurTargetTransform.InverseTransformPosition((FVector3d)StateOut.Position));
-	const int32 TriID = TrianglesOctree->FindNearestHitObject(
-		LocalRay,
-		[this, Mesh, &LocalEyePosition](int TriangleID)
-	{
-		FVector3d Normal, Centroid;
-		double Area;
-		Mesh->GetTriInfo(TriangleID, Normal, Area, Centroid);
-		return Normal.Dot((Centroid - LocalEyePosition)) < 0;
-	});
-	
+
+	int32 TriID = IndexConstants::InvalidID;
+
+	if (BrushProperties && BrushProperties->bHitBackFaces)
+	{
+		TriID = TrianglesOctree->FindNearestHitObject(LocalRay);
+	}
+	else
+	{
+		TriID = TrianglesOctree->FindNearestHitObject(
+			LocalRay,
+			[this, Mesh, &LocalEyePosition](int TriangleID)
+			{
+				FVector3d Normal, Centroid;
+				double Area;
+				Mesh->GetTriInfo(TriangleID, Normal, Area, Centroid);
+				return Normal.Dot((Centroid - LocalEyePosition)) < 0;
+			});
+	}
+
 	if (TriID != IndexConstants::InvalidID)
 	{	
 		FastTriWinding::FTriangle3d Triangle;
diff --git a/Engine/Source/Runtime/InteractiveToolsFramework/Public/BaseTools/BaseBrushTool.h b/Engine/Source/Runtime/InteractiveToolsFramework/Public/BaseTools/BaseBrushTool.h
--- a/Engine/Source/Runtime/InteractiveToolsFramework/Public/BaseTools/BaseBrushTool.h
+++ b/Engine/Source/Runtime/InteractiveToolsFramework/Public/BaseTools/BaseBrushTool.h
@@ -60,6 +60,11 @@ public:
 		ModelingQuickEdit, ModelingQuickSettings = 300))
 	float BrushFalloffAmount;
 
+
+	/** Allow the Brush to hit the back-side of the mesh */
+	UPROPERTY(EditAnywhere, Category = Brush, meta = (DisplayPriority = 6, HideEditConditionToggle, EditConditionHides, EditCondition = "bShowHitBackFaces"))
+	bool bHitBackFaces = false;
+
 	/** If false, then BrushStrength will not be shown in DetailsView panels (otherwise no effect) */
 	UPROPERTY( meta = (TransientToolProperty))
 	bool bShowStrength = true;
@@ -67,6 +72,10 @@ public:
 	/** If false, then BrushFalloffAmount will not be shown in DetailsView panels (otherwise no effect) */
 	UPROPERTY(meta = (TransientToolProperty))
 	bool bShowFalloff = true;
+
+	/** If false, then bHitBackFaces will not be shown in DetailsView panels (otherwise no effect) */
+	UPROPERTY(meta = (TransientToolProperty))
+	bool bShowHitBackFaces = false;
 };
 
 
