ModelingTools: Added support for a HitBackfaces option to the MeshAttributePaintTool.

diff --git a/Engine/Plugins/Runtime/MeshModelingToolset/Source/MeshModelingTools/Private/MeshAttributePaintTool.cpp b/Engine/Plugins/Runtime/MeshModelingToolset/Source/MeshModelingTools/Private/MeshAttributePaintTool.cpp
--- a/Engine/Plugins/Runtime/MeshModelingToolset/Source/MeshModelingTools/Private/MeshAttributePaintTool.cpp
+++ b/Engine/Plugins/Runtime/MeshModelingToolset/Source/MeshModelingTools/Private/MeshAttributePaintTool.cpp
@@ -322,6 +322,7 @@ void UMeshAttributePaintTool::Setup()
 	if (ensure(BrushProperties))
 	{
 		BrushProperties->bToolSupportsPressureSensitivity = true;
+		BrushProperties->bShowHitBackFaces = true;
 	}
 
 	// hide strength and falloff
@@ -546,13 +547,11 @@ bool UMeshAttributePaintTool::OnUpdateHover(const FInputDeviceRay& DevicePos)
 
 bool UMeshAttributePaintTool::HitTest(const FRay& Ray, FHitResult& OutHit)
 {
-	TFunction<bool(int32)> Filter = nullptr;
-	if (ShouldFilterTriangles())
-	{
-		Filter = [this](int32 Tid) { return SelectionTids.Contains(Tid); };
-	}
-		
-	return PreviewMesh->FindRayIntersection(FRay3d(Ray), OutHit, Filter);
+	return PreviewMesh->FindRayIntersection(FRay3d(Ray), OutHit, [this](int32 Tid)
+		{
+			return (IsFilteringTrianglesBySelection() ? SelectionTids.Contains(Tid) : true) &&
+				(IsFilteringTrianglesByFrontFacing() ? !IsTriangleBackFacing(Tid, PreviewMesh->GetMesh()) : true);
+		});
 }
 
 
@@ -566,12 +565,33 @@ void UMeshAttributePaintTool::CalculateVertexROI(const FBrushStampData& Stamp, T
 	const FDynamicMesh3* Mesh = PreviewMesh->GetPreviewDynamicMesh();
 	FAxisAlignedBox3d QueryBox(StampPosLocal, Radius);
 	
-	if (ShouldFilterTriangles() && ensure(!SelectionVids.IsEmpty()))
+	if (IsFilteringTrianglesBySelection() || IsFilteringTrianglesByFrontFacing())
 	{
+		FVector3d LocalEyePosition;
+		if (IsFilteringTrianglesByFrontFacing())
+		{
+			FViewCameraState StateOut;
+			GetToolManager()->GetContextQueriesAPI()->GetCurrentViewState(StateOut);
+			LocalEyePosition = (CurrentTargetTransform.InverseTransformPosition((FVector3d)StateOut.Position));
+		}
+		
 		VerticesOctree.RangeQuery(QueryBox,
-			[this, Mesh, &StampPosLocal, RadiusSqr](int32 VertexID)
-		{ 
-			return SelectionVids.Contains(VertexID) 
+			[this, Mesh, &StampPosLocal, RadiusSqr, &LocalEyePosition](int32 VertexID)
+		{
+			const bool bIsSelectionValid = !IsFilteringTrianglesBySelection() || SelectionVids.Contains(VertexID);
+			bool bIsFrontFaceValid = !IsFilteringTrianglesByFrontFacing();
+			if (!bIsFrontFaceValid)
+			{
+				bIsFrontFaceValid = true;
+				Mesh->EnumerateVertexTriangles(VertexID, [this, Mesh, &LocalEyePosition, &bIsFrontFaceValid](int32 TriID)
+				{
+					FVector3d Normal, Centroid;
+					double Area;
+					Mesh->GetTriInfo(TriID, Normal, Area, Centroid);
+					bIsFrontFaceValid = bIsFrontFaceValid && Normal.Dot((Centroid - LocalEyePosition)) < 0;
+				});
+			}
+			return bIsSelectionValid && bIsFrontFaceValid
 				&& DistanceSquared(Mesh->GetVertex(VertexID), StampPosLocal) < RadiusSqr; 
 		}, VertexROI);
 	}
@@ -884,7 +904,7 @@ void UMeshAttributePaintTool::ApplyStamp_FloodFill(const FBrushStampData& Stamp,
 	TSet<int32> RemainingTriangles;
 	for (int32 vid : ActionData.ROIVertices)
 	{
-		if (ShouldFilterTriangles())
+		if (IsFilteringTrianglesBySelection())
 		{
 			CurrentMesh->EnumerateVertexTriangles(vid, [&](int32 tid) 
 			{ 
@@ -919,7 +939,7 @@ void UMeshAttributePaintTool::ApplyStamp_FloodFill(const FBrushStampData& Stamp,
 		// get a single set element via an iterator
 		InputTriROI.Add(*RemainingTriangles.CreateConstIterator());
 		TFunction<bool(int32, int32)> TidFilter = [](int32, int32){ return true; };
-		if (ShouldFilterTriangles())
+		if (IsFilteringTrianglesBySelection())
 		{
 			TidFilter = [this](int32 Index, int32 Tid) { return SelectionTids.Contains(Tid); };
 		}
@@ -1079,12 +1099,31 @@ void UMeshAttributePaintTool::ExternalUpdateValues(int32 AttribIndex, const TArr
 	}
 }
 
-bool UMeshAttributePaintTool::ShouldFilterTriangles() const
+bool UMeshAttributePaintTool::IsFilteringTrianglesBySelection() const
+{
+	return BrushActionProps->bIsolateGeometrySelection && !SelectionTids.IsEmpty();
+}
+
+bool UMeshAttributePaintTool::IsFilteringTrianglesByFrontFacing() const
 {
-	return BrushActionProps->bIsolateGeometrySelection 
-		// Could check GeometrySelection.IsSet(), but might as well check SelectionTids instead
-		//  which will be nonempty if there was a selection, since that's what we'll actually use.
-		&& !SelectionTids.IsEmpty();
+	return !BrushProperties->bHitBackFaces;
+}
+
+bool UMeshAttributePaintTool::IsTriangleBackFacing(int32 TriangleID, const FDynamicMesh3* QueryMesh) const
+{
+	if (TriangleID != IndexConstants::InvalidID)
+	{
+		FViewCameraState StateOut;
+		GetToolManager()->GetContextQueriesAPI()->GetCurrentViewState(StateOut);
+		FVector3d LocalEyePosition(CurrentTargetTransform.InverseTransformPosition((FVector3d)StateOut.Position));
+
+		FVector3d Normal, Centroid;
+		double Area;
+		QueryMesh->GetTriInfo(TriangleID, Normal, Area, Centroid);
+
+		return (Normal.Dot((Centroid - LocalEyePosition)) >= 0);
+	}
+	return false;
 }
 
 
diff --git a/Engine/Plugins/Runtime/MeshModelingToolset/Source/MeshModelingTools/Public/MeshAttributePaintTool.h b/Engine/Plugins/Runtime/MeshModelingToolset/Source/MeshModelingTools/Public/MeshAttributePaintTool.h
--- a/Engine/Plugins/Runtime/MeshModelingToolset/Source/MeshModelingTools/Public/MeshAttributePaintTool.h
+++ b/Engine/Plugins/Runtime/MeshModelingToolset/Source/MeshModelingTools/Public/MeshAttributePaintTool.h
@@ -386,12 +386,15 @@ private:
 	TOptional<UE::Geometry::FGeometrySelection> GeometrySelection;
 	TSet<int32> SelectionTids;
 	TSet<int32> SelectionVids;
-	UE_API bool ShouldFilterTriangles() const;
+	bool IsFilteringTrianglesBySelection() const;
+	bool IsFilteringTrianglesByFrontFacing() const;
 	
 	void ApplyStamp_Paint(const FBrushStampData& Stamp, FStampActionData& ActionData, const FBrushActionData& BrushData);
 
 	void UpdateMaterialMode(EMeshAttributePaintMaterialMode MaterialMode);
 	void UpdateFlatShadingSetting(bool bNewValue);
+
+	bool IsTriangleBackFacing(int32 TriangleID, const FDynamicMesh3* QueryMesh) const;
 };
 
 
