Mesh selector : add the option to get the direct selected vertex indices  without beng remapped to the original mesh
This is useful if the tool using it already takes care of the remapping

diff --git a/Engine/Plugins/Runtime/MeshModelingToolset/Source/MeshModelingToolsEditorOnly/Private/ToolMeshSelector.cpp b/Engine/Plugins/Runtime/MeshModelingToolset/Source/MeshModelingToolsEditorOnly/Private/ToolMeshSelector.cpp
--- a/Engine/Plugins/Runtime/MeshModelingToolset/Source/MeshModelingToolsEditorOnly/Private/ToolMeshSelector.cpp
+++ b/Engine/Plugins/Runtime/MeshModelingToolset/Source/MeshModelingToolsEditorOnly/Private/ToolMeshSelector.cpp
@@ -187,7 +187,7 @@ void UToolMeshSelector::Render(IToolsContextRenderAPI* RenderAPI)
 	PolygonSelectionMechanic->Render(RenderAPI);
 }
 
-const TArray<int32>& UToolMeshSelector::GetSelectedVertices()
+const TArray<int32>& UToolMeshSelector::GetSelectedVertices(bool bRemapToOriginal)
 {
 	using UE::Geometry::FNonManifoldMappingSupport;
 
@@ -203,17 +203,27 @@ const TArray<int32>& UToolMeshSelector::GetSelectedVertices()
 	const FNonManifoldMappingSupport NonManifoldMappingSupport(*DynamicMesh);
 
 	// validate and add vertices to the output array
-	auto AddVertices = [this, &AllVertices, &NonManifoldMappingSupport](const TSet<int32>& VerticesToAdd)
+	auto AddVertices = [this, bRemapToOriginal, &AllVertices, &NonManifoldMappingSupport](const TSet<int32>& VerticesToAdd)
 		{
-			// we have to make sure that the vertex ids are safe to use as PolygonSelectionMechanic does not act on the
-			// mesh description but on the dynamic mesh that can duplicate vertices when dealing with degenerate triangles.
-			// cf. FMeshDescriptionToDynamicMesh::Convert for more details.
-			for (const int32 VertexToAdd : VerticesToAdd)
+			if (bRemapToOriginal)
 			{
-				const int32 SrcVertexID = NonManifoldMappingSupport.GetOriginalNonManifoldVertexID(VertexToAdd);
-				if (AllVertices.IsValid(SrcVertexID))
+				// we have to make sure that the vertex ids are safe to use as PolygonSelectionMechanic does not act on the
+				// mesh description but on the dynamic mesh that can duplicate vertices when dealing with degenerate triangles.
+				// cf. FMeshDescriptionToDynamicMesh::Convert for more details.
+				for (const int32 VertexToAdd : VerticesToAdd)
 				{
-					SelectedVerticesInternal.Add(SrcVertexID);
+					const int32 SrcVertexID = NonManifoldMappingSupport.GetOriginalNonManifoldVertexID(VertexToAdd);
+					if (AllVertices.IsValid(SrcVertexID))
+					{
+						SelectedVerticesInternal.Add(SrcVertexID);
+					}
+				}
+			}
+			else
+			{
+				for (const int32 VertexToAdd : VerticesToAdd)
+				{
+					SelectedVerticesInternal.Add(VertexToAdd);
 				}
 			}
 		};
diff --git a/Engine/Plugins/Runtime/MeshModelingToolset/Source/MeshModelingToolsEditorOnly/Public/ToolMeshSelector.h b/Engine/Plugins/Runtime/MeshModelingToolset/Source/MeshModelingToolsEditorOnly/Public/ToolMeshSelector.h
--- a/Engine/Plugins/Runtime/MeshModelingToolset/Source/MeshModelingToolsEditorOnly/Public/ToolMeshSelector.h
+++ b/Engine/Plugins/Runtime/MeshModelingToolset/Source/MeshModelingToolsEditorOnly/Public/ToolMeshSelector.h
@@ -71,7 +71,7 @@ public:
 	UE_API void Render(IToolsContextRenderAPI* RenderAPI);
 
 	// get a list of currently selected vertices (converting edges and faces to vertices)
-	UE_API const TArray<int32>& GetSelectedVertices();
+	UE_API const TArray<int32>& GetSelectedVertices(bool bRemapToOriginal = true);
 	UE_API bool IsAnyComponentSelected() const;
 	UE_API void GetSelectedTriangles(TArray<int32>& OutTriangleIndices) const;
 
