ModelingTools: Fixed nullptr dereference in MeshSceneAdapter::AddComponents

diff --git a/Engine/Plugins/Runtime/MeshModelingToolset/Source/ModelingComponents/Private/Scene/MeshSceneAdapter.cpp b/Engine/Plugins/Runtime/MeshModelingToolset/Source/ModelingComponents/Private/Scene/MeshSceneAdapter.cpp
--- a/Engine/Plugins/Runtime/MeshModelingToolset/Source/ModelingComponents/Private/Scene/MeshSceneAdapter.cpp
+++ b/Engine/Plugins/Runtime/MeshModelingToolset/Source/ModelingComponents/Private/Scene/MeshSceneAdapter.cpp
@@ -1562,6 +1562,11 @@ void FMeshSceneAdapter::AddComponents(const TArray<UActorComponent*>& ComponentS
 	TArray<FActorAdapter*> NewActorAdapters;
 	for (UActorComponent* Component : ComponentSetIn)
 	{
+		if (Component == nullptr)
+		{
+			continue;
+		}
+		
 		TUniquePtr<FActorAdapter> Adapter = MakeUnique<FActorAdapter>();
 		Adapter->SourceActor = Component->GetOwner();
 		CollectComponentChildMeshes(Component, *Adapter);
