[SKMmodellingTools] add the same UI for setting current editing LOD for all tools similar to modeling mode(a drop down at the bottom left corner) + refactor skin weight paint tool to be dynamic mesh based.

diff --git a/Engine/Plugins/Animation/SkeletalMeshModelingTools/Source/SkeletalMeshModelingTools/Private/SkeletalMeshEditorUtils.cpp b/Engine/Plugins/Animation/SkeletalMeshModelingTools/Source/SkeletalMeshModelingTools/Private/SkeletalMeshEditorUtils.cpp
--- a/Engine/Plugins/Animation/SkeletalMeshModelingTools/Source/SkeletalMeshModelingTools/Private/SkeletalMeshEditorUtils.cpp
+++ b/Engine/Plugins/Animation/SkeletalMeshModelingTools/Source/SkeletalMeshModelingTools/Private/SkeletalMeshEditorUtils.cpp
@@ -3,14 +3,13 @@
 #include "SkeletalMeshEditorUtils.h"
 
 #include "ContextObjectStore.h"
+#include "EditorModeManager.h"
 #include "InteractiveToolsContext.h"
-#include "IPersonaToolkit.h"
-#include "ISkeletalMeshEditor.h"
 #include "ISkeletonEditorModule.h"
+#include "SkeletalMeshModelingToolsEditorMode.h"
 #include "SkeletonModifier.h"
 #include "SReferenceSkeletonTree.h"
 #include "Widgets/Docking/SDockTab.h"
-#include "Preferences/PersonaOptions.h"
 
 #include UE_INLINE_GENERATED_CPP_BY_NAME(SkeletalMeshEditorUtils)
 
@@ -77,49 +76,42 @@ void USkeletalMeshEditorContextObject::Unregister(UInteractiveToolManager* InToo
 	bRegistered = false;
 }
 
-void USkeletalMeshEditorContextObject::Init(const TWeakPtr<ISkeletalMeshEditor>& InEditor)
+void USkeletalMeshEditorContextObject::Init(USkeletalMeshModelingToolsEditorMode* InEditorMode)
 {
-	Editor = InEditor;
+	EditorMode = InEditorMode;
+	
 	EditorBindings.Reset();
 	TreeBindings.Reset();
 }
 
-void USkeletalMeshEditorContextObject::HideSkeleton()
+EMeshLODIdentifier USkeletalMeshEditorContextObject::GetEditingLOD()
 {
-	if (!Editor.IsValid())
-	{
-		return;
-	}
+	return EditorMode->GetEditingLOD();
+}
 
-	const TSharedPtr<ISkeletalMeshEditor> SkeletalMeshEditor = Editor.Pin();
-	UDebugSkelMeshComponent* SkeletalMeshComponent = SkeletalMeshEditor->GetPersonaToolkit()->GetPreviewMeshComponent();
-	if (!SkeletalMeshComponent)
-	{
-		return;
-	}
-	
-	SkeletonDrawMode = SkeletalMeshComponent->SkeletonDrawMode;
-	SkeletalMeshComponent->SkeletonDrawMode = ESkeletonDrawMode::Hidden;
+void USkeletalMeshEditorContextObject::HideSkeleton()
+{
+	EditorMode->HideSkeletonForTool();
 }
 
 void USkeletalMeshEditorContextObject::ShowSkeleton()
 {
-	if (!Editor.IsValid())
-	{
-		return;
-	}
-	
-	const TSharedPtr<ISkeletalMeshEditor> SkeletalMeshEditor = Editor.Pin();
-	UDebugSkelMeshComponent* SkeletalMeshComponent = SkeletalMeshEditor->GetPersonaToolkit()->GetPreviewMeshComponent();
-	if (!SkeletalMeshComponent)
-	{
-		return;
-	}
-	
-	SkeletalMeshComponent->SkeletonDrawMode = SkeletonDrawMode;
-	SkeletonDrawMode = ESkeletonDrawMode::Default;
+	EditorMode->ShowSkeletonForTool();
+}
+
+void USkeletalMeshEditorContextObject::ToggleBoneManipulation(bool bEnable)
+{
+	EditorMode->ToggleBoneManipulation(bEnable);
+}
+
+const TArray<FTransform>& USkeletalMeshEditorContextObject::GetComponentSpaceBoneTransforms(UToolTarget* InToolTarget)
+{
+	return EditorMode->GetComponentSpaceBoneTransforms();
 }
 
+
+
+
 void USkeletalMeshEditorContextObject::BindTo(ISkeletalMeshEditingInterface* InEditingInterface)
 {
 	if (!InEditingInterface)
@@ -181,7 +173,7 @@ void USkeletalMeshEditorContextObject::UnbindInterfaceFrom(
 
 void USkeletalMeshEditorContextObject::BindEditor(ISkeletalMeshEditingInterface* InEditingInterface)
 {
-	if (!InEditingInterface || !Editor.IsValid())
+	if (!InEditingInterface || !EditorMode.IsValid())
 	{
 		return;
 	}
@@ -191,7 +183,7 @@ void USkeletalMeshEditorContextObject::BindEditor(ISkeletalMeshEditingInterface*
 		return;
 	}
 
-	TSharedPtr<ISkeletalMeshEditorBinding> Binding = Editor.Pin()->GetBinding();
+	TSharedPtr<ISkeletalMeshEditorBinding> Binding = EditorMode.Pin()->GetEditorBinding();
 	if (!Binding.IsValid())
 	{
 		return;
@@ -215,9 +207,9 @@ void USkeletalMeshEditorContextObject::UnbindEditor(ISkeletalMeshEditingInterfac
 			InEditingInterface->Unbind();
 		}
 
-		if (Editor.IsValid())
+		if (EditorMode.IsValid())
 		{
-			const TSharedPtr<ISkeletalMeshEditorBinding> Binding = Editor.Pin()->GetBinding();
+			const TSharedPtr<ISkeletalMeshEditorBinding> Binding = EditorMode.Pin()->GetEditorBinding();
 			UnbindInterfaceFrom(InEditingInterface, Binding->GetNotifier(), *BindData);
 		}
 	
@@ -227,7 +219,7 @@ void USkeletalMeshEditorContextObject::UnbindEditor(ISkeletalMeshEditingInterfac
 
 void USkeletalMeshEditorContextObject::BindRefSkeletonTree(ISkeletalMeshEditingInterface* InEditingInterface)
 {
-	if (!InEditingInterface || !Editor.IsValid())
+	if (!InEditingInterface || !EditorMode.IsValid())
 	{
 		return;
 	}
@@ -238,7 +230,7 @@ void USkeletalMeshEditorContextObject::BindRefSkeletonTree(ISkeletalMeshEditingI
 		return;
 	}
 	
-	const TSharedPtr<FTabManager> TabManager = Editor.Pin()->GetAssociatedTabManager();
+	const TSharedPtr<FTabManager> TabManager = EditorMode->GetAssociatedTabManager();
 	TSharedPtr<SDockTab> SkeletonTab = TabManager->FindExistingLiveTab(GetSkeletonTreeTabId());
 	if (SkeletonTab.IsValid())
 	{
@@ -267,7 +259,7 @@ void USkeletalMeshEditorContextObject::BindRefSkeletonTree(ISkeletalMeshEditingI
 	TreeBindings.Emplace(InEditingInterface, BindData);
 
 	// unbind notifications on the initial SSkeletonTree to avoid useless notifications
-	if(const TSharedPtr<ISkeletalMeshEditorBinding> Binding = Editor.Pin()->GetBinding())
+	if(const TSharedPtr<ISkeletalMeshEditorBinding> Binding = EditorMode->GetEditorBinding())
 	{
 		const TArray<FName> SelectedBones = Binding->GetSelectedBones();
 		
@@ -292,20 +284,23 @@ void USkeletalMeshEditorContextObject::UnbindRefSkeletonTree(ISkeletalMeshEditin
 			UnbindInterfaceFrom(InEditingInterface, RefSkeletonTree->GetNotifier(), *BindData);
 		}
 
-		if (Editor.IsValid())
+		if (EditorMode.IsValid())
 		{
-			const TSharedPtr<FTabManager> TabManager = Editor.Pin()->GetAssociatedTabManager();
-			const TSharedPtr<SDockTab> SkeletonTab = TabManager->FindExistingLiveTab(GetSkeletonTreeTabId());
-			if (SkeletonTab.IsValid())
+			const TSharedPtr<FTabManager> TabManager = EditorMode->GetAssociatedTabManager();
+			if (TabManager.IsValid())
 			{
-				if (SkeletonTab->GetContent() == RefSkeletonWidget)
+				const TSharedPtr<SDockTab> SkeletonTab = TabManager->FindExistingLiveTab(GetSkeletonTreeTabId());
+				if (SkeletonTab.IsValid())
 				{
-					SkeletonTab->SetContent(DefaultSkeletonWidget.IsValid() ? DefaultSkeletonWidget.ToSharedRef() : SNullWidget::NullWidget);
-				}
+					if (SkeletonTab->GetContent() == RefSkeletonWidget)
+					{
+						SkeletonTab->SetContent(DefaultSkeletonWidget.IsValid() ? DefaultSkeletonWidget.ToSharedRef() : SNullWidget::NullWidget);
+					}
+				}	
 			}
 
 			// update the SSkeletonTree with the new RefSkeleton and set current selection 
-			if (const TSharedPtr<ISkeletalMeshEditorBinding> Binding = Editor.Pin()->GetBinding())
+			if (const TSharedPtr<ISkeletalMeshEditorBinding> Binding = EditorMode->GetEditorBinding())
 			{
 				ISkeletalMeshNotifier& Notifier = Binding->GetNotifier();
 				Notifier.HandleNotification(SelectedBones, ESkeletalMeshNotifyType::HierarchyChanged);
@@ -328,5 +323,5 @@ const FName& USkeletalMeshEditorContextObject::GetSkeletonTreeTabId()
 
 TSharedPtr<ISkeletalMeshEditorBinding> USkeletalMeshEditorContextObject::GetBinding() const
 {
-	return Editor.IsValid() ? Editor.Pin()->GetBinding() : nullptr;
+	return EditorMode.IsValid() ? EditorMode->GetEditorBinding() : nullptr;
 }
diff --git a/Engine/Plugins/Animation/SkeletalMeshModelingTools/Source/SkeletalMeshModelingTools/Private/SkeletalMeshEditorUtils.h b/Engine/Plugins/Animation/SkeletalMeshModelingTools/Source/SkeletalMeshModelingTools/Private/SkeletalMeshEditorUtils.h
--- a/Engine/Plugins/Animation/SkeletalMeshModelingTools/Source/SkeletalMeshModelingTools/Private/SkeletalMeshEditorUtils.h
+++ b/Engine/Plugins/Animation/SkeletalMeshModelingTools/Source/SkeletalMeshModelingTools/Private/SkeletalMeshEditorUtils.h
@@ -3,7 +3,7 @@
 #pragma once
 
 #include "Animation/DebugSkelMeshComponent.h"
-#include "SkeletalMesh/SkeletalMeshEditionInterface.h"
+#include "SkeletalMesh/SkeletalMeshEditingInterface.h"
 
 #include "UObject/Object.h"
 #include "UObject/WeakInterfacePtr.h"
@@ -12,9 +12,9 @@
 
 #define UE_API SKELETALMESHMODELINGTOOLS_API
 
+class USkeletalMeshModelingToolsEditorMode;
 class UInteractiveToolsContext;
 class UInteractiveToolManager;
-class ISkeletalMeshEditor;
 class USkeletalMeshEditorContextObject;
 class SReferenceSkeletonTree;
 
@@ -23,10 +23,10 @@ namespace UE
 	
 namespace SkeletalMeshEditorUtils
 {
-	SKELETALMESHMODELINGTOOLS_API bool RegisterEditorContextObject(UInteractiveToolsContext* InToolsContext);
-	SKELETALMESHMODELINGTOOLS_API bool UnregisterEditorContextObject(UInteractiveToolsContext* InToolsContext);
+	UE_API bool RegisterEditorContextObject(UInteractiveToolsContext* InToolsContext);
+	UE_API bool UnregisterEditorContextObject(UInteractiveToolsContext* InToolsContext);
 
-	SKELETALMESHMODELINGTOOLS_API USkeletalMeshEditorContextObject* GetEditorContextObject(UInteractiveToolsContext* InToolsContext);
+	UE_API USkeletalMeshEditorContextObject* GetEditorContextObject(UInteractiveToolsContext* InToolsContext);
 }
 	
 }
@@ -44,16 +44,21 @@ public:
 	UE_API void Register(UInteractiveToolManager* InToolManager);
 	UE_API void Unregister(UInteractiveToolManager* InToolManager);
 
-	UE_API void Init(const TWeakPtr<ISkeletalMeshEditor>& InEditor);
+	UE_API void Init(USkeletalMeshModelingToolsEditorMode* InEditorMode);
 
+	UE_API virtual EMeshLODIdentifier GetEditingLOD() override;
+	
 	UE_API virtual void HideSkeleton() override;
 	UE_API virtual void ShowSkeleton() override;
+	UE_API virtual void ToggleBoneManipulation(bool bEnable) override;
+	UE_API virtual const TArray<FTransform>& GetComponentSpaceBoneTransforms(UToolTarget* InToolTarget) override;
+
 
 	UE_API virtual void BindTo(ISkeletalMeshEditingInterface* InEditingInterface) override;
 	UE_API virtual void UnbindFrom(ISkeletalMeshEditingInterface* InEditingInterface) override;
 	
 private:
-	TWeakPtr<ISkeletalMeshEditor> Editor = nullptr;
+	TWeakObjectPtr<USkeletalMeshModelingToolsEditorMode> EditorMode = nullptr;
 
 	UPROPERTY()
 	ESkeletonDrawMode SkeletonDrawMode = ESkeletonDrawMode::Default;
diff --git a/Engine/Plugins/Animation/SkeletalMeshModelingTools/Source/SkeletalMeshModelingTools/Private/SkeletalMeshGizmoUtils.h b/Engine/Plugins/Animation/SkeletalMeshModelingTools/Source/SkeletalMeshModelingTools/Private/SkeletalMeshGizmoUtils.h
--- a/Engine/Plugins/Animation/SkeletalMeshModelingTools/Source/SkeletalMeshModelingTools/Private/SkeletalMeshGizmoUtils.h
+++ b/Engine/Plugins/Animation/SkeletalMeshModelingTools/Source/SkeletalMeshModelingTools/Private/SkeletalMeshGizmoUtils.h
@@ -4,7 +4,7 @@
 
 #include "UObject/Object.h"
 #include "Containers/UnrealString.h"
-#include "SkeletalMesh/SkeletalMeshEditionInterface.h"
+#include "SkeletalMesh/SkeletalMeshEditingInterface.h"
 
 #include "SkeletalMeshGizmoUtils.generated.h"
 
diff --git a/Engine/Plugins/Animation/SkeletalMeshModelingTools/Source/SkeletalMeshModelingTools/Private/SkeletalMeshModelingToolsEditorMode.cpp b/Engine/Plugins/Animation/SkeletalMeshModelingTools/Source/SkeletalMeshModelingTools/Private/SkeletalMeshModelingToolsEditorMode.cpp
--- a/Engine/Plugins/Animation/SkeletalMeshModelingTools/Source/SkeletalMeshModelingTools/Private/SkeletalMeshModelingToolsEditorMode.cpp
+++ b/Engine/Plugins/Animation/SkeletalMeshModelingTools/Source/SkeletalMeshModelingTools/Private/SkeletalMeshModelingToolsEditorMode.cpp
@@ -35,6 +35,7 @@
 #include "EditMeshPolygonsTool.h"
 #include "EditorInteractiveGizmoManager.h"
 #include "HoleFillTool.h"
+#include "IPersonaToolkit.h"
 #include "ISkeletalMeshEditor.h"
 #include "LatticeDeformerTool.h"
 #include "MeshAttributePaintTool.h"
@@ -58,7 +59,7 @@
 #include "SkeletalMeshNotifier.h"
 #include "Animation/DebugSkelMeshComponent.h"
 
-#include "SkeletalMesh/SkeletalMeshEditionInterface.h"
+#include "SkeletalMesh/SkeletalMeshEditingInterface.h"
 #include "SkeletalMesh/SkeletonEditingTool.h"
 #include "SkeletalMesh/SkinWeightsBindingTool.h"
 #include "SkeletalMesh/SkinWeightsPaintTool.h"
@@ -271,6 +272,9 @@ void USkeletalMeshModelingToolsEditorMode::Enter()
 	ToolSwitchModeToRestoreOnExit = GetInteractiveToolsContext()->ToolManager->GetToolSwitchMode();
 	// default to NOT applying changes to skeletal meshes when switching between tools without accepting
 	GetInteractiveToolsContext()->ToolManager->SetToolSwitchMode(EToolManagerToolSwitchMode::CancelIfAble);
+
+	GetInteractiveToolsContext()->ToolManager->OnToolPostBuild.AddUObject(this, &USkeletalMeshModelingToolsEditorMode::OnToolPostBuild);
+	GetInteractiveToolsContext()->ToolManager->OnToolEndedWithStatus.AddUObject(this, &USkeletalMeshModelingToolsEditorMode::OnToolEndedWithStatus);
 }
 
 UDebugSkelMeshComponent* USkeletalMeshModelingToolsEditorMode::GetSkelMeshComponent() const
@@ -284,6 +288,61 @@ UDebugSkelMeshComponent* USkeletalMeshModelingToolsEditorMode::GetSkelMeshCompon
 	return Cast<UDebugSkelMeshComponent>(SkeletalMeshComponent);
 }
 
+ESkeletonDrawMode USkeletalMeshModelingToolsEditorMode::GetSkeletalMeshSkeletonDrawMode() const
+{
+	if (UDebugSkelMeshComponent* SkeletalMeshComponent = GetSkelMeshComponent())
+	{
+		return SkeletalMeshComponent->SkeletonDrawMode;
+	}
+
+	return ESkeletonDrawMode::Default;
+}
+
+void USkeletalMeshModelingToolsEditorMode::SetSkeletalMeshSkeletonDrawMode(ESkeletonDrawMode DrawMode) const
+{
+	if (UDebugSkelMeshComponent* SkeletalMeshComponent = GetSkelMeshComponent())
+	{
+		SkeletalMeshComponent->SkeletonDrawMode = DrawMode;
+	}
+}
+
+ESkeletonDrawMode USkeletalMeshModelingToolsEditorMode::GetCurrentSkeletonDrawMode() const
+{
+	return GetSkeletalMeshSkeletonDrawMode();
+}
+
+
+TSharedPtr<FTabManager> USkeletalMeshModelingToolsEditorMode::GetAssociatedTabManager()
+{
+	if (Editor.IsValid())
+	{
+		return Editor.Pin()->GetAssociatedTabManager();
+	}
+
+	return nullptr;
+}
+
+const TArray<FTransform>& USkeletalMeshModelingToolsEditorMode::GetComponentSpaceBoneTransforms() const
+{
+	return GetSkelMeshComponent()->GetComponentSpaceTransforms();
+}
+
+void USkeletalMeshModelingToolsEditorMode::ToggleBoneManipulation(bool bEnable)
+{
+	if (Owner)
+	{
+		if (bEnable)
+		{
+			Owner->ActivateMode(FPersonaEditModes::SkeletonSelection);
+		}
+		else
+		{
+			Owner->DeactivateMode(FPersonaEditModes::SkeletonSelection);
+		}
+	}
+}
+
+
 void USkeletalMeshModelingToolsEditorMode::Exit()
 {
 	UEditorInteractiveToolsContext* InteractiveToolsContext = GetInteractiveToolsContext();
@@ -391,10 +450,9 @@ void USkeletalMeshModelingToolsEditorMode::Tick(FEditorViewportClient* ViewportC
 
 bool USkeletalMeshModelingToolsEditorMode::HandleClick(FEditorViewportClient* InViewportClient, HHitProxy *HitProxy, const FViewportClick &Click)
 {
-	if (Binding.IsValid())
+	// Making sure we can still select bones from viewport even if FPersonaEditModes::SkeletonSelection is deactivated
+	if (TSharedPtr<ISkeletalMeshEditorBinding> BindingPtr = GetEditorBinding())
 	{
-		TSharedPtr<ISkeletalMeshEditorBinding> BindingPtr = Binding.Pin();
-
 		TArray<FName> SelectedBones;
 		if (HitProxy && BindingPtr->GetNameFunction())
 		{
@@ -429,9 +487,9 @@ bool USkeletalMeshModelingToolsEditorMode::ComputeBoundingBoxForViewportFocus(AA
 	{
 		check(Component->GetSkeletalMeshAsset());
 		
-		if (Binding.IsValid())
+		if (EditorBinding.IsValid())
 		{
-			const TArray<FName> Selection = Binding.Pin()->GetSelectedBones();
+			const TArray<FName> Selection = EditorBinding.Pin()->GetSelectedBones();
 			if (!Selection.IsEmpty())
 			{
 				TArray<FName> AllChildren;
@@ -476,16 +534,17 @@ void USkeletalMeshModelingToolsEditorMode::OnToolStarted(UInteractiveToolManager
 	{
 		ExtensionToolCommandsGetter().BindCommandsForCurrentTool(Toolkit->GetToolkitCommands(), Tool);
 	}
+}
 
-	// deactivate SkeletonSelection when a tool is activated.
-	// each tool is responsible for activating SkeletonSelection if necessary
-	if (Owner)
-	{
-		Owner->DeactivateMode(FPersonaEditModes::SkeletonSelection);
-	}
+void USkeletalMeshModelingToolsEditorMode::OnToolPostBuild( UInteractiveToolManager* InToolManager, EToolSide InSide, UInteractiveTool* InBuiltTool,
+	UInteractiveToolBuilder* InToolBuilder, const FToolBuilderState& ToolState)
+{
+	// Tools can decide for themselves if they want to enable bone manipulation
+	ToggleBoneManipulation(false);
 }
 
-void USkeletalMeshModelingToolsEditorMode::OnToolEnded(UInteractiveToolManager* Manager, UInteractiveTool* Tool)
+void USkeletalMeshModelingToolsEditorMode::OnToolEndedWithStatus(UInteractiveToolManager* Manager, UInteractiveTool* Tool,
+	EToolShutdownType ShutdownType)
 {
 	FSkeletalMeshModelingToolsActionCommands::UpdateToolCommandBinding(Tool, Toolkit->GetToolkitCommands(), true);
 	if (ExtensionToolCommandsGetter)
@@ -494,11 +553,7 @@ void USkeletalMeshModelingToolsEditorMode::OnToolEnded(UInteractiveToolManager*
 		ExtensionToolCommandsGetter = nullptr;
 	}
 	
-	// reactivate SkeletonSelection when deactivating a tool
-	if (Owner)
-	{
-		Owner->ActivateMode(FPersonaEditModes::SkeletonSelection);
-	}
+	ToggleBoneManipulation(true);
 }
 
 bool USkeletalMeshModelingToolsEditorMode::ShouldToolStartBeAllowed(const FString& ToolIdentifier) const
@@ -515,37 +570,78 @@ void USkeletalMeshModelingToolsEditorMode::SetEditorBinding(const TWeakPtr<ISkel
 		return;
 	}
 	
-	Binding = InSkeletalMeshEditor.Pin()->GetBinding();
-
+	EditorBinding = InSkeletalMeshEditor.Pin()->GetBinding();
+	Editor = InSkeletalMeshEditor;
+	SkeletalMesh = InSkeletalMeshEditor.Pin()->GetPersonaToolkit()->GetMesh();
+	
 	if (USkeletalMeshEditorContextObject* ContextObject = UE::SkeletalMeshEditorUtils::GetEditorContextObject(GetInteractiveToolsContext()))
 	{
-		ContextObject->Init(InSkeletalMeshEditor);
+		ContextObject->Init(this);
 	}
+	
+	OnInitializedDelegate.Broadcast();
 }
 
-ISkeletalMeshEditingInterface* USkeletalMeshModelingToolsEditorMode::GetSkeletonInterface(UInteractiveTool* InTool)
+TSharedPtr<ISkeletalMeshEditorBinding> USkeletalMeshModelingToolsEditorMode::GetEditorBinding()
 {
-	if (!IsValid(InTool) || !InTool->Implements<USkeletalMeshEditingInterface>())
+	if (!EditorBinding.IsValid())
 	{
 		return nullptr;
 	}
-	return static_cast<ISkeletalMeshEditingInterface*>(InTool->GetInterfaceAddress(USkeletalMeshEditingInterface::StaticClass()));
+	
+	return EditorBinding.Pin();
+}
+
+USkeletalMesh* USkeletalMeshModelingToolsEditorMode::GetSkeletalMesh() const
+{
+	return SkeletalMesh.Get();
 }
 
-bool USkeletalMeshModelingToolsEditorMode::NeedsTransformGizmo() const
+bool USkeletalMeshModelingToolsEditorMode::CanSetEditingLOD()
 {
-	UInteractiveTool* Tool = GetToolManager()->GetActiveTool(EToolSide::Mouse);
-	if (const ISkeletalMeshEditingInterface* SkeletonInterface = GetSkeletonInterface(Tool))
+	return !GetToolManager()->HasAnyActiveTool();
+}
+
+void USkeletalMeshModelingToolsEditorMode::SetEditingLOD(EMeshLODIdentifier InEditingLOD)
+{	
+	EditingLOD = InEditingLOD;
+
+	TObjectPtr<UToolTargetManager> TargetManager = GetInteractiveToolsContext(EToolsContextScope::EdMode)->TargetManager;
+	if (USkeletalMeshComponentToolTargetFactory* SkeletalMeshTargetFactory = TargetManager->FindFirstFactoryByType<USkeletalMeshComponentToolTargetFactory>())
 	{
-		return !SkeletonInterface->GetSelectedBones().IsEmpty();
+		SkeletalMeshTargetFactory->SetActiveEditingLOD(EditingLOD);
 	}
+}
 
-	if (Binding.IsValid())
+EMeshLODIdentifier USkeletalMeshModelingToolsEditorMode::GetEditingLOD()
+{
+	return EditingLOD;
+}
+
+void USkeletalMeshModelingToolsEditorMode::HideSkeletonForTool()
+{
+	if (GetSkeletalMeshSkeletonDrawMode() != ESkeletonDrawMode::Hidden)
 	{
-		return !Binding.Pin()->GetSelectedBones().IsEmpty();
+		SavedSkeletonDrawMode = GetSkeletalMeshSkeletonDrawMode();
+		SetSkeletalMeshSkeletonDrawMode(ESkeletonDrawMode::Hidden);
 	}
-	
-	return false;
+}
+
+void USkeletalMeshModelingToolsEditorMode::ShowSkeletonForTool()
+{
+	if (GetSkeletalMeshSkeletonDrawMode() == ESkeletonDrawMode::Hidden)
+	{
+		SetSkeletalMeshSkeletonDrawMode(SavedSkeletonDrawMode);
+	}
+}
+
+ISkeletalMeshEditingInterface* USkeletalMeshModelingToolsEditorMode::GetSkeletonInterface(UInteractiveTool* InTool)
+{
+	if (!IsValid(InTool) || !InTool->Implements<USkeletalMeshEditingInterface>())
+	{
+		return nullptr;
+	}
+	return static_cast<ISkeletalMeshEditingInterface*>(InTool->GetInterfaceAddress(USkeletalMeshEditingInterface::StaticClass()));
 }
 
 #undef LOCTEXT_NAMESPACE
diff --git a/Engine/Plugins/Animation/SkeletalMeshModelingTools/Source/SkeletalMeshModelingTools/Private/SkeletalMeshModelingToolsEditorMode.h b/Engine/Plugins/Animation/SkeletalMeshModelingTools/Source/SkeletalMeshModelingTools/Private/SkeletalMeshModelingToolsEditorMode.h
--- a/Engine/Plugins/Animation/SkeletalMeshModelingTools/Source/SkeletalMeshModelingTools/Private/SkeletalMeshModelingToolsEditorMode.h
+++ b/Engine/Plugins/Animation/SkeletalMeshModelingTools/Source/SkeletalMeshModelingTools/Private/SkeletalMeshModelingToolsEditorMode.h
@@ -5,6 +5,8 @@
 #include "CoreMinimal.h"
 #include "EdMode.h"
 #include "SkeletalMeshModelingModeToolExtensions.h"
+#include "Animation/DebugSkelMeshComponent.h"
+#include "TargetInterfaces/MeshTargetInterfaceTypes.h"
 
 #include "SkeletalMeshModelingToolsEditorMode.generated.h"
 
@@ -18,6 +20,11 @@ class HHitProxy;
 class UDebugSkelMeshComponent;
 class ISkeletalMeshEditor;
 enum class EToolManagerToolSwitchMode;
+class FTabManager;
+enum class EToolSide;
+enum class EToolShutdownType : uint8;
+struct FToolBuilderState;
+
 namespace UE
 {
 	class IInteractiveToolCommandsInterface;
@@ -51,12 +58,35 @@ public:
 
 	// binding
 	void SetEditorBinding(const TWeakPtr<ISkeletalMeshEditor>& InSkeletalMeshEditor);
-		
+	TSharedPtr<ISkeletalMeshEditorBinding> GetEditorBinding();
+	
+	USkeletalMesh* GetSkeletalMesh() const;
+	
+	bool CanSetEditingLOD();
+	void SetEditingLOD(EMeshLODIdentifier EditingLOD);
+	EMeshLODIdentifier GetEditingLOD();
+	
+	void HideSkeletonForTool();	
+	void ShowSkeletonForTool();
+	const TArray<FTransform>& GetComponentSpaceBoneTransforms() const;
+	void ToggleBoneManipulation(bool bEnable);
+
+
+	TSharedPtr<FTabManager> GetAssociatedTabManager();
+	FSimpleMulticastDelegate& OnInitialized() { return OnInitializedDelegate; }
+	
 protected:
 	virtual void OnToolStarted(UInteractiveToolManager* Manager, UInteractiveTool* Tool) override;
-	virtual void OnToolEnded(UInteractiveToolManager* Manager, UInteractiveTool* Tool) override;
-	
+	void OnToolPostBuild(UInteractiveToolManager* InToolManager, EToolSide InSide, UInteractiveTool* InBuiltTool, UInteractiveToolBuilder* InToolBuilder, const FToolBuilderState& ToolState);
+	void OnToolEndedWithStatus(UInteractiveToolManager* Manager, UInteractiveTool* Tool, EToolShutdownType ShutdownType);
 private:
+
+	TWeakPtr<ISkeletalMeshEditor> Editor;
+	TWeakPtr<ISkeletalMeshEditorBinding> EditorBinding;
+	TWeakObjectPtr<USkeletalMesh> SkeletalMesh;
+	
+	EMeshLODIdentifier EditingLOD;
+	
 	// Stylus support is currently disabled; this is left in for reference if/when it is brought back
 	//TUniquePtr<FStylusStateTracker> StylusStateTracker;
 
@@ -67,15 +97,15 @@ private:
 
 	UDebugSkelMeshComponent* GetSkelMeshComponent() const;
 
-	bool NeedsTransformGizmo() const;
-
 	bool bDeactivateOnPIEStartStateToRestore;
 
-	FDelegateHandle ToToolNotifierHandle;
-	FDelegateHandle FromToolNotifierHandle;
+	ESkeletonDrawMode GetSkeletalMeshSkeletonDrawMode() const;
+	void SetSkeletalMeshSkeletonDrawMode(ESkeletonDrawMode DrawMode) const;
+	ESkeletonDrawMode GetCurrentSkeletonDrawMode() const;
 	
-	TWeakPtr<ISkeletalMeshEditorBinding> Binding;
+	ESkeletonDrawMode SavedSkeletonDrawMode = ESkeletonDrawMode::Default;
 
+	
 	void RegisterExtensions();
 	// Support extension tools having their own hotkey classes
 	TMap<FString, FExtensionToolDescription> ExtensionToolToInfo;
@@ -85,4 +115,6 @@ private:
 		TFunction<const UE::IInteractiveToolCommandsInterface&()>& OutGetter) const;
 	// Used to unbind extension tool commands
 	TFunction<const UE::IInteractiveToolCommandsInterface& ()> ExtensionToolCommandsGetter;
+
+	FSimpleMulticastDelegate OnInitializedDelegate;
 };
diff --git a/Engine/Plugins/Animation/SkeletalMeshModelingTools/Source/SkeletalMeshModelingTools/Private/SkeletalMeshModelingToolsEditorModeToolkit.cpp b/Engine/Plugins/Animation/SkeletalMeshModelingTools/Source/SkeletalMeshModelingTools/Private/SkeletalMeshModelingToolsEditorModeToolkit.cpp
--- a/Engine/Plugins/Animation/SkeletalMeshModelingTools/Source/SkeletalMeshModelingTools/Private/SkeletalMeshModelingToolsEditorModeToolkit.cpp
+++ b/Engine/Plugins/Animation/SkeletalMeshModelingTools/Source/SkeletalMeshModelingTools/Private/SkeletalMeshModelingToolsEditorModeToolkit.cpp
@@ -19,10 +19,15 @@
 #include "Widgets/Images/SImage.h"
 #include "Widgets/Text/STextBlock.h"
 #include "ModelingToolsManagerActions.h"
+#include "ModelingToolTargetEditorOnlyUtil.h"
 #include "SkeletalMeshModelingModeToolExtensions.h"
+#include "SkeletalMeshModelingToolsEditorMode.h"
 #include "BaseTools/MeshSurfacePointTool.h"
 #include "SPrimaryButton.h"
 #include "Toolkits/AssetEditorModeUILayer.h"
+#include "Widgets/Input/STextComboBox.h"
+#include "Widgets/Layout/SExpandableArea.h"
+
 
 
 #define LOCTEXT_NAMESPACE "SkeletalMeshModelingToolsEditorModeToolkit"
@@ -74,6 +79,7 @@ void FSkeletalMeshModelingToolsEditorModeToolkit::Init(
 		ToolkitSections->ModeWarningArea = ModeWarningArea;
 		ToolkitSections->DetailsView = ModeDetailsView;
 		ToolkitSections->ToolWarningArea = ToolWarningArea;
+		ToolkitSections->Footer = MakeFooterWidget();
 
 		SAssignNew(ToolkitWidget, SBorder)
 		.HAlign(HAlign_Fill)
@@ -95,6 +101,10 @@ void FSkeletalMeshModelingToolsEditorModeToolkit::Init(
 
 	// add viewport overlay widget to accept / cancel tool
 	MakeToolAcceptCancelWidget();
+
+	
+	EditorMode = CastChecked<USkeletalMeshModelingToolsEditorMode>(InOwningMode.Get());
+	EditorMode->OnInitialized().AddSP(this, &FSkeletalMeshModelingToolsEditorModeToolkit::OnEditorModeInitialized);
 }
 
 FName FSkeletalMeshModelingToolsEditorModeToolkit::GetToolkitFName() const
@@ -391,6 +401,90 @@ void FSkeletalMeshModelingToolsEditorModeToolkit::ClearWarning()
 	ToolWarningArea->SetVisibility(EVisibility::Collapsed);
 }
 
+void FSkeletalMeshModelingToolsEditorModeToolkit::OnEditorModeInitialized()
+{
+	USkeletalMesh* SkeletalMesh = EditorMode->GetSkeletalMesh();
+
+	constexpr bool bSkipAutoGenerated= true;
+	AssetAvailableLODs = UE::ToolTarget::GetAvailableLODs(SkeletalMesh, bSkipAutoGenerated);
+	
+	for (EMeshLODIdentifier LOD : AssetAvailableLODs)
+	{
+		int32 LODIndex = static_cast<int32>(LOD);
+		FString LODName = TEXT("LOD") + FString::FromInt(LODIndex);
+		AssetLODModes.Add(MakeShared<FString>(LODName));
+	}
+
+	AssetLODMode->SetSelectedItem(AssetLODModes[0]);
+
+}
+TSharedPtr<SWidget> FSkeletalMeshModelingToolsEditorModeToolkit::MakeFooterWidget()
+{
+	AssetLODMode = SNew(STextComboBox)
+		.OptionsSource(&AssetLODModes)
+		.IsEnabled_Lambda([this]()
+		{
+			if (!EditorMode.IsValid())
+			{
+				return false;
+			}
+
+			return EditorMode->CanSetEditingLOD();
+		})
+		.OnSelectionChanged_Lambda([this](TSharedPtr<FString> String, ESelectInfo::Type)
+		{
+			int32 Index = AssetLODModes.IndexOfByKey(String);
+			check(AssetAvailableLODs.Num() == AssetLODModes.Num());
+			EMeshLODIdentifier NewSelectedLOD = AssetAvailableLODs[Index];
+
+			if (EditorMode.IsValid())
+			{
+				EditorMode->SetEditingLOD(NewSelectedLOD);
+			}
+		});
+
+	AssetLODModeLabel = SNew(STextBlock)
+	.Text(LOCTEXT("ActiveLODLabel", "Editing LOD"))
+	.ToolTipText(LOCTEXT("ActiveLODLabelToolTip", "Select the LOD to be used when editing an existing mesh."));
+
+	const TSharedPtr<SVerticalBox> Content = SNew(SVerticalBox)
+	+ SVerticalBox::Slot().HAlign(HAlign_Fill)
+	.Padding(0, 8, 0, 4)
+		[
+		
+			SNew(SVerticalBox)	
+			+SVerticalBox::Slot()
+			.AutoHeight()
+			[
+				SNew(SHorizontalBox)
+				+ SHorizontalBox::Slot().Padding(0).HAlign(HAlign_Left).VAlign(VAlign_Center).FillWidth(2.f)
+					[ AssetLODModeLabel->AsShared() ]
+				+ SHorizontalBox::Slot().Padding(0).HAlign(HAlign_Fill).Padding(0).FillWidth(4.f)
+					[ AssetLODMode->AsShared() ]
+			]
+		];
+
+	TSharedPtr<SExpandableArea> AssetConfigPanel = SNew(SExpandableArea)
+		.HeaderPadding(FMargin(0.f))
+		.Padding(FMargin(8.f))
+		.BorderImage(FAppStyle::Get().GetBrush("DetailsView.CategoryTop"))
+		.AreaTitleFont(FAppStyle::Get().GetFontStyle("EditorModesPanel.CategoryFontStyle"))
+		.BodyContent()
+		[
+			Content->AsShared()
+		]
+		.HeaderContent()
+		[
+			SNew(SHorizontalBox)
+			+ SHorizontalBox::Slot().HAlign(HAlign_Left).VAlign(VAlign_Center).FillWidth(2.f)
+			[
+				SNew(STextBlock)
+				.Text(LOCTEXT("SkeletalMeshModelingToolFooterHeader", "Asset Editing Settings"))
+			]	
+		];
+	
+	return AssetConfigPanel;
+}
 
 void FSkeletalMeshModelingToolsEditorModeToolkit::UpdateActiveToolProperties(UInteractiveTool* Tool)
 {
diff --git a/Engine/Plugins/Animation/SkeletalMeshModelingTools/Source/SkeletalMeshModelingTools/Private/SkeletalMeshModelingToolsEditorModeToolkit.h b/Engine/Plugins/Animation/SkeletalMeshModelingTools/Source/SkeletalMeshModelingTools/Private/SkeletalMeshModelingToolsEditorModeToolkit.h
--- a/Engine/Plugins/Animation/SkeletalMeshModelingTools/Source/SkeletalMeshModelingTools/Private/SkeletalMeshModelingToolsEditorModeToolkit.h
+++ b/Engine/Plugins/Animation/SkeletalMeshModelingTools/Source/SkeletalMeshModelingTools/Private/SkeletalMeshModelingToolsEditorModeToolkit.h
@@ -3,13 +3,15 @@
 #pragma once
 
 #include "CoreMinimal.h"
+#include "SkeletalMeshEditorUtils.h"
 #include "StatusBarSubsystem.h"
 #include "Toolkits/BaseToolkit.h"
 
 
+class USkeletalMeshModelingToolsEditorMode;
+enum class EMeshLODIdentifier;
 class IToolkitHost;
 class STextBlock;
-class SButton;
 
 
 class FSkeletalMeshModelingToolsEditorModeToolkit : public FModeToolkit
@@ -37,6 +39,9 @@ private:
 	void PostWarning(const FText& Message);
 	void ClearWarning();
 
+
+	void OnEditorModeInitialized();
+	TSharedPtr<SWidget> MakeFooterWidget();
 	void UpdateActiveToolProperties(UInteractiveTool* Tool);
 
 	void RegisterPalettes();
@@ -58,4 +63,11 @@ private:
 	TSharedPtr<STextBlock> ModeWarningArea;
 	TSharedPtr<STextBlock> ModeHeaderArea;
 	TSharedPtr<STextBlock> ToolWarningArea;
+
+	TArray<EMeshLODIdentifier> AssetAvailableLODs;
+	TArray<TSharedPtr<FString>> AssetLODModes;
+	TSharedPtr<STextBlock> AssetLODModeLabel;
+	TSharedPtr<STextComboBox> AssetLODMode;
+	
+	TWeakObjectPtr<USkeletalMeshModelingToolsEditorMode> EditorMode;
 };
diff --git a/Engine/Plugins/Experimental/Animation/SkeletalMeshMorphTargetEditingTools/Source/SkeletalMeshMorphTargetEditingTools/Private/IMorphTargetEditingToolInterface.cpp b/Engine/Plugins/Experimental/Animation/SkeletalMeshMorphTargetEditingTools/Source/SkeletalMeshMorphTargetEditingTools/Private/IMorphTargetEditingToolInterface.cpp
--- a/Engine/Plugins/Experimental/Animation/SkeletalMeshMorphTargetEditingTools/Source/SkeletalMeshMorphTargetEditingTools/Private/IMorphTargetEditingToolInterface.cpp
+++ b/Engine/Plugins/Experimental/Animation/SkeletalMeshMorphTargetEditingTools/Source/SkeletalMeshMorphTargetEditingTools/Private/IMorphTargetEditingToolInterface.cpp
@@ -10,7 +10,7 @@
 #include "PersonaModule.h"
 #include "SingleSelectionTool.h"
 #include "SKMMorphTargetBackedTarget.h"
-#include "SkeletalMesh/SkeletalMeshEditionInterface.h"
+#include "SkeletalMesh/SkeletalMeshEditingInterface.h"
 
 #include UE_INLINE_GENERATED_CPP_BY_NAME(IMorphTargetEditingToolInterface)
 
diff --git a/Engine/Plugins/Experimental/Animation/SkeletalMeshMorphTargetEditingTools/Source/SkeletalMeshMorphTargetEditingTools/Public/MorphTargetVertexSculptTool.h b/Engine/Plugins/Experimental/Animation/SkeletalMeshMorphTargetEditingTools/Source/SkeletalMeshMorphTargetEditingTools/Public/MorphTargetVertexSculptTool.h
--- a/Engine/Plugins/Experimental/Animation/SkeletalMeshMorphTargetEditingTools/Source/SkeletalMeshMorphTargetEditingTools/Public/MorphTargetVertexSculptTool.h
+++ b/Engine/Plugins/Experimental/Animation/SkeletalMeshMorphTargetEditingTools/Source/SkeletalMeshMorphTargetEditingTools/Public/MorphTargetVertexSculptTool.h
@@ -4,7 +4,7 @@
 
 #include "IMorphTargetEditingToolInterface.h"
 #include "MeshVertexSculptTool.h"
-#include "SkeletalMesh/SkeletalMeshEditionInterface.h"
+#include "SkeletalMesh/SkeletalMeshEditingInterface.h"
 #include "BaseTools/MultiSelectionMeshEditingTool.h"
 #include "MorphTargetEditingToolProperties.h"
 #include "SKMMorphTargetBackedTarget.h"
diff --git a/Engine/Plugins/Experimental/Dataflow/Source/DataflowEditor/Private/DataflowEditorTools/DataflowEditorSkinWeightsPaintTool.cpp b/Engine/Plugins/Experimental/Dataflow/Source/DataflowEditor/Private/DataflowEditorTools/DataflowEditorSkinWeightsPaintTool.cpp
--- a/Engine/Plugins/Experimental/Dataflow/Source/DataflowEditor/Private/DataflowEditorTools/DataflowEditorSkinWeightsPaintTool.cpp
+++ b/Engine/Plugins/Experimental/Dataflow/Source/DataflowEditor/Private/DataflowEditorTools/DataflowEditorSkinWeightsPaintTool.cpp
@@ -12,6 +12,7 @@
 #include "ModelingToolTargetUtil.h"
 #include "Components/SkeletalMeshComponent.h"
 #include "DynamicMesh/DynamicVertexSkinWeightsAttribute.h"
+#include "DynamicMesh/NonManifoldMappingSupport.h"
 #include "TargetInterfaces/MeshDescriptionCommitter.h"
 #include "TargetInterfaces/MeshDescriptionProvider.h"
 #include "TargetInterfaces/PrimitiveComponentBackedTarget.h"
@@ -241,48 +242,47 @@ void UDataflowEditorSkinWeightsPaintTool::OnShutdown(EToolShutdownType ShutdownT
 			SkinWeightNode->ExtractVertexWeights(*DataflowContext, SetupIndices, SetupWeights,
 				TArrayView<TArray<int32>>(CurrentIndices), TArrayView<TArray<float>>(CurrentWeights));
 			
-			if(FCleanedEditMesh* EditedMesh = EditedMeshes.Find(EMeshLODIdentifier::LOD0))
+			// profile to edit
+			const FName ActiveProfile = WeightToolProperties->GetActiveSkinWeightProfile();
+
+			FDynamicMeshVertexSkinWeightsAttribute* SkinWeights = EditedMesh.Attributes()->GetSkinWeightsAttribute(ActiveProfile);
+
+			FNonManifoldMappingSupport NonManifoldMappingSupport(EditedMesh);
+			for (int32 DynaMeshVertexIndex : EditedMesh.VertexIndicesItr())
 			{
-				const FMeshDescription& MeshDescription = EditedMesh->GetEditableMeshDescription();
-			
-				// profile to edit
-				const FName ActiveProfile = WeightToolProperties->GetActiveSkinWeightProfile();
-			
-				FSkeletalMeshConstAttributes MeshAttribs(MeshDescription);
-				FSkinWeightsVertexAttributesConstRef SkinWeights = MeshAttribs.GetVertexSkinWeights(ActiveProfile);
+				int32 VertexIndex = NonManifoldMappingSupport.GetOriginalNonManifoldVertexID(DynaMeshVertexIndex);
+				
+				UE::AnimationCore::FBoneWeights BoneWeights;
+				SkinWeights->GetValue(DynaMeshVertexIndex, BoneWeights);
+				
+				const int32 CurrentIndex = VertexOffset+VertexIndex;
 			
-				const int32 NumVertices = MeshDescription.Vertices().Num();
-				for(int32 VertexIndex = 0; VertexIndex < NumVertices; ++VertexIndex)
+				if(CurrentIndex < CurrentWeights.Num() && CurrentIndex < CurrentIndices.Num() &&
+					CurrentWeights[CurrentIndex].Num() == CurrentIndices[CurrentIndex].Num())
 				{
-					FVertexBoneWeightsConst BoneWeights = SkinWeights.Get(FVertexID(VertexIndex));
-					const int32 CurrentIndex = VertexOffset+VertexIndex;
-			
-					if(CurrentIndex < CurrentWeights.Num() && CurrentIndex < CurrentIndices.Num() &&
-						CurrentWeights[CurrentIndex].Num() == CurrentIndices[CurrentIndex].Num())
-					{
-						CurrentIndices[CurrentIndex].Reset(BoneWeights.Num());
-						CurrentWeights[CurrentIndex].Reset(BoneWeights.Num());
+					CurrentIndices[CurrentIndex].Reset(BoneWeights.Num());
+					CurrentWeights[CurrentIndex].Reset(BoneWeights.Num());
 
-						TArray<TPair<int32,float>> SortedWeights;
-						SortedWeights.Reserve(BoneWeights.Num());
+					TArray<TPair<int32,float>> SortedWeights;
+					SortedWeights.Reserve(BoneWeights.Num());
 						
-						for(int32 WeightIndex = 0, NumWeights = BoneWeights.Num(); WeightIndex < NumWeights; ++WeightIndex)
-						{
-							SortedWeights.Add({BoneWeights[WeightIndex].GetBoneIndex(), BoneWeights[WeightIndex].GetWeight()});
-						}
-						Algo::Sort(SortedWeights, [](const TPair<int32,float>& A, const TPair<int32,float>& B)
-						{
-							return A.Value > B.Value; 
-						});
+					for(int32 WeightIndex = 0, NumWeights = BoneWeights.Num(); WeightIndex < NumWeights; ++WeightIndex)
+					{
+						SortedWeights.Add({BoneWeights[WeightIndex].GetBoneIndex(), BoneWeights[WeightIndex].GetWeight()});
+					}
+					Algo::Sort(SortedWeights, [](const TPair<int32,float>& A, const TPair<int32,float>& B)
+					{
+						return A.Value > B.Value; 
+					});
 
-						for(int32 WeightIndex = 0, NumWeights = BoneWeights.Num(); WeightIndex < NumWeights; ++WeightIndex)
-						{
-							CurrentIndices[CurrentIndex].Add(SortedWeights[WeightIndex].Key);
-							CurrentWeights[CurrentIndex].Add(SortedWeights[WeightIndex].Value);
-						}
+					for(int32 WeightIndex = 0, NumWeights = BoneWeights.Num(); WeightIndex < NumWeights; ++WeightIndex)
+					{
+						CurrentIndices[CurrentIndex].Add(SortedWeights[WeightIndex].Key);
+						CurrentWeights[CurrentIndex].Add(SortedWeights[WeightIndex].Value);
 					}
-				}
+				}	
 			}
+			
 			SkinWeightNode->ReportVertexWeights(*DataflowContext, SetupIndices, SetupWeights, CurrentIndices, CurrentWeights);
 			SkinWeightNode->Invalidate();
 
diff --git a/Engine/Plugins/Experimental/Dataflow/Source/DataflowEditor/Private/DataflowEditorTools/DataflowEditorVertexAttributePaintTool.cpp b/Engine/Plugins/Experimental/Dataflow/Source/DataflowEditor/Private/DataflowEditorTools/DataflowEditorVertexAttributePaintTool.cpp
--- a/Engine/Plugins/Experimental/Dataflow/Source/DataflowEditor/Private/DataflowEditorTools/DataflowEditorVertexAttributePaintTool.cpp
+++ b/Engine/Plugins/Experimental/Dataflow/Source/DataflowEditor/Private/DataflowEditorTools/DataflowEditorVertexAttributePaintTool.cpp
@@ -655,13 +655,8 @@ void UDataflowEditorVertexAttributePaintTool::Setup()
 
 	// compute the mesh description for the mesh selector to use
 	{
-		using namespace UE::Geometry;
-		FStaticMeshAttributes Attributes(PreviewMeshDescription);
-		Attributes.Register();
-		FConversionToMeshDescriptionOptions ConverterOptions;
-		FDynamicMeshToMeshDescription Converter(ConverterOptions);
-		Converter.Convert(Mesh, PreviewMeshDescription);
-		MeshSelector->SetMesh(PreviewMesh, &PreviewMeshDescription, FTransform::Identity);
+
+		MeshSelector->SetMesh(PreviewMesh, FTransform::Identity);
 	}
 
 	// initialize brush radius range interval, brush properties
@@ -1558,9 +1553,7 @@ TArray<int32> UDataflowEditorVertexAttributePaintTool::GetSelectedVertices() con
 {
 	if (MeshSelector)
 	{
-		// we do not want to remap to original as the tool already take care of doing the remapping when necessary 
-		// we operate in the selected sub-mesh vertex range
-		return MeshSelector->GetSelectedVertices(/*bRemapToOriginal*/false);
+		return MeshSelector->GetSelectedVertices();
 	}
 	return {};
 }
diff --git a/Engine/Plugins/Runtime/GeometryProcessing/Source/DynamicMesh/Private/Operations/SmoothBoneWeights.cpp b/Engine/Plugins/Runtime/GeometryProcessing/Source/DynamicMesh/Private/Operations/SmoothBoneWeights.cpp
--- a/Engine/Plugins/Runtime/GeometryProcessing/Source/DynamicMesh/Private/Operations/SmoothBoneWeights.cpp
+++ b/Engine/Plugins/Runtime/GeometryProcessing/Source/DynamicMesh/Private/Operations/SmoothBoneWeights.cpp
@@ -133,18 +133,18 @@ bool TSmoothBoneWeights<BoneIndexType, BoneWeightType>::SmoothWeightsAtVertex(co
 {
 	using namespace SmoothBoneWeightsLocals;
 	
-	const UE::Geometry::FNonManifoldMappingSupport NonManifoldMappingSupport(*SourceMesh);
+	// const UE::Geometry::FNonManifoldMappingSupport NonManifoldMappingSupport(*SourceMesh);
 
 	// for each vertex in the stamp...
 	constexpr int32 AvgNumNeighbors = 8;
 	using VertexNeighborWeights = TArray<BoneWeightType, TInlineAllocator<AvgNumNeighbors>>;
 	TArray<int32> AllNeighborVertices;
 	TMap<BoneIndexType, VertexNeighborWeights> WeightsOnAllNeighbors;
-	const int32 SrcVertexID = NonManifoldMappingSupport.GetOriginalNonManifoldVertexID(VertexID);
+	// const int32 SrcVertexID = NonManifoldMappingSupport.GetOriginalNonManifoldVertexID(VertexID);
 
 	// get list of all neighboring vertices, AND this vertex
 	AllNeighborVertices.Add(VertexID);
-	for (const int32 NeighborVertexID : SourceMesh->VtxVerticesItr(SrcVertexID))
+	for (const int32 NeighborVertexID : SourceMesh->VtxVerticesItr(VertexID))
 	{
 		AllNeighborVertices.Add(NeighborVertexID);
 	}
diff --git a/Engine/Plugins/Runtime/MeshModelingToolset/Source/MeshModelingToolsEditorOnly/Private/SkeletalMesh/SkeletalMeshEditingInterface.cpp b/Engine/Plugins/Runtime/MeshModelingToolset/Source/MeshModelingToolsEditorOnly/Private/SkeletalMesh/SkeletalMeshEditingInterface.cpp
new file mode 100644
--- /dev/null
+++ b/Engine/Plugins/Runtime/MeshModelingToolset/Source/MeshModelingToolsEditorOnly/Private/SkeletalMesh/SkeletalMeshEditingInterface.cpp
@@ -0,0 +1,72 @@
+﻿// Copyright Epic Games, Inc. All Rights Reserved.
+
+#include "SkeletalMesh/SkeletalMeshEditingInterface.h"
+
+#include UE_INLINE_GENERATED_CPP_BY_NAME(SkeletalMeshEditingInterface)
+
+IMPLEMENT_HIT_PROXY(HBoneHitProxy, HHitProxy)
+
+ISkeletalMeshNotifier& ISkeletalMeshEditingInterface::GetNotifier()
+{
+	if (!Notifier)
+	{
+		Notifier.Reset(new FSkeletalMeshToolNotifier(TWeakInterfacePtr<ISkeletalMeshEditingInterface>(this)));
+	}
+	
+	return *Notifier;
+}
+
+bool ISkeletalMeshEditingInterface::NeedsNotification() const
+{
+	return Notifier && Notifier->Delegate().IsBound(); 
+}
+
+void ISkeletalMeshEditingInterface::BindTo(TSharedPtr<ISkeletalMeshEditorBinding> InBinding)
+{
+	Binding = InBinding;
+}
+
+void ISkeletalMeshEditingInterface::Unbind()
+{
+	Binding.Reset();
+}
+
+TWeakObjectPtr<USkeletonModifier> ISkeletalMeshEditingInterface::GetModifier() const
+{
+	return nullptr;
+}
+
+TOptional<FName> ISkeletalMeshEditingInterface::GetBoneName(HHitProxy* InHitProxy) const
+{
+	if (const HBoneHitProxy* BoneProxy = HitProxyCast<HBoneHitProxy>(InHitProxy))
+	{
+		return BoneProxy->BoneName;
+	}
+
+	return Binding.IsValid() && Binding.Pin()->GetNameFunction() ? Binding.Pin()->GetNameFunction()(InHitProxy) : TOptional<FName>();
+}
+
+TArray<FName> ISkeletalMeshEditingInterface::GetSelectedBones() const
+{
+	static const TArray<FName> Dummy;
+	return Binding.IsValid() ? Binding.Pin()->GetSelectedBones() : Dummy;
+}
+
+TArray<FName> ISkeletalMeshEditingInterface::GetSelectedMorphTargets() const
+{
+	static const TArray<FName> Dummy;
+	return Binding.IsValid() ? Binding.Pin()->GetSelectedMorphTargets() : Dummy;
+}
+
+FSkeletalMeshToolNotifier::FSkeletalMeshToolNotifier(TWeakInterfacePtr<ISkeletalMeshEditingInterface> InInterface)
+	: ISkeletalMeshNotifier()
+	, Interface(InInterface)
+{}
+
+void FSkeletalMeshToolNotifier::HandleNotification(const TArray<FName>& BoneNames, const ESkeletalMeshNotifyType InNotifyType)
+{
+	if (Interface.IsValid())
+	{
+		Interface->HandleSkeletalMeshModified(BoneNames, InNotifyType);
+	}
+}
diff --git a/Engine/Plugins/Runtime/MeshModelingToolset/Source/MeshModelingToolsEditorOnly/Private/SkeletalMesh/SkeletalMeshEditionInterface.cpp b/Engine/Plugins/Runtime/MeshModelingToolset/Source/MeshModelingToolsEditorOnly/Private/SkeletalMesh/SkeletalMeshEditionInterface.cpp
deleted file mode 100644
diff --git a/Engine/Plugins/Runtime/MeshModelingToolset/Source/MeshModelingToolsEditorOnly/Private/SkeletalMesh/SkinWeightsPaintTool.cpp b/Engine/Plugins/Runtime/MeshModelingToolset/Source/MeshModelingToolsEditorOnly/Private/SkeletalMesh/SkinWeightsPaintTool.cpp
--- a/Engine/Plugins/Runtime/MeshModelingToolset/Source/MeshModelingToolsEditorOnly/Private/SkeletalMesh/SkinWeightsPaintTool.cpp
+++ b/Engine/Plugins/Runtime/MeshModelingToolset/Source/MeshModelingToolsEditorOnly/Private/SkeletalMesh/SkinWeightsPaintTool.cpp
@@ -35,14 +35,10 @@
 
 #include "MeshDescription.h"
 #include "MeshModelingToolsEditorOnly.h"
-#include "DynamicMeshToMeshDescription.h"
 #include "DynamicSubmesh3.h"
-#include "MeshDescriptionToDynamicMesh.h"
 #include "DynamicMesh/DynamicVertexSkinWeightsAttribute.h"
 #include "DynamicMesh/MeshAdapterUtil.h"
-#include "DynamicMesh/NonManifoldMappingSupport.h"
 #include "TargetInterfaces/DynamicMeshCommitter.h"
-#include "TargetInterfaces/MeshDescriptionCommitter.h"
 #include "TargetInterfaces/MeshDescriptionProvider.h"
 #include "SkeletalMeshAttributes.h"
 
@@ -58,6 +54,9 @@
 #include "HAL/PlatformApplicationMisc.h"
 
 #include "Serialization/JsonSerializable.h"
+#include "TargetInterfaces/DynamicMeshProvider.h"
+#include "TargetInterfaces/SkeletalMeshBackedTarget.h"
+#include "TargetInterfaces/SkeletonProvider.h"
 
 #include UE_INLINE_GENERATED_CPP_BY_NAME(SkinWeightsPaintTool)
 
@@ -132,27 +131,6 @@ FName GetLODName(const EMeshLODIdentifier InLOD)
 	return Name ? *Name : NAME_None;
 }
 
-USkeletalMeshComponent* GetSkeletalMeshComponent(const UToolTarget* InTarget)
-{
-	if (!ensure(InTarget))
-	{
-		return nullptr;
-	}
-	
-	const IPrimitiveComponentBackedTarget* TargetComponent = Cast<IPrimitiveComponentBackedTarget>(InTarget);
-	if (!ensure(TargetComponent))
-	{
-		return nullptr;
-	}
-	
-	USkeletalMeshComponent* Component = Cast<USkeletalMeshComponent>(TargetComponent->GetOwnerComponent());
-	if (!(Component && Component->GetSkeletalMeshAsset()))
-	{
-		return nullptr;
-	}
-
-	return Component;
-}
 
 const FName& CreateNewName()
 {
@@ -160,27 +138,30 @@ const FName& CreateNewName()
 	return CreateNew;
 }
 
-FSkinWeightsVertexAttributesRef GetOrCreateSkinWeightsAttribute(FMeshDescription& InMesh, const FName InProfileName)
+UE::Geometry::FDynamicMeshVertexSkinWeightsAttribute* GetOrCreateSkinWeightsAttribute(FDynamicMesh3& InMesh, const FName InProfileName)
 {
-	FSkeletalMeshAttributes MeshAttribs(InMesh);
-	const TArray<FName> Profiles = MeshAttribs.GetSkinWeightProfileNames();
-	if (!Profiles.Contains(InProfileName))
+	using namespace UE::Geometry;
+	FDynamicMeshVertexSkinWeightsAttribute* Attribute = InMesh.Attributes()->GetSkinWeightsAttribute(InProfileName);
+	if (Attribute == nullptr)
 	{
-		ensure( MeshAttribs.RegisterSkinWeightAttribute(InProfileName) );
+		Attribute = new FDynamicMeshVertexSkinWeightsAttribute(&InMesh);
+		InMesh.Attributes()->AttachSkinWeightsAttribute(InProfileName, Attribute);
 	}
-	return MeshAttribs.GetVertexSkinWeights(InProfileName);
+	return Attribute;
 }
 
-bool RenameSkinWeightsAttribute(FMeshDescription& InMesh, const FName InOldName, const FName InNewName)
+bool RenameSkinWeightsAttribute(FDynamicMesh3& InMesh, const FName InOldName, const FName InNewName)
 {
-	FSkeletalMeshAttributes MeshAttribs(InMesh);
-	const TArray<FName> Profiles = MeshAttribs.GetSkinWeightProfileNames();
-	if (Profiles.Contains(InOldName))
+	using namespace UE::Geometry;
+	if (FDynamicMeshVertexSkinWeightsAttribute* Attribute = InMesh.Attributes()->GetSkinWeightsAttribute(InOldName))
 	{
-		FSkinWeightsVertexAttributesRef NewWeightsAttr = GetOrCreateSkinWeightsAttribute(InMesh, InNewName);
-		NewWeightsAttr = MeshAttribs.GetVertexSkinWeights(InOldName);
-		return MeshAttribs.UnregisterSkinWeightAttribute(InOldName);
+		FDynamicMeshVertexSkinWeightsAttribute* NewAttribute = GetOrCreateSkinWeightsAttribute(InMesh, InNewName);
+		*NewAttribute = MoveTemp(*Attribute);
+
+		InMesh.Attributes()->RemoveSkinWeightsAttribute(InOldName);
+		return true;
 	}
+
 	return false;
 }
 	
@@ -196,9 +177,8 @@ class FPaintToolWeightsDataSource : public UE::Geometry::TBoneWeightsDataSource<
 {
 public:
 
-	FPaintToolWeightsDataSource(const FSkinToolWeights* InWeights, const FDynamicMesh3& InDynaMesh)
+	FPaintToolWeightsDataSource(const FSkinToolWeights* InWeights)
 		: Weights(InWeights)
-		, NonManifoldMappingSupport(InDynaMesh)
 	{
 		checkSlow(Weights);
 	}
@@ -207,37 +187,27 @@ public:
 
 	virtual int32 GetBoneNum(const int32 VertexID) override
 	{
-		const int32 SrcVertexId = GetSourceVertexId(VertexID);
-		return Weights->PreChangeWeights[SrcVertexId].Num();
+		return Weights->PreChangeWeights[VertexID].Num();
 	}
 
 	virtual int32 GetBoneIndex(const int32 VertexID, const int32 Index) override
 	{
-		const int32 SrcVertexId = GetSourceVertexId(VertexID);
-		return Weights->PreChangeWeights[SrcVertexId][Index].BoneID;
+		return Weights->PreChangeWeights[VertexID][Index].BoneID;
 	}
 
 	virtual float GetBoneWeight(const int32 VertexID, const int32 Index) override
 	{
-		const int32 SrcVertexId = GetSourceVertexId(VertexID);
-		return Weights->PreChangeWeights[SrcVertexId][Index].Weight;
+		return Weights->PreChangeWeights[VertexID][Index].Weight;
 	}
 
 	virtual float GetWeightOfBoneOnVertex(const int32 VertexID, const int32 BoneIndex) override
 	{
-		const int32 SrcVertexId = GetSourceVertexId(VertexID);
-		return Weights->GetWeightOfBoneOnVertex(BoneIndex, SrcVertexId, Weights->PreChangeWeights);
+		return Weights->GetWeightOfBoneOnVertex(BoneIndex, VertexID, Weights->PreChangeWeights);
 	}
 
 protected:
-
-	int32 GetSourceVertexId(const int32 InVertexID) const
-	{
-		return NonManifoldMappingSupport.GetOriginalNonManifoldVertexID(InVertexID);
-	}
 	
 	const FSkinToolWeights* Weights = nullptr;
-	const UE::Geometry::FNonManifoldMappingSupport NonManifoldMappingSupport;
 };
 
 void FDirectEditWeightState::Reset()
@@ -346,13 +316,32 @@ TArray<FName> GetProfilesFromToolTarget(UToolTarget* InTarget, const FName InLOD
 {
 	if (InTarget)
 	{
-		TArray<FName> Profiles;
 		const EMeshLODIdentifier LODId = GetLODId(InLOD);
-		const FGetMeshParameters Params(true, LODId);
-		if (const FMeshDescription* MeshDescription = UE::ToolTarget::GetMeshDescription(InTarget, Params))
+		
+		if (IMeshDescriptionProvider* MeshDescriptionProvider = Cast<IMeshDescriptionProvider>(InTarget))
 		{
-			const FSkeletalMeshConstAttributes MeshAttribs(*MeshDescription);
-			return MeshAttribs.GetSkinWeightProfileNames();
+			const FMeshDescription* MeshDescription = nullptr;
+			if (MeshDescriptionProvider->SupportsLODs())
+			{
+				const FGetMeshParameters Params(true, LODId);
+				MeshDescription = UE::ToolTarget::GetMeshDescription(InTarget, Params);	
+			}
+			else
+			{
+				MeshDescription = UE::ToolTarget::GetMeshDescription(InTarget);		
+			}
+		
+			if (MeshDescription)
+			{
+				const FSkeletalMeshConstAttributes MeshAttribs(*MeshDescription);
+				return MeshAttribs.GetSkinWeightProfileNames();
+			}	
+		}
+		else if (IDynamicMeshProvider* DynamicMeshProvider = Cast<IDynamicMeshProvider>(InTarget))
+		{
+			TArray<FName> Profiles;
+			DynamicMeshProvider->GetDynamicMesh().Attributes()->GetSkinWeightsAttributes().GenerateKeyArray(Profiles);
+			return Profiles;
 		}
 	}
 	
@@ -362,14 +351,10 @@ TArray<FName> GetProfilesFromToolTarget(UToolTarget* InTarget, const FName InLOD
 
 }
 
-TArray<FName> USkinWeightsPaintToolProperties::GetTargetLODsFunc() const
-{
-	return SkinWeightLayer::GetLODs(WeightTool->GetTarget());
-}
 
 TArray<FName> USkinWeightsPaintToolProperties::GetTargetSkinWeightProfilesFunc() const
 {
-	TArray<FName> Profiles = SkinWeightLayer::GetProfilesFromToolTarget(WeightTool->GetTarget(), SourceLOD);
+	TArray<FName> Profiles = SkinWeightLayer::GetProfilesFromToolTarget(WeightTool->GetTarget(), GetLODName(WeightTool->GetEditingLOD()));
 	Profiles.Add(CreateNewName());
 	return Profiles;
 }
@@ -451,15 +436,16 @@ void UWeightToolTransferManager::SetSourceMesh(USkeletalMesh* InSkeletalMesh)
 			return;
 		}
 
-		const USkeletalMeshComponent* SkeletalMeshComponent = GetSkeletalMeshComponent(MainTarget);
-		if (!SkeletalMeshComponent)
+		IPrimitiveComponentBackedTarget* PrimitiveComponentTarget = CastChecked<IPrimitiveComponentBackedTarget>(MainTarget);
+		const UPrimitiveComponent* PrimitiveComponent = PrimitiveComponentTarget->GetOwnerComponent();
+		if (!PrimitiveComponent)
 		{
 			return;
 		}
 		
-		const FBoxSphereBounds TargetBounds = SkeletalMeshComponent->GetSkeletalMeshAsset()->GetBounds();
-		const FBoxSphereBounds SourceBounds = InSkeletalMesh->GetBounds();
 		FTransform Transform = UE::ToolTarget::GetLocalToWorldTransform(MainTarget);
+		const FBoxSphereBounds TargetBounds = PrimitiveComponent->CalcBounds(Transform);
+		const FBoxSphereBounds SourceBounds = InSkeletalMesh->GetBounds();
 		FVector Location = Transform.GetLocation();
 		Location.X -= TargetBounds.GetBoxExtrema(1).X;
 		Location.X -= 1.1 * SourceBounds.GetBoxExtrema(1).X;
@@ -480,13 +466,9 @@ void UWeightToolTransferManager::SetSourceMesh(USkeletalMesh* InSkeletalMesh)
 	SourcePreviewMesh->SetMaterials(MaterialSet.Materials);
 
 	// setup the mesh selection for the source
-	const FMeshDescription* MeshDescription =  UE::ToolTarget::GetMeshDescription(SourceTarget, SourceParams);
-	if (ensure(MeshDescription))
-	{
-		MeshSelector->SetMesh(SourcePreviewMesh, MeshDescription, ToolProperties->SourcePreviewOffset);
-		MeshSelector->SetIsEnabled(ToolProperties->EditingMode == EWeightEditMode::Mesh);
-		MeshSelector->SetComponentSelectionMode(ToolProperties->ComponentSelectionMode);
-	}
+	MeshSelector->SetMesh(SourcePreviewMesh, ToolProperties->SourcePreviewOffset);
+	MeshSelector->SetIsEnabled(ToolProperties->EditingMode == EWeightEditMode::Mesh);
+	MeshSelector->SetComponentSelectionMode(ToolProperties->ComponentSelectionMode);
 
 	WeightTool->UpdateSelectorState();
 }
@@ -534,9 +516,10 @@ void UWeightToolTransferManager::TransferWeights()
 
 	USkinWeightsPaintToolProperties* ToolProperties = WeightTool->GetWeightToolProperties();
 
-	const USkeletalMesh* TargetSkeletalMesh = GetSkeletalMeshComponent(WeightTool->GetTarget())->GetSkeletalMeshAsset();
+	ISkeletalMeshBackedTarget* SkeletalMeshBackedTarget = Cast<ISkeletalMeshBackedTarget>(WeightTool->GetTarget()); 
+	const USkeletalMesh* TargetSkeletalMesh = SkeletalMeshBackedTarget ? SkeletalMeshBackedTarget->GetSkeletalMesh() : nullptr;
 	const bool bSameMesh = SourceSkeletalMesh == TargetSkeletalMesh;
-	const bool bSameLOD = GetLODId(ToolProperties->ActiveLOD) == GetLODId(ToolProperties->SourceLOD);
+	const bool bSameLOD = WeightTool->GetEditingLOD() == GetLODId(ToolProperties->SourceLOD);
 	const bool bSameProfile = ToolProperties->GetActiveSkinWeightProfile() == ToolProperties->SourceSkinWeightProfile;
 	const bool bVerticesSelectedOnTarget = ToolProperties->EditingMode == EWeightEditMode::Mesh && !WeightTool->GetMainMeshSelector()->GetSelectedVertices().IsEmpty();
 	const bool bVerticesSelectedOnSource = ToolProperties->EditingMode == EWeightEditMode::Mesh && !MeshSelector->GetSelectedVertices().IsEmpty();
@@ -577,7 +560,7 @@ void UWeightToolTransferManager::TransferWeightsFromOtherMeshOrSubset()
 	USkinWeightsPaintToolProperties* ToolProperties = WeightTool->GetWeightToolProperties();
 
 	// get LOD IDs
-	const EMeshLODIdentifier TargetLODId = GetLODId(ToolProperties->ActiveLOD);
+	const EMeshLODIdentifier TargetLODId = WeightTool->GetEditingLOD();
 	const EMeshLODIdentifier SourceLODId = GetLODId(ToolProperties->SourceLOD);
 
 	// get selection
@@ -587,7 +570,8 @@ void UWeightToolTransferManager::TransferWeightsFromOtherMeshOrSubset()
 	TArray<int> TargetSelectedVertices = WeightTool->GetMainMeshSelector()->GetSelectedVertices();
 	
 	// if transferring between the same mesh, ensure that the LODs are different and that the transfer is done on a subset
-	const USkeletalMesh* TargetSkeletalMesh = GetSkeletalMeshComponent(WeightTool->GetTarget())->GetSkeletalMeshAsset();
+	ISkeletalMeshBackedTarget* SkeletalMeshBackedTarget = Cast<ISkeletalMeshBackedTarget>(WeightTool->GetTarget()); 
+	const USkeletalMesh* TargetSkeletalMesh = SkeletalMeshBackedTarget ? SkeletalMeshBackedTarget->GetSkeletalMesh() : nullptr;
 	const bool bSameMeshAndLOD = SourceSkeletalMesh == TargetSkeletalMesh && TargetLODId == SourceLODId;
 	if (bSameMeshAndLOD)
 	{
@@ -599,28 +583,25 @@ void UWeightToolTransferManager::TransferWeightsFromOtherMeshOrSubset()
 		}
 	}
 
+	FDynamicMesh3& EditedMesh = WeightTool->GetMesh();
 	// get target dynamic mesh
-	const FGetMeshParameters TargetParams(true, TargetLODId);
-	FDynamicMesh3 TargetMesh = UE::ToolTarget::GetDynamicMeshCopy(WeightTool->GetTarget(), TargetParams);
+	FDynamicMesh3 TargetMesh = EditedMesh;
 	
 	// get the source dynamic mesh and validate it
 	FDynamicMesh3 TmpSourceCopy;
-	auto GetSourceMesh = [this, bSameMeshAndLOD, SourceLODId, &TmpSourceCopy](const TArray<int32>& ToIsolate) -> const FDynamicMesh3*
+	auto GetSourceMesh = [this, bSameMeshAndLOD, SourceLODId, &TmpSourceCopy, &EditedMesh](const TArray<int32>& ToIsolate) -> const FDynamicMesh3*
 	{
 		// use the current EditedMesh to get the current data without having to commit
-		FCleanedEditMesh* CleanedEditMesh = WeightTool->GetCurrentCleanedEditMesh();
-		if (const FDynamicMesh3* EditedMesh = bSameMeshAndLOD && CleanedEditMesh ? &CleanedEditMesh->GetOriginalMesh() : nullptr)
+		if (bSameMeshAndLOD)
 		{
 			if (ToIsolate.IsEmpty())
 			{
-				return EditedMesh;
+				return &EditedMesh;
 			}
-
+			
 			// create a sub-mesh from the selected triangles to filter the transfer so that it only copies from the selected components
-			const UE::Geometry::FDynamicSubmesh3 PartialSubMesh(EditedMesh, ToIsolate);
+			const UE::Geometry::FDynamicSubmesh3 PartialSubMesh(&EditedMesh, ToIsolate);
 			TmpSourceCopy = PartialSubMesh.GetSubmesh();			
-			FDynamicMeshAttributeSet* SourceAttributes = TmpSourceCopy.Attributes();
-			SourceAttributes->CopyBoneAttributes(*EditedMesh->Attributes());
 			return &TmpSourceCopy;
 		}
 
@@ -633,8 +614,6 @@ void UWeightToolTransferManager::TransferWeightsFromOtherMeshOrSubset()
 		{
 			const UE::Geometry::FDynamicSubmesh3 PartialSubMesh(&SourceMeshOrig, ToIsolate);
 			TmpSourceCopy = PartialSubMesh.GetSubmesh();
-			FDynamicMeshAttributeSet* SourceAttributes = TmpSourceCopy.Attributes();
-			SourceAttributes->CopyBoneAttributes(*SourceMeshOrig.Attributes());
 		}
 		else
 		{
@@ -667,26 +646,12 @@ void UWeightToolTransferManager::TransferWeightsFromOtherMeshOrSubset()
 	FTransferBoneWeights TransferBoneWeights(SourceMesh, ToolProperties->SourceSkinWeightProfile);
 	TransferBoneWeights.TransferMethod = FTransferBoneWeights::ETransferBoneWeightsMethod::InpaintWeights;
 	
-	if (!TargetMesh.HasAttributes())
-	{
-		TargetMesh.EnableAttributes();
-	}
-
 	FDynamicMeshAttributeSet* TargetAttributes = TargetMesh.Attributes();
 	if (!TargetAttributes->HasBones())
 	{
 		TargetAttributes->CopyBoneAttributes(*SourceMesh->Attributes());
 	}
-	else
-	{
-		const FReferenceSkeleton& TargetRefSkeleton = TargetSkeletalMesh->GetRefSkeleton();
-		if (!ensure(TargetAttributes->GetNumBones() == TargetRefSkeleton.GetRawBoneNum()))
-		{
-			const FText NotificationText = LOCTEXT("TransferWeightsInvalidBoneCount", "Target mesh bone attributes mismatch. No weights were transferred.");
-			ShowEditorMessage(ELogVerbosity::Error, NotificationText);
-			return;
-		}
-	}
+	
 
 	// NOTE should we expose all the options?
 	// 	TransferBoneWeights.NormalThreshold;
@@ -732,18 +697,11 @@ void UWeightToolTransferManager::TransferWeightsFromSameMeshAndLOD()
 	const USkinWeightsPaintToolProperties* ToolProperties = WeightTool->GetWeightToolProperties();
 	
 	// get target dynamic mesh
-	const EMeshLODIdentifier TargetLODId = GetLODId(ToolProperties->ActiveLOD);
 	const FName TargetProfile = ToolProperties->GetActiveSkinWeightProfile();
-	const FGetMeshParameters TargetParams(true, TargetLODId);
-	FDynamicMesh3 TargetMesh = UE::ToolTarget::GetDynamicMeshCopy(WeightTool->GetTarget(), TargetParams);
+	FDynamicMesh3 TargetMesh = WeightTool->GetMesh();
 	
 	const FDynamicMeshAttributeSet* TargetAttributes = TargetMesh.Attributes();
 	FDynamicMeshVertexSkinWeightsAttribute* TransferredSkinWeights = TargetAttributes->GetSkinWeightsAttribute(TargetProfile);
-	if (!TransferredSkinWeights)
-	{
-		TransferredSkinWeights = new FDynamicMeshVertexSkinWeightsAttribute(&TargetMesh);
-		TargetMesh.Attributes()->AttachSkinWeightsAttribute(TargetProfile, TransferredSkinWeights);
-	}
 		
 	const FDynamicMeshVertexSkinWeightsAttribute* SourceAttributes = TargetAttributes->GetSkinWeightsAttribute(ToolProperties->SourceSkinWeightProfile);
 	check(SourceAttributes != nullptr);
@@ -835,20 +793,29 @@ void UWeightToolTransferManager::ApplyTranferredWeightsAsTransaction(
 	const UWeightToolSelectionIsolator* Isolator = WeightTool->GetSelectionIsolator();
 	const bool bIsSelectionIsolated = Isolator->IsSelectionIsolated();
 	
-	// make sure the target has been cleaned as FSkinToolWeights data is based on a cleaned mesh description
-	FMeshDescription DummyCleanedMeshDescription;
-	FCleanedEditMesh CleanedTarget(InTargetMesh, DummyCleanedMeshDescription);
-
 	// InVertexSubset and Weights.PreChangeWeights/CurrentWeights reflect the current state of the mesh (partial or not) being edited
 	// so the target is one or the other.
 	// the full mesh is still needed to be able to get the original VertexID in order to get the transferred weights.  
-	const FDynamicMesh3& TargetMesh = bIsSelectionIsolated ? Isolator->GetPartialMesh() : CleanedTarget.GetEditableMesh();
-	
-	const int32 NumVertices = bUseSubset ? InVertexSubset.Num() : TargetMesh.VertexCount();
+	const FDynamicMesh3& TargetMesh = bIsSelectionIsolated ? Isolator->GetPartialMesh() : InTargetMesh;
+
+	int32 NumVertices = bUseSubset ? InVertexSubset.Num() : TargetMesh.MaxVertexID();
 	for (int32 VertexIndex = 0; VertexIndex < NumVertices; ++VertexIndex)
 	{
-		const int32 VertexID = bUseSubset ? InVertexSubset[VertexIndex] : VertexIndex;
+		int32 VertexID = INDEX_NONE;
+		if (bUseSubset)
+		{
+			VertexID = InVertexSubset[VertexIndex];
+		}
+		else if (TargetMesh.IsVertex(VertexIndex))
+		{
+			VertexID = VertexIndex;
+		}
 		
+		if (VertexID == INDEX_NONE)
+		{
+			continue;
+		}
+
 		// remove all weight on vertex
 		const VertexWeights& VertexBoneWeights = Weights.PreChangeWeights[VertexID];
 		if (!VertexBoneWeights.IsEmpty())
@@ -869,9 +836,8 @@ void UWeightToolTransferManager::ApplyTranferredWeightsAsTransaction(
 		}
 
 		// map from partial to cleaned mesh (if isolated) 
-		const int32 CleanedVertexID = bIsSelectionIsolated ? Isolator->PartialToFullMeshVertexIndex(VertexID) : VertexID;
-		// map from cleaned mesh to the full dynamesh
-		const int32 FullVertexID = CleanedTarget.GetOriginalVertexID(CleanedVertexID);
+		const int32 FullVertexID = bIsSelectionIsolated ? Isolator->PartialToFullMeshVertexIndex(VertexID) : VertexID;
+
 
 		// update with new weight
 		UE::AnimationCore::FBoneWeights TransferredBoneWeights;
@@ -1005,11 +971,10 @@ void FMultiBoneWeightEdits::AddEditedVerticesToSet(TSet<int32>& OutEditedVertexS
 	}
 }
 
-void FSkinToolDeformer::Initialize(const USkeletalMeshComponent* InSkelMeshComponent, const FMeshDescription* InMeshDescription)
+void FSkinToolDeformer::Initialize(const FReferenceSkeleton& InRefSkeleton, const TArray<FTransform>& PoseComponentSpace, const FDynamicMesh3& InMesh)
 {
 	// get all bone transforms in the reference pose store a copy in component space
-	Component = InSkelMeshComponent;
-	const FReferenceSkeleton& RefSkeleton = Component->GetSkeletalMeshAsset()->GetRefSkeleton();
+	RefSkeleton = InRefSkeleton;
 	const TArray<FTransform> &LocalSpaceBoneTransforms = RefSkeleton.GetRefBonePose();
 	const int32 NumBones = LocalSpaceBoneTransforms.Num();
 	InvCSRefPoseTransforms.SetNumUninitialized(NumBones);
@@ -1039,29 +1004,33 @@ void FSkinToolDeformer::Initialize(const USkeletalMeshComponent* InSkelMeshCompo
 	}
 
 	// store reference pose vertex positions
-	const TArrayView<const FVector3f> VertexPositions = InMeshDescription->GetVertexPositions().GetRawArray();
-	RefPoseVertexPositions = VertexPositions;
+	Mesh = &InMesh;
+	
+	RefPoseVertexPositions.SetNumUninitialized(Mesh->MaxVertexID());
+	for (int32 VertexIndex : Mesh->VertexIndicesItr())
+	{
+		RefPoseVertexPositions[VertexIndex] = Mesh->GetVertex(VertexIndex);
+	}
 
 	// set all vertices to be updated on first tick
 	SetAllVerticesToBeUpdated();
 
 	// record "prev" bone transforms to detect change in pose
-	PreviousPoseComponentSpace = Component->GetComponentSpaceTransforms();
+	PreviousPoseComponentSpace = PoseComponentSpace;
 }
 
 void FSkinToolDeformer::SetAllVerticesToBeUpdated()
 {
-	VerticesWithModifiedWeights.Empty(RefPoseVertexPositions.Num());
-	for (int32 VertexID=0; VertexID<RefPoseVertexPositions.Num(); ++VertexID)
+	VerticesWithModifiedWeights.Reset();
+	for (int32 VertID : Mesh->VertexIndicesItr())
 	{
-		VerticesWithModifiedWeights.Add(VertexID);
+		VerticesWithModifiedWeights.Add(VertID);
 	}
 }
 
 void FSkinToolDeformer::SetToRefPose(USkinWeightsPaintTool* Tool)
 {
 	// get ref pose
-	const FReferenceSkeleton& RefSkeleton = Component->GetSkeletalMeshAsset()->GetRefSkeleton();
 	const TArray<FTransform>& RefPoseLocalSpace = RefSkeleton.GetRefBonePose();
 	// convert to global space and store in current pose
 	FAnimationRuntime::FillUpComponentSpaceTransforms(RefSkeleton, RefPoseLocalSpace, RefPoseComponentSpace);
@@ -1103,12 +1072,12 @@ void FSkinToolDeformer::UpdateVertexDeformation(
 	// update vertex positions
 	UPreviewMesh* PreviewMesh = Tool->PreviewMesh;
 	const TArray<VertexWeights>& CurrentWeights = Tool->Weights.CurrentWeights;
-	PreviewMesh->DeferredEditMesh([this, &CurrentWeights, &PoseComponentSpace](FDynamicMesh3& Mesh)
+	PreviewMesh->DeferredEditMesh([this, &CurrentWeights, &PoseComponentSpace](FDynamicMesh3& InOutMesh)
 	{
 		TRACE_CPUPROFILER_EVENT_SCOPE(SkinTool::UpdateDeformation);
 		const TArray<int32> VertexIndices = VerticesWithModifiedWeights.Array();
 		
-		ParallelFor( VerticesWithModifiedWeights.Num(), [this, &VertexIndices, &Mesh, &PoseComponentSpace, &CurrentWeights](int32 Index)
+		ParallelFor( VerticesWithModifiedWeights.Num(), [this, &VertexIndices, &InOutMesh, &PoseComponentSpace, &CurrentWeights](int32 Index)
 		{
 			const int32 VertexID = VertexIndices[Index];
 			FVector VertexNewPosition = FVector::ZeroVector;
@@ -1123,7 +1092,7 @@ void FSkinToolDeformer::UpdateVertexDeformation(
 				VertexNewPosition += CurrentTransform.TransformPosition(VertexData.VertexInBoneSpace) * VertexData.Weight;
 			}
 			
-			Mesh.SetVertex(VertexID, VertexNewPosition, false);
+			InOutMesh.SetVertex(VertexID, VertexNewPosition, false);
 		});
 	}, false);
 	PreviewMesh->NotifyDeferredEditCompleted(UPreviewMesh::ERenderUpdateMode::FastUpdate, EMeshRenderAttributeFlags::Positions, false);
@@ -1201,29 +1170,35 @@ void FSkinToolDeformer::SetVertexNeedsUpdated(int32 VertexIndex)
 }
 
 void FSkinToolWeights::InitializeSkinWeights(
-	const USkeletalMeshComponent* SkeletalMeshComponent,
-	const FMeshDescription* Mesh)
+	USkinWeightsPaintTool* Tool,
+	const FDynamicMesh3& Mesh
+	)
 {
 	static constexpr int32 RootBoneIndex = 0;
 	static constexpr float FullWeight = 1.f;
 
+	ISkeletonProvider* SkeletonProvider = CastChecked<ISkeletonProvider>(Tool->GetTarget());
+	
 	// initialize deformer data
-	Deformer.Initialize(SkeletalMeshComponent, Mesh);
+	Deformer.Initialize(SkeletonProvider->GetSkeleton(), Tool->GetComponentSpaceBoneTransforms(), Mesh);
 
 	// initialize current weights (using compact format: num_verts * max_influences)
-	const FSkeletalMeshConstAttributes MeshAttribs(*Mesh);
-	const FSkinWeightsVertexAttributesConstRef VertexSkinWeights = MeshAttribs.GetVertexSkinWeights(Profile);
-	const int32 NumVertices = Mesh->Vertices().Num();
+	using namespace UE::Geometry;
+	using namespace UE::AnimationCore;
+	FDynamicMeshVertexSkinWeightsAttribute* VertexSkinWeights = Mesh.Attributes()->GetSkinWeightsAttribute(Profile);
+	int32 MaxVertexIndex = Mesh.MaxVertexID();
 	TMap<BoneIndex, float> NormalizedWeights;
-	CurrentWeights.SetNum(NumVertices);
-	for (int32 VertexIndex = 0; VertexIndex < NumVertices; VertexIndex++)
+	CurrentWeights.SetNum(MaxVertexIndex);
+	
+	for (int32 VertexIndex : Mesh.VertexIndicesItr())
 	{
 		// we have to normalize here because there are edges cases where skeletal meshes are loaded with non-normalized weights.
 		// even when/if we fix these, it's best to be 100% sure that things are normalized since weight data can be generated by tools we don't control
 		NormalizedWeights.Reset();
 		int32 NumInfluences = 0;
-		const FVertexBoneWeightsConst BoneWeights = VertexSkinWeights.Get(FVertexID(VertexIndex));
-		for (const UE::AnimationCore::FBoneWeight BoneWeight: BoneWeights)
+		FBoneWeights BoneWeights;
+		VertexSkinWeights->GetValue(VertexIndex, BoneWeights);
+		for (const FBoneWeight BoneWeight: BoneWeights)
 		{
 			if (!ensure(NumInfluences < MAX_TOTAL_INFLUENCES))
 			{
@@ -1266,7 +1241,7 @@ void FSkinToolWeights::InitializeSkinWeights(
 	PreChangeWeights = CurrentWeights;
 
 	// maintain relax-per stroke map
-	MaxFalloffPerVertexThisStroke.SetNumZeroed(NumVertices);
+	MaxFalloffPerVertexThisStroke.SetNumZeroed(MaxVertexIndex);
 
 	// maintain bool-per-bone if weighted or not
 	IsBoneWeighted.Init(false, Deformer.BoneNames.Num());
@@ -1274,7 +1249,7 @@ void FSkinToolWeights::InitializeSkinWeights(
 	{
 		for (const FVertexBoneWeight& VertexBoneData : VertexData)
 		{
-			if (VertexBoneData.Weight > UE::AnimationCore::BoneWeightThreshold)
+			if (VertexBoneData.Weight > BoneWeightThreshold)
 			{
 				IsBoneWeighted[VertexBoneData.BoneID] = true;
 			}
@@ -1421,10 +1396,10 @@ void FSkinToolWeights::CreateWeightEditForVertex(
 	}
 }
 
-void FSkinToolWeights::ApplyCurrentWeightsToMeshDescription(FMeshDescription* MeshDescription)
+void FSkinToolWeights::ApplyCurrentWeightsToMesh(FDynamicMesh3& Mesh)
 {
-	FSkeletalMeshAttributes MeshAttribs(*MeshDescription);
-	FSkinWeightsVertexAttributesRef VertexWeightAttrs = MeshAttribs.GetVertexSkinWeights(Profile);
+	using namespace UE::Geometry;
+	FDynamicMeshVertexSkinWeightsAttribute* VertexWeightAttrs = Mesh.Attributes()->GetSkinWeightsAttribute(Profile);
 	
 	UE::AnimationCore::FBoneWeightsSettings Settings;
 	Settings.SetNormalizeType(UE::AnimationCore::EBoneWeightNormalizeType::None);
@@ -1432,14 +1407,13 @@ void FSkinToolWeights::ApplyCurrentWeightsToMeshDescription(FMeshDescription* Me
 	TArray<UE::AnimationCore::FBoneWeight> BoneWeightsToApply;
 	BoneWeightsToApply.Reserve(UE::AnimationCore::MaxInlineBoneWeightCount);
 
-	const int32 NumVertices = MeshDescription->Vertices().Num();
-	if (!ensure(CurrentWeights.Num() == NumVertices))
+	if (!ensure(CurrentWeights.Num() == Mesh.MaxVertexID()))
 	{
-		// weights are out of sync with mesh description you're trying to apply them to
+		// weights are out of sync with mesh you're trying to apply them to
 		return;
 	}
-	
-	for (int32 VertexIndex = 0; VertexIndex < NumVertices; VertexIndex++)
+
+	for (int32 VertexIndex : Mesh.VertexIndicesItr())
 	{
 		BoneWeightsToApply.Reset();
 
@@ -1453,7 +1427,7 @@ void FSkinToolWeights::ApplyCurrentWeightsToMeshDescription(FMeshDescription* Me
 			BoneWeightsToApply.Add(UE::AnimationCore::FBoneWeight(SingleBoneWeight.BoneID, SingleBoneWeight.Weight));
 		}
 
-		VertexWeightAttrs.Set(FVertexID(VertexIndex), UE::AnimationCore::FBoneWeights::Create(BoneWeightsToApply, Settings));
+		VertexWeightAttrs->SetValue(VertexIndex, UE::AnimationCore::FBoneWeights::Create(BoneWeightsToApply, Settings));
 	}
 }
 
@@ -1696,11 +1670,7 @@ void FSkinToolWeights::UpdateIsBoneWeighted(BoneIndex BoneToUpdate)
 
 BoneIndex FSkinToolWeights::GetParentBoneToWeightTo(BoneIndex ChildBone)
 {
-	int32 ParentBoneIndex = 0;
-	if (const USkeletalMesh* SkeletalMesh = Deformer.Component->GetSkeletalMeshAsset())
-	{
-		ParentBoneIndex = SkeletalMesh->GetRefSkeleton().GetParentIndex(ChildBone);
-	}
+	int32 ParentBoneIndex = ParentBoneIndex = Deformer.RefSkeleton.GetParentIndex(ChildBone);
 
 	// are we at the root? (no parent)
 	if (ParentBoneIndex == INDEX_NONE)
@@ -1709,14 +1679,11 @@ BoneIndex FSkinToolWeights::GetParentBoneToWeightTo(BoneIndex ChildBone)
 		
 		// in this case return the first child bone, if there is one
 		// NOTE: this allows the user to forcibly remove all weight on the root bone, without having another recorded influence on it
-		if (const USkeletalMesh* SkeletalMesh = Deformer.Component->GetSkeletalMeshAsset())
+		TArray<int32> RootsChildren;
+		Deformer.RefSkeleton.GetDirectChildBones(0, RootsChildren);
+		if (!RootsChildren.IsEmpty())
 		{
-			TArray<int32> RootsChildren;
-			SkeletalMesh->GetRefSkeleton().GetDirectChildBones(0, RootsChildren);
-			if (!RootsChildren.IsEmpty())
-			{
-				ParentBoneIndex = RootsChildren[0];
-			}
+			ParentBoneIndex = RootsChildren[0];
 		}
 	}
 
@@ -1727,7 +1694,7 @@ void FMeshSkinWeightsChange::Apply(UObject* Object)
 {
 	USkinWeightsPaintTool* Tool = CastChecked<USkinWeightsPaintTool>(Object);
 
-	Tool->ExternalUpdateSkinWeightLayer(LOD, SkinWeightProfile);
+	Tool->ExternalUpdateSkinWeightLayer(SkinWeightProfile);
 	
 	// apply weight edits
 	for (TTuple<BoneIndex, FSingleBoneWeightEdits>& Pair : AllWeightEdits.PerBoneWeightEdits)
@@ -1749,7 +1716,7 @@ void FMeshSkinWeightsChange::Revert(UObject* Object)
 	USkinWeightsPaintTool* Tool = CastChecked<USkinWeightsPaintTool>(Object);
 
 	// update the skin weight profile
-	Tool->ExternalUpdateSkinWeightLayer(LOD, SkinWeightProfile);
+	Tool->ExternalUpdateSkinWeightLayer(SkinWeightProfile);
 	
 	// apply weight edits
 	for (TTuple<int32, FSingleBoneWeightEdits>& Pair : AllWeightEdits.PerBoneWeightEdits)
@@ -1837,9 +1804,10 @@ const FToolTargetTypeRequirements& USkinWeightsPaintToolBuilder::GetTargetRequir
 {
 	static FToolTargetTypeRequirements TypeRequirements({
 		UMaterialProvider::StaticClass(),
-		UMeshDescriptionProvider::StaticClass(),
-		UMeshDescriptionCommitter::StaticClass(),
-		UPrimitiveComponentBackedTarget::StaticClass()
+		UDynamicMeshProvider::StaticClass(),
+		UDynamicMeshCommitter::StaticClass(),
+		UPrimitiveComponentBackedTarget::StaticClass(),
+		USkeletonProvider::StaticClass(),
 		});
 	return TypeRequirements;
 }
@@ -1858,31 +1826,20 @@ void USkinWeightsPaintTool::Setup()
 	
 	UDynamicMeshBrushTool::Setup();
 
-	const USkeletalMeshComponent* Component = GetSkeletalMeshComponent(Target);
-	check(Component && Component->GetSkeletalMeshAsset())
+	ISkeletonProvider* SkeletonProvider = CastChecked<ISkeletonProvider>(Target);
+	const FReferenceSkeleton& RefSkeleton = SkeletonProvider->GetSkeleton();
 
-	// create a mesh description for editing (this must be done before calling UpdateBonePositionInfos)
+	IDynamicMeshProvider* DynamicMeshProvider = CastChecked<IDynamicMeshProvider>(Target);
+	EditedMesh = DynamicMeshProvider->GetDynamicMesh();
 	
-	bool bSupportsLODs = false;
-	const EMeshLODIdentifier DefaultLOD = UE::ToolTarget::GetTargetMeshDescriptionLOD(Target, bSupportsLODs);
-	ensure(bSupportsLODs);
-	const FGetMeshParameters Params(true, DefaultLOD);
-	// copy the mesh description from the tool target to edit it
-	const FMeshDescription* ToolTargetMeshDescription = UE::ToolTarget::GetMeshDescription(Target, Params);
-	EditedMeshes.Emplace(DefaultLOD, FCleanedEditMesh(*PreviewMesh->GetMesh(), *ToolTargetMeshDescription));
-	CurrentlyEditedLOD = DefaultLOD;
-
 	// create a custom set of properties inheriting from the base tool properties
 	WeightToolProperties = NewObject<USkinWeightsPaintToolProperties>(this);
 	WeightToolProperties->RestoreProperties(this);
 	WeightToolProperties->WeightTool = this;
 	WeightToolProperties->bSpecifyRadius = true;
 	// watch for skin weight layer changes
-	WeightToolProperties->ActiveLOD = GetLODName(DefaultLOD);
-	int32 WatcherIndex = WeightToolProperties->WatchProperty(WeightToolProperties->ActiveLOD, [this](FName) { OnActiveLODChanged(); });
-	WeightToolProperties->SilentUpdateWatcherAtIndex(WatcherIndex);
 	WeightToolProperties->ActiveSkinWeightProfile = FSkeletalMeshAttributesShared::DefaultSkinWeightProfileName;
-	WatcherIndex = WeightToolProperties->WatchProperty(WeightToolProperties->ActiveSkinWeightProfile, [this](FName) { OnActiveSkinWeightProfileChanged(); });
+	int32 WatcherIndex = WeightToolProperties->WatchProperty(WeightToolProperties->ActiveSkinWeightProfile, [this](FName) { OnActiveSkinWeightProfileChanged(); });
 	WeightToolProperties->SilentUpdateWatcherAtIndex(WatcherIndex);
 	WatcherIndex = WeightToolProperties->WatchProperty(WeightToolProperties->NewSkinWeightProfile, [this](FName) { OnNewSkinWeightProfileChanged(); });
 	WeightToolProperties->SilentUpdateWatcherAtIndex(WatcherIndex);
@@ -1899,7 +1856,7 @@ void USkinWeightsPaintTool::Setup()
 	RecalculateBrushRadius();
 
 	// default to the root bone as current bone
-	PendingCurrentBone = CurrentBone = Component->GetSkeletalMeshAsset()->GetRefSkeleton().GetBoneName(0);
+	PendingCurrentBone = CurrentBone = RefSkeleton.GetBoneName(0);
 
 	// configure preview mesh
 	PreviewMesh->SetTangentsMode(EDynamicMeshComponentTangentsMode::AutoCalculated);
@@ -1919,7 +1876,7 @@ void USkinWeightsPaintTool::Setup()
 	MeshSelector->InitialSetup(TargetWorld.Get(), this, OnSelectionChangedLambda);
 
 	// run all initialization for mesh/weights
-	UpdateCurrentlyEditedMesh(Component, *GetCurrentlyEditedDynamicMesh(), *GetCurrentlyEditedMeshDescription());
+	UpdateCurrentlyEditedMesh(EditedMesh);
 
 	// bind the skeletal mesh editor context
 	if (EditorContext.IsValid())
@@ -2076,9 +2033,8 @@ FBox USkinWeightsPaintTool::GetWorldSpaceFocusBox()
 	// TODO, there are several places in the engine that frame bone selections. Let's consolidate this logic.
 	if (!SelectedBoneIndices.IsEmpty())
 	{
-		const USkeletalMeshComponent* MeshComponent = Weights.Deformer.Component;
-		const FReferenceSkeleton& RefSkeleton = MeshComponent->GetSkeletalMeshAsset()->GetRefSkeleton();
-		const TArray<FTransform>& CurrentBoneTransforms = MeshComponent->GetComponentSpaceTransforms();
+		const FReferenceSkeleton& RefSkeleton = Weights.Deformer.RefSkeleton;
+		const TArray<FTransform>& CurrentBoneTransforms = GetComponentSpaceBoneTransforms();
 		if (!CurrentBoneTransforms.IsEmpty())
 		{
 			FAxisAlignedBox3d Bounds = FAxisAlignedBox3d::Empty();
@@ -2205,13 +2161,10 @@ void USkinWeightsPaintTool::OnTick(float DeltaTime)
 	}
 
 	// sparsely updates vertex positions (only on vertices with modified weights)
-	Weights.Deformer.UpdateVertexDeformation(this, Weights.Deformer.Component->GetComponentSpaceTransforms());
+	Weights.Deformer.UpdateVertexDeformation(this, GetComponentSpaceBoneTransforms());
 }
 
-void USkinWeightsPaintTool::UpdateCurrentlyEditedMesh(
-	const USkeletalMeshComponent* InComponent,
-	const FDynamicMesh3& InDynamicMesh,
-	const FMeshDescription& InMeshDescription)
+void USkinWeightsPaintTool::UpdateCurrentlyEditedMesh(const FDynamicMesh3& InDynamicMesh)
 {
 	// update the preview mesh in the viewport
 	PreviewMesh->ReplaceMesh(InDynamicMesh);
@@ -2227,8 +2180,10 @@ void USkinWeightsPaintTool::UpdateCurrentlyEditedMesh(
 	// update vertices & triangle octrees (this must be done after PreviewMesh has been updated)
 	InitializeOctrees();
 
+	IPrimitiveComponentBackedTarget* PrimitiveComponentTarget = CastChecked<IPrimitiveComponentBackedTarget>(GetTarget());
+	
 	// update the mesh selection mechanic (this must be done after PreviewMesh has been updated)
-	MeshSelector->SetMesh(PreviewMesh, &InMeshDescription, InComponent->GetComponentTransform());
+	MeshSelector->SetMesh(PreviewMesh, PrimitiveComponentTarget->GetOwnerComponent()->GetComponentTransform());
 
 	// update weights
 	Weights = FSkinToolWeights();
@@ -2238,7 +2193,7 @@ void USkinWeightsPaintTool::UpdateCurrentlyEditedMesh(
 		WeightToolProperties->bShowNewProfileName = false;
 	}
 	Weights.Profile = WeightToolProperties->GetActiveSkinWeightProfile();
-	Weights.InitializeSkinWeights(InComponent, &InMeshDescription);
+	Weights.InitializeSkinWeights(this, InDynamicMesh);
 	bVertexColorsNeedUpdated = true;
 	
 	// update smooth operator (this must be done after PreviewMesh & Weights have been updated)
@@ -2308,6 +2263,41 @@ FString FIsolateSelectionChange::ToString() const
 	return FToolCommandChange::ToString();
 }
 
+EMeshLODIdentifier USkinWeightsPaintTool::GetEditingLOD()
+{
+	if (EditorContext.IsValid())
+	{
+		return EditorContext->GetEditingLOD();
+	}
+	
+	return EMeshLODIdentifier::LOD0;
+}
+
+const TArray<FTransform>& USkinWeightsPaintTool::GetComponentSpaceBoneTransforms()
+{
+	if (EditorContext.IsValid())
+	{
+		return EditorContext->GetComponentSpaceBoneTransforms(Target);
+	}
+
+	if (DefaultComponentSpaceBoneTransforms.IsEmpty())
+	{
+		ISkeletonProvider* SkeletonProvider = CastChecked<ISkeletonProvider>(GetTarget());
+		SkeletonProvider->GetSkeleton().GetBoneAbsoluteTransforms(DefaultComponentSpaceBoneTransforms);
+	}
+		
+	return DefaultComponentSpaceBoneTransforms;
+}
+
+void USkinWeightsPaintTool::ToggleBoneManipulation(bool bEnable)
+{
+	if (EditorContext.IsValid())
+	{
+		EditorContext->ToggleBoneManipulation(bEnable);
+	}
+}
+
+
 bool USkinWeightsPaintTool::HitTest(const FRay& Ray, FHitResult& OutHit)
 {
 	const bool bIsInBrushMode = WeightToolProperties->EditingMode == EWeightEditMode::Brush;
@@ -2420,11 +2410,8 @@ double USkinWeightsPaintTool::EstimateMaximumTargetDimension()
 {
 	if (const IPrimitiveComponentBackedTarget* TargetComponent = Cast<IPrimitiveComponentBackedTarget>(Target))
 	{
-		const USkeletalMeshComponent* Component = Cast<USkeletalMeshComponent>(TargetComponent->GetOwnerComponent());
-		if (USkeletalMesh* SkeletalMesh = Component->GetSkeletalMeshAsset())
-		{
-			return SkeletalMesh->GetBounds().SphereRadius * 2.0f;
-		}
+		const UPrimitiveComponent* Component = Cast<UPrimitiveComponent>(TargetComponent->GetOwnerComponent());
+		return Component->CalcBounds(FTransform::Identity).SphereRadius * 2.0f;
 	}
 	
 	return Super::EstimateMaximumTargetDimension();
@@ -2458,19 +2445,13 @@ void USkinWeightsPaintTool::CalculateVertexROI(
 			[&](int32 VertexID) { return FVector3d::DistSquared(Mesh->GetVertex(VertexID), StampPosLocal) < RadiusSqr; },
 			OutVertexIDs);
 
-		const FNonManifoldMappingSupport NonManifoldMappingSupport(*Mesh);
-		TArray<VertexIndex> SourceVertexIDs;
-		SourceVertexIDs.Reserve(OutVertexIDs.Num());
 		OutVertexFalloffs.Reserve(OutVertexIDs.Num());
 		for (const int32 VertexID : OutVertexIDs)
 		{
 			const float DistSq = FVector3d::DistSquared(Mesh->GetVertex(VertexID), StampPosLocal);
 
-			const int32 SrcVertexId = NonManifoldMappingSupport.GetOriginalNonManifoldVertexID(VertexID);
-			SourceVertexIDs.Add(SrcVertexId);
-			OutVertexFalloffs.Add(DistanceToFalloff(SrcVertexId, DistSq));
+			OutVertexFalloffs.Add(DistanceToFalloff(VertexID, DistSq));
 		}
-		OutVertexIDs = MoveTemp(SourceVertexIDs);
 		
 		return;
 	}
@@ -2480,7 +2461,9 @@ void USkinWeightsPaintTool::CalculateVertexROI(
 		// create the ExpMap generator, computes vertex polar coordinates in a plane tangent to the surface
 		const FDynamicMesh3* Mesh = PreviewMesh->GetPreviewDynamicMesh();
 		FFrame3d SeedFrame = Mesh->GetTriFrame(TriangleUnderStamp);
-		SeedFrame.Origin = InStamp.WorldPosition;
+		
+		const FTransform3d MeshWorldTransform(PreviewMesh->GetTransform());
+		SeedFrame.Origin = MeshWorldTransform.InverseTransformPosition(FVector3d(InStamp.WorldPosition));
 		
 		TMeshLocalParam<FDynamicMesh3> Param(Mesh);
 		Param.ParamMode = ELocalParamTypes::PlanarProjection;
@@ -2488,7 +2471,6 @@ void USkinWeightsPaintTool::CalculateVertexROI(
 		Param.ComputeToMaxDistance(SeedFrame, TriVerts, InStamp.Radius * 1.5f);
 		// store vertices under the brush and their distances from the stamp
 		const float StampRadSq = FMath::Pow(InStamp.Radius, 2);
-		const FNonManifoldMappingSupport NonManifoldMappingSupport(*Mesh);
 		for (int32 VertexID : Mesh->VertexIndicesItr())
 		{
 			if (!Param.HasUV(VertexID))
@@ -2503,9 +2485,8 @@ void USkinWeightsPaintTool::CalculateVertexROI(
 				continue;
 			}
 
-			const int32 SrcVertexId = NonManifoldMappingSupport.GetOriginalNonManifoldVertexID(VertexID);
-			OutVertexFalloffs.Add(DistanceToFalloff(SrcVertexId, DistSq));
-			OutVertexIDs.Add(SrcVertexId);
+			OutVertexFalloffs.Add(DistanceToFalloff(VertexID, DistSq));
+			OutVertexIDs.Add(VertexID);
 		}
 		
 		return;
@@ -2622,13 +2603,11 @@ void USkinWeightsPaintTool::UpdateVertexColorForAllVertices()
 	// update mesh with new value colors
 	PreviewMesh->DeferredEditMesh([this, &CurrentBoneIndex](FDynamicMesh3& Mesh)
 	{
-		const UE::Geometry::FNonManifoldMappingSupport NonManifoldMappingSupport(Mesh);
 		UE::Geometry::FDynamicMeshColorOverlay* ColorOverlay = Mesh.Attributes()->PrimaryColors();
 		for (const int32 ElementId : ColorOverlay->ElementIndicesItr())
 		{
 			const int32 VertexID = ColorOverlay->GetParentVertex(ElementId);	
-			const int32 SrcVertexID = NonManifoldMappingSupport.GetOriginalNonManifoldVertexID(VertexID);
-			ColorOverlay->SetElement(ElementId, GetColorOfVertex(SrcVertexID, CurrentBoneIndex));
+			ColorOverlay->SetElement(ElementId, GetColorOfVertex(VertexID, CurrentBoneIndex));
 		}
 		
 	}, false);
@@ -2899,7 +2878,7 @@ void USkinWeightsPaintTool::InitializeSmoothWeightsOperator()
 
 	// NOTE: this could probably be initialized lazily as it's only used with the relax brush
 	const FDynamicMesh3* DynaMesh = PreviewMesh->GetMesh();
-	SmoothWeightsDataSource = MakeUnique<FPaintToolWeightsDataSource>(&Weights, *DynaMesh);
+	SmoothWeightsDataSource = MakeUnique<FPaintToolWeightsDataSource>(&Weights);
 	SmoothWeightsOp = MakeUnique<UE::Geometry::TSmoothBoneWeights<int32, float>>(DynaMesh, SmoothWeightsDataSource.Get());
 	SmoothWeightsOp->MinimumWeightThreshold = MinimumWeightThreshold;
 }
@@ -2975,17 +2954,14 @@ void USkinWeightsPaintTool::OnShutdown(EToolShutdownType ShutdownType)
 		const FName ActiveProfile = WeightToolProperties->GetActiveSkinWeightProfile();
 	
 		// apply the currently edited weights to the mesh description
-		Weights.ApplyCurrentWeightsToMeshDescription(GetCurrentlyEditedMeshDescription());
+		Weights.ApplyCurrentWeightsToMesh(EditedMesh);
 
 		// this block bakes the modified DynamicMeshComponent back into the StaticMeshComponent inside an undo transaction
 		GetToolManager()->BeginUndoTransaction(LOCTEXT("SkinWeightsPaintTool", "Paint Skin Weights"));
-		static constexpr bool bHaveTargetLOD = true;
-		for (auto& [LOD, CleanedMesh] : EditedMeshes)
-		{
-			CleanedMesh.CopyWeightsToOriginalMesh(ActiveProfile);
-			const FCommitMeshParameters CommitParams(bHaveTargetLOD, LOD);
-			UE::ToolTarget::CommitMeshDescriptionUpdate(Target, &CleanedMesh.GetOriginalMeshDescription(), nullptr, CommitParams);	
-		}
+		
+		IDynamicMeshCommitter* DynamicMeshCommitter = CastChecked<IDynamicMeshCommitter>(Target);
+		DynamicMeshCommitter->CommitDynamicMesh(EditedMesh);
+		
 		GetToolManager()->EndUndoTransaction();
 	}
 
@@ -2998,10 +2974,7 @@ void USkinWeightsPaintTool::OnShutdown(EToolShutdownType ShutdownType)
 		EditorContext->UnbindFrom(this);
 	}
 
-	if (PersonaModeManagerContext.IsValid())
-	{
-		PersonaModeManagerContext->GetPersonaEditorModeManager()->DeactivateMode(FPersonaEditModes::SkeletonSelection);
-	}
+	ToggleBoneManipulation(false);
 
 	GetWeightTransferManager()->Shutdown();
 }
@@ -3030,41 +3003,11 @@ USkinWeightsPaintToolProperties* USkinWeightsPaintTool::GetWeightToolProperties(
 	return WeightToolProperties;
 }
 
-FCleanedEditMesh* USkinWeightsPaintTool::GetCurrentCleanedEditMesh() const
-{
-	if (EditedMeshes.Contains(CurrentlyEditedLOD))
-	{
-		return &EditedMeshes[CurrentlyEditedLOD];
-	}
-
-	return nullptr;
-}
-
-FMeshDescription* USkinWeightsPaintTool::GetCurrentlyEditedMeshDescription() const
-{
-	if (FCleanedEditMesh* CleanedEditMesh = GetCurrentCleanedEditMesh())
-	{
-		return &CleanedEditMesh->GetEditableMeshDescription();
-	}
-	
-	return nullptr;
-}
-
-FDynamicMesh3* USkinWeightsPaintTool::GetCurrentlyEditedDynamicMesh() const
-{
-	if (FCleanedEditMesh* CleanedEditMesh = GetCurrentCleanedEditMesh())
-	{
-		return &CleanedEditMesh->GetEditableMesh();
-	}
-	
-	return nullptr;
-}
 
 void USkinWeightsPaintTool::BeginChange()
 {
-	const EMeshLODIdentifier LOD = GetLODId(WeightToolProperties->ActiveLOD);
 	const FName SkinProfile = WeightToolProperties->GetActiveSkinWeightProfile();
-	ActiveChange = MakeUnique<FMeshSkinWeightsChange>(LOD, SkinProfile);
+	ActiveChange = MakeUnique<FMeshSkinWeightsChange>(SkinProfile);
 }
 
 void USkinWeightsPaintTool::EndChange(const FText& TransactionLabel)
@@ -3099,12 +3042,11 @@ void USkinWeightsPaintTool::ExternalUpdateWeights(const int32 BoneIndex, const T
 	Weights.UpdateIsBoneWeighted(BoneIndex);
 }
 
-void USkinWeightsPaintTool::ExternalUpdateSkinWeightLayer(const EMeshLODIdentifier InLOD, const FName InSkinWeightProfile)
+void USkinWeightsPaintTool::ExternalUpdateSkinWeightLayer(const FName InSkinWeightProfile)
 {
 	enum class ESkinWeightChangeState
 	{
 		SkinProfile,
-		LOD,
 		None
 	} State = ESkinWeightChangeState::None;
 	
@@ -3114,19 +3056,10 @@ void USkinWeightsPaintTool::ExternalUpdateSkinWeightLayer(const EMeshLODIdentifi
 		State = ESkinWeightChangeState::SkinProfile;
 	}
 
-	const FName LODName = GetLODName(InLOD);
-	if (LODName != NAME_None && LODName != WeightToolProperties->ActiveLOD)
-	{
-		WeightToolProperties->ActiveLOD = LODName;
-		State = ESkinWeightChangeState::LOD;
-	}
-
 	switch (State)
 	{
 	case ESkinWeightChangeState::SkinProfile:
 		return OnActiveSkinWeightProfileChanged();
-	case ESkinWeightChangeState::LOD:
-		return OnActiveLODChanged();
 	case ESkinWeightChangeState::None:
 		default:
 		break;
@@ -3284,13 +3217,12 @@ void USkinWeightsPaintTool::MirrorWeights(EAxis::Type Axis, EMirrorDirection Dir
 	
 	// get all ref pose vertices
 	const TArray<FVector>& RefPoseVertices = Weights.Deformer.RefPoseVertexPositions;
-	const FReferenceSkeleton& RefSkeleton = Weights.Deformer.Component->GetSkeletalMeshAsset()->GetRefSkeleton();
 
 	// refresh mirror tables (cached / lazy generated)
 	MirrorData.EnsureMirrorDataIsUpdated(
 		Weights.Deformer.BoneNames,
 		Weights.Deformer.BoneNameToIndexMap,
-		RefSkeleton,
+		Weights.Deformer.RefSkeleton,
 		RefPoseVertices,
 		Axis,
 		Direction);
@@ -4039,47 +3971,8 @@ void USkinWeightsPaintTool::HandleSkeletalMeshModified(const TArray<FName>& InBo
 	}
 }
 
-void USkinWeightsPaintTool::OnActiveLODChanged()
-{
-	const USkeletalMeshComponent* Component = GetSkeletalMeshComponent(Target);
-	if (!ensure(Component))
-	{
-		return;
-	}
-
-	if (SelectionIsolator->IsSelectionIsolated())
-	{
-		SelectionIsolator->RestoreFullMesh();
-	}
-
-	// apply previous changes
-	Weights.ApplyCurrentWeightsToMeshDescription(GetCurrentlyEditedMeshDescription());
-
-	// update current mesh using the new LOD
-	const EMeshLODIdentifier ActiveLODID = GetLODId(WeightToolProperties->ActiveLOD);
-	const FGetMeshParameters Params(true, ActiveLODID);
-	FCleanedEditMesh* CleanedEditMesh = EditedMeshes.Find(ActiveLODID);
-	if (!CleanedEditMesh)
-	{
-		const FDynamicMesh3 DynamicMesh = UE::ToolTarget::GetDynamicMeshCopy(Target, Params);
-		const FMeshDescription* MeshDescription = UE::ToolTarget::GetMeshDescription(Target, Params);
-		CleanedEditMesh = &EditedMeshes.Add(ActiveLODID, FCleanedEditMesh(DynamicMesh, *MeshDescription));
-	}
-	
-	// assign new LOD as the currently edited one
-	CurrentlyEditedLOD = ActiveLODID;
-
-	// reinitialize all mesh data structures
-	UpdateCurrentlyEditedMesh(Component, *GetCurrentlyEditedDynamicMesh(), *GetCurrentlyEditedMeshDescription());
-}
-
 void USkinWeightsPaintTool::OnActiveSkinWeightProfileChanged()
 {
-	const USkeletalMeshComponent* SkeletalMeshComponent = GetSkeletalMeshComponent(Target);
-	if (!SkeletalMeshComponent)
-	{
-		return;
-	}
 
 	WeightToolProperties->bShowNewProfileName = WeightToolProperties->ActiveSkinWeightProfile == CreateNewName();
 
@@ -4092,7 +3985,7 @@ void USkinWeightsPaintTool::OnActiveSkinWeightProfileChanged()
 	{
 		if (!IsProfileValid(WeightToolProperties->NewSkinWeightProfile))
 		{
-			GetOrCreateSkinWeightsAttribute(*GetCurrentlyEditedMeshDescription(), WeightToolProperties->NewSkinWeightProfile);
+			GetOrCreateSkinWeightsAttribute(EditedMesh, WeightToolProperties->NewSkinWeightProfile);
 		} 
 	}
 	
@@ -4108,12 +4001,12 @@ void USkinWeightsPaintTool::OnActiveSkinWeightProfileChanged()
 	}
 
 	// apply previous changes
-	Weights.ApplyCurrentWeightsToMeshDescription(GetCurrentlyEditedMeshDescription());
+	Weights.ApplyCurrentWeightsToMesh(EditedMesh);
 
 	// re-init Weights with new skin profile
 	Weights = FSkinToolWeights();
 	Weights.Profile = WeightToolProperties->GetActiveSkinWeightProfile();
-	Weights.InitializeSkinWeights(SkeletalMeshComponent, GetCurrentlyEditedMeshDescription());
+	Weights.InitializeSkinWeights(this, EditedMesh);
 	bVertexColorsNeedUpdated = true;
 }
 
@@ -4121,7 +4014,7 @@ void USkinWeightsPaintTool::OnNewSkinWeightProfileChanged()
 {
 	if (WeightToolProperties->bShowNewProfileName && WeightToolProperties->NewSkinWeightProfile != Weights.Profile)
 	{
-		const bool bRenamed = RenameSkinWeightsAttribute(*GetCurrentlyEditedMeshDescription(), Weights.Profile, WeightToolProperties->NewSkinWeightProfile);
+		const bool bRenamed = RenameSkinWeightsAttribute(EditedMesh, Weights.Profile, WeightToolProperties->NewSkinWeightProfile);
 		if (ensure(bRenamed))
 		{
 			Weights.Profile = WeightToolProperties->NewSkinWeightProfile;
@@ -4129,23 +4022,9 @@ void USkinWeightsPaintTool::OnNewSkinWeightProfileChanged()
 	}
 }
 
-bool USkinWeightsPaintTool::IsProfileValid(const FName InProfileName) const
+bool USkinWeightsPaintTool::IsProfileValid(const FName InProfileName)
 {
-	const USkeletalMeshComponent* SkeletalMeshComponent = GetSkeletalMeshComponent(Target);
-	if (!SkeletalMeshComponent)
-	{
-		return false;
-	}
-
-	// check current MeshDescription
-	const FSkeletalMeshConstAttributes MeshAttribs(*GetCurrentlyEditedMeshDescription());
-	const TArray<FName> MeshDescProfiles = MeshAttribs.GetSkinWeightProfileNames();
-	const bool bHasProfile = MeshDescProfiles.ContainsByPredicate([InProfileName](const FName Name)
-	{
-		return Name == InProfileName;
-	});
-	
-	return bHasProfile;
+	return EditedMesh.Attributes()->HasSkinWeightsAttribute(InProfileName);
 }
 
 void USkinWeightsPaintTool::ToggleEditingMode()
@@ -4157,19 +4036,7 @@ void USkinWeightsPaintTool::ToggleEditingMode()
 	// toggle which mesh we're selecting and what components (vert/edge/face)
 	UpdateSelectorState();
 
-	// toggle bone select mode
-	// this mode is set to be compatible with the 
-	if (PersonaModeManagerContext.IsValid())
-	{
-		if (WeightToolProperties->EditingMode == EWeightEditMode::Bones)
-		{
-			PersonaModeManagerContext->GetPersonaEditorModeManager()->ActivateMode(FPersonaEditModes::SkeletonSelection);	
-		}
-		else
-		{
-			PersonaModeManagerContext->GetPersonaEditorModeManager()->DeactivateMode(FPersonaEditModes::SkeletonSelection);	
-		}
-	}
+	ToggleBoneManipulation( WeightToolProperties->EditingMode == EWeightEditMode::Bones);
 
 	SetFocusInViewport();
 }
@@ -4236,7 +4103,7 @@ void UWeightToolSelectionIsolator::UpdateIsolatedSelection()
 
 bool UWeightToolSelectionIsolator::IsSelectionIsolated() const
 {
-	return PartialMeshDescription.IsValid();
+	return PartialSubMesh.GetBaseMesh() != nullptr;
 }
 
 void UWeightToolSelectionIsolator::IsolateSelectionAsTransaction()
@@ -4290,47 +4157,38 @@ void UWeightToolSelectionIsolator::SetTrianglesToIsolate(const TArray<int32>& Tr
 
 void UWeightToolSelectionIsolator::RestoreFullMesh()
 {
-	const USkeletalMeshComponent* SkeletalMeshComponent = GetSkeletalMeshComponent(WeightTool->GetTarget());
-	if (!SkeletalMeshComponent)
-	{
-		// this can happen at shutdown
-		return;
-	}
-
-	if (!PartialMeshDescription)
+	if (!IsSelectionIsolated())
 	{
 		// nothing hidden
 		return;
 	}
 
+	FDynamicMesh3& PartialMesh = PartialSubMesh.GetSubmesh();
+	FDynamicMesh3& FullMesh = WeightTool->GetMesh();
 	FSkinToolWeights& Weights = WeightTool->GetWeights();
 
 	// apply partial mesh weights to partial mesh description
-	Weights.ApplyCurrentWeightsToMeshDescription(PartialMeshDescription.Get());
+	Weights.ApplyCurrentWeightsToMesh(PartialMesh);
 
 	// reinitialize with full mesh
-	FMeshDescription* CurrentlyEditedMesh = WeightTool->GetCurrentlyEditedMeshDescription();
-	FDynamicMesh3* CurrentlyEditedDynaMesh = WeightTool->GetCurrentlyEditedDynamicMesh();
 	// this resizes Weights to the full mesh size
-	WeightTool->UpdateCurrentlyEditedMesh(SkeletalMeshComponent, *CurrentlyEditedDynaMesh, *CurrentlyEditedMesh);
+	WeightTool->UpdateCurrentlyEditedMesh(FullMesh);
 
-	// copy the remapped weights back to the full mesh
-	const FSkeletalMeshConstAttributes MeshAttribs(*PartialMeshDescription.Get());
+	using namespace UE::Geometry;
+	using namespace UE::AnimationCore;
 	const FName ActiveProfile = WeightTool->GetWeightToolProperties()->GetActiveSkinWeightProfile();
-	const FSkinWeightsVertexAttributesConstRef AllVertexWeights = MeshAttribs.GetVertexSkinWeights(ActiveProfile);
-	const int32 NumVerticesInPartialMesh = PartialMeshDescription.Get()->Vertices().Num();
-	UE::Geometry::FNonManifoldMappingSupport NonManifoldMapping(*CurrentlyEditedDynaMesh);
-	for (int32 VertexIndexPartial = 0; VertexIndexPartial < NumVerticesInPartialMesh; VertexIndexPartial++)
+	FDynamicMeshVertexSkinWeightsAttribute* PartialWeights = PartialMesh.Attributes()->GetSkinWeightsAttribute(ActiveProfile);
+	
+	for (int32 VertexIndexPartial : PartialSubMesh.GetSubmesh().VertexIndicesItr() )
 	{
-		// get the equivalent vertex on the full mesh
 		int32 VertexIndexFull = PartialSubMesh.MapVertexToBaseMesh(VertexIndexPartial);
-		VertexIndexFull = NonManifoldMapping.GetOriginalNonManifoldVertexID(VertexIndexFull);
-
+		
 		// clear all the weights on this vertex
 		Weights.CurrentWeights[VertexIndexFull].Reset();
 		// replace with weights from partial mesh
-		const FVertexBoneWeightsConst& VertexWeightsPartial = AllVertexWeights.Get(VertexIndexPartial);
-		for (const UE::AnimationCore::FBoneWeight& SingleBoneWeight : VertexWeightsPartial)
+		FBoneWeights VertexWeightsPartial;	
+		PartialWeights->GetValue(VertexIndexPartial, VertexWeightsPartial);
+		for (const FBoneWeight& SingleBoneWeight : VertexWeightsPartial)
 		{
 			const FBoneIndexType BoneID = SingleBoneWeight.GetBoneIndex();
 			const float Weight = SingleBoneWeight.GetWeight();
@@ -4341,10 +4199,11 @@ void UWeightToolSelectionIsolator::RestoreFullMesh()
 			Weights.CurrentWeights[VertexIndexFull].Emplace(BoneID, VertexInBoneSpace, Weight);
 		}
 	}
+
 	// sync both weight buffers
 	Weights.PreChangeWeights = Weights.CurrentWeights;
 	// apply full mesh weights to full mesh description
-	Weights.ApplyCurrentWeightsToMeshDescription(CurrentlyEditedMesh);
+	Weights.ApplyCurrentWeightsToMesh(WeightTool->GetMesh());
 
 	// restore selection (a convenience for the user that allows for easily adjusting the isolation by going back/forth as needed)
 	if (UPolygonSelectionMechanic* SelectionMechanic = WeightTool->GetMainMeshSelector()->GetSelectionMechanic())
@@ -4354,13 +4213,13 @@ void UWeightToolSelectionIsolator::RestoreFullMesh()
 		SelectionMechanic->SetSelection_AsTriangleTopology(IsolatedSelectionToRestoreFaces);
 	}
 
-	PartialMeshDescription = nullptr;
+	PartialSubMesh = {};
 	CurrentlyIsolatedTriangles.Reset();
 }
 
 int32 UWeightToolSelectionIsolator::PartialToFullMeshVertexIndex(int32 PartialMeshVertexIndex) const
 {
-	if (!PartialMeshDescription)
+	if (!IsSelectionIsolated())
 	{
 		return PartialMeshVertexIndex;
 	}
@@ -4370,7 +4229,7 @@ int32 UWeightToolSelectionIsolator::PartialToFullMeshVertexIndex(int32 PartialMe
 
 int32 UWeightToolSelectionIsolator::FullToPartialMeshVertexIndex(int32 FullMeshVertexIndex) const
 {
-	if (!PartialMeshDescription)
+	if (!IsSelectionIsolated())
 	{
 		return FullMeshVertexIndex;
 	}
@@ -4381,7 +4240,7 @@ int32 UWeightToolSelectionIsolator::FullToPartialMeshVertexIndex(int32 FullMeshV
 const FDynamicMesh3& UWeightToolSelectionIsolator::GetPartialMesh() const
 {
 	static const FDynamicMesh3 Dummy;
-	return PartialMeshDescription.IsValid() ? PartialSubMesh.GetSubmesh() : Dummy;
+	return IsSelectionIsolated() ? PartialSubMesh.GetSubmesh() : Dummy;
 }
 
 void UWeightToolSelectionIsolator::CreatePartialMesh()
@@ -4392,29 +4251,17 @@ void UWeightToolSelectionIsolator::CreatePartialMesh()
 		return;
 	}
 
-	const USkeletalMeshComponent* SkeletalMeshComponent = GetSkeletalMeshComponent(WeightTool->GetTarget());
-	if (!ensure(SkeletalMeshComponent))
-	{
-		return;
-	}
-
 	if (!ensure(!CurrentlyIsolatedTriangles.IsEmpty()))
 	{
 		return;
 	}
 
-	FMeshDescription* EditedMeshDesc = WeightTool->GetCurrentlyEditedMeshDescription();
-	FDynamicMesh3* EditedDynaMesh = WeightTool->GetCurrentlyEditedDynamicMesh();
-	if (!ensure(EditedMeshDesc && EditedDynaMesh))
-	{
-		return;
-	}
-	
 	// get the weights
 	FSkinToolWeights& Weights = WeightTool->GetWeights();
 
-	// apply the current weights to the full mesh description
-	Weights.ApplyCurrentWeightsToMeshDescription(EditedMeshDesc);
+	FDynamicMesh3& FullMesh = WeightTool->GetMesh();
+	// apply the current weights to the full mesh 
+	Weights.ApplyCurrentWeightsToMesh(FullMesh);
 	
 	// put into ref pose, BEFORE copying the mesh, so that submesh deformer initializes with vertices in ref pose
 	Weights.Deformer.SetToRefPose(WeightTool);
@@ -4429,203 +4276,15 @@ void UWeightToolSelectionIsolator::CreatePartialMesh()
 	SelectionMechanic->GetSelection_AsTriangleTopology(IsolatedSelectionToRestoreVertices);
 	SelectionMechanic->GetSelection_AsTriangleTopology(IsolatedSelectionToRestoreEdges);
 	SelectionMechanic->GetSelection_AsTriangleTopology(IsolatedSelectionToRestoreFaces);
-
-	// copy the full mesh description data to the full dynamic mesh
-	// 
-	// NOTE: this copies skin weights to the dynamic mesh.
-	//	 	 this is important because the partial mesh is created from the dynamic mesh NOT the mesh description where the weight have been applied
-	// NOTE: this conversion also restores NonManifold mappings to the converted dynamic mesh
-	EditedDynaMesh->Clear();
-	FMeshDescriptionToDynamicMesh Converter;
-	Converter.Convert(EditedMeshDesc, *EditedDynaMesh);
 	
 	// create a partial sub-mesh from a subset of triangles on the dynamic mesh
-	PartialSubMesh = UE::Geometry::FDynamicSubmesh3(EditedDynaMesh, CurrentlyIsolatedTriangles);
+	PartialSubMesh = UE::Geometry::FDynamicSubmesh3(&FullMesh, CurrentlyIsolatedTriangles);
 
-	// remap the non-manifold mapping from full mesh to partial mesh if needed
-	const UE::Geometry::FNonManifoldMappingSupport OriginalNonManifoldMapping(*EditedDynaMesh);
-	if (OriginalNonManifoldMapping.IsNonManifoldVertexInSource())
-	{
-		FDynamicMesh3& IsolatedMesh = PartialSubMesh.GetSubmesh();
-		
-		TArray<int32> IsolatedMeshVertexToNonManifoldVertexID;
-		IsolatedMeshVertexToNonManifoldVertexID.Reserve(IsolatedMesh.VertexCount());
-		for (int32 IsolatedVertexID = 0; IsolatedVertexID < IsolatedMesh.VertexCount(); ++IsolatedVertexID)
-		{
-			const int32 OrigMeshVertexID = PartialSubMesh.MapVertexToBaseMesh(IsolatedVertexID);
-			const int32 SrcBaseMeshVertexID = OriginalNonManifoldMapping.GetOriginalNonManifoldVertexID(OrigMeshVertexID);
-			int32 SrcSubMeshVertexID = PartialSubMesh.MapVertexToSubmesh(SrcBaseMeshVertexID);
-			if (SrcSubMeshVertexID == INDEX_NONE)
-			{
-				SrcSubMeshVertexID = IsolatedVertexID;
-			}
-			IsolatedMeshVertexToNonManifoldVertexID.Add(SrcSubMeshVertexID);
-		}
 
-		// replace the non-manifold vertex map in the cleaned mesh attributes
-		UE::Geometry::FNonManifoldMappingSupport CleanedNonManifoldMapping(IsolatedMesh);
-		CleanedNonManifoldMapping.AttachNonManifoldVertexMappingData(IsolatedMeshVertexToNonManifoldVertexID, IsolatedMesh);
-	}
-	
-	// create mesh description for sub-mesh
-	PartialMeshDescription = MakeShared<FMeshDescription>();
-	// registering skeletal mesh attributes is required to create room to copy attributes during conversion from dynamic mesh
-	FSkeletalMeshAttributes Attributes(*PartialMeshDescription);
-	Attributes.Register();
-	// convert the partial dynamic mesh to a mesh description
-	// NOTE: this copies vertex weights to partial mesh description (later used to load weights into the tool)
-	FDynamicMeshToMeshDescription DnyToDescConverter;
-	constexpr bool bCopyTangents = true;
-	DnyToDescConverter.Convert(&PartialSubMesh.GetSubmesh(), *PartialMeshDescription, bCopyTangents);
-	
 	// reinitialize all mesh data structures
-	WeightTool->UpdateCurrentlyEditedMesh(SkeletalMeshComponent, PartialSubMesh.GetSubmesh(), *PartialMeshDescription.Get());
+	WeightTool->UpdateCurrentlyEditedMesh(PartialSubMesh.GetSubmesh());
 }
 
-FCleanedEditMesh::FCleanedEditMesh(const FDynamicMesh3& InDynamicMesh, const FMeshDescription& InMeshDescription)
-	: OriginalDynamicMesh(InDynamicMesh)
-	, OriginalMeshDescription(InMeshDescription)
-{
-	using namespace UE::Geometry;
-	
-	// store copy of the original mesh data
-	OriginalDynamicMesh = InDynamicMesh;
-	OriginalMeshDescription = InMeshDescription;
-	
-	// if the mesh doesn't have any unused or non-manifold vertices, then we can skip making a cleaned duplicate and just use it directly
-	// (the getters will return the original mesh)
-	const FNonManifoldMappingSupport OriginalNonManifoldMapping(OriginalDynamicMesh);
-	if (!OriginalDynamicMesh.HasUnusedVertices() && !OriginalNonManifoldMapping.IsNonManifoldVertexInSource())
-	{
-		return;
-	}
-	
-	// create a submesh from all the triangles in the mesh
-	// NOTE: this is where the cleaning happens. this will leave out unused vertices because they don't belong to any triangle
-	TArray<int32> AllTris;
-	AllTris.Reserve(OriginalDynamicMesh.TriangleCount());
-	for (const int32 TriID : OriginalDynamicMesh.TriangleIndicesItr())
-	{
-		AllTris.Add(TriID);
-	}
-	CleanedSubMesh = FDynamicSubmesh3(&OriginalDynamicMesh, AllTris);
-
-	// remap the non-manifold mapping from full mesh to partial mesh
-	TArray<int32> SubMeshVertexToNonManifoldVertexIDMap;
-	FDynamicMesh3& SubMeshDynamicMesh = CleanedSubMesh.GetSubmesh();
-	for (int32 SubMeshVertexID=0; SubMeshVertexID<SubMeshDynamicMesh.VertexCount(); ++SubMeshVertexID)
-	{
-		int32 BaseMeshVertexID = CleanedSubMesh.MapVertexToBaseMesh(SubMeshVertexID);
-		int32 SrcBaseMeshVertexID = OriginalNonManifoldMapping.GetOriginalNonManifoldVertexID(BaseMeshVertexID);
-		int32 SrcSubMeshVertexID = CleanedSubMesh.MapVertexToSubmesh(SrcBaseMeshVertexID);
-		if (SrcSubMeshVertexID == INDEX_NONE)
-		{
-			SrcSubMeshVertexID = SubMeshVertexID;
-		}
-		SubMeshVertexToNonManifoldVertexIDMap.Add(SrcSubMeshVertexID);
-	}
-
-	// copy bone attributes
-	const FDynamicMeshAttributeSet* OrigAttributes = InDynamicMesh.Attributes();
-	if (OrigAttributes && OrigAttributes->HasBones())
-	{
-		if (!SubMeshDynamicMesh.HasAttributes())
-		{
-			SubMeshDynamicMesh.EnableAttributes();
-		}
-		FDynamicMeshAttributeSet* SubMeshAttributes = SubMeshDynamicMesh.Attributes();
-		SubMeshAttributes->CopyBoneAttributes(*InDynamicMesh.Attributes());
-	}
-
-	// replace the non-manifold vertex map in the cleaned mesh attributes
-	FNonManifoldMappingSupport CleanedNonManifoldMapping(CleanedSubMesh.GetSubmesh());
-	CleanedNonManifoldMapping.AttachNonManifoldVertexMappingData(SubMeshVertexToNonManifoldVertexIDMap,CleanedSubMesh.GetSubmesh());
-
-	// create mesh description for sub-mesh
-	CleanedSubMeshDescription = MakeShared<FMeshDescription>();
-	// registering skeletal mesh attributes is required to create room to copy attributes during conversion from dynamic mesh
-	FSkeletalMeshAttributes Attributes(*CleanedSubMeshDescription);
-	Attributes.Register();
-	// convert the partial dynamic mesh to a mesh description
-	// NOTE: this copies vertex weights to partial mesh description (later used to load weights into the tool)
-	FDynamicMeshToMeshDescription DnyToDescConverter;
-	constexpr bool bCopyTangents = true;
-	DnyToDescConverter.Convert(&CleanedSubMesh.GetSubmesh(), *CleanedSubMeshDescription, bCopyTangents);
-}
-
-void FCleanedEditMesh::CopyWeightsToOriginalMesh(const FName InProfile)
-{
-	// if the mesh was not duplicated and cleaned, we have nothing to copy from
-	if (!CleanedSubMeshDescription)
-	{
-		return;
-	}
-
-	// ensure that this profile is available on the cleaned mesh description
-	const FSkeletalMeshConstAttributes SubMeshAttrs(*CleanedSubMeshDescription);
-	const FSkinWeightsVertexAttributesConstRef SubMeshWeightAttrs = SubMeshAttrs.GetVertexSkinWeights(InProfile);
-	if (!SubMeshWeightAttrs.IsValid())
-	{
-		const FText NotificationText = FText::Format(LOCTEXT("NoSubProfile", "Failed copying skin weights, profile '{0}' not found on sub mesh."), FText::FromName(InProfile));
-		ShowEditorMessage(ELogVerbosity::Error, NotificationText);
-		return;
-	}
-
-	FSkinWeightsVertexAttributesRef BaseMeshWeightAttrs = GetOrCreateSkinWeightsAttribute(OriginalMeshDescription, InProfile);
-	if (!BaseMeshWeightAttrs.IsValid())
-	{
-		const FText NotificationText = FText::Format(LOCTEXT("NoBaseProfile", "Failed copying skin weights, profile '{0}' not found on base mesh."), FText::FromName(InProfile));
-		ShowEditorMessage(ELogVerbosity::Error, NotificationText);
-		return;
-	}
-
-	UE::AnimationCore::FBoneWeightsSettings Settings;
-	Settings.SetNormalizeType(UE::AnimationCore::EBoneWeightNormalizeType::None);
-
-	TArray<UE::AnimationCore::FBoneWeight> BoneWeightsToApply;
-	BoneWeightsToApply.Reserve(UE::AnimationCore::MaxInlineBoneWeightCount);
-	
-	const int32 NumVerticesInSubMesh = CleanedSubMeshDescription.Get()->Vertices().Num();
-	for (int32 SubVertexID=0; SubVertexID < NumVerticesInSubMesh; SubVertexID++)
-	{
-		BoneWeightsToApply.Reset();
-		
-		const FVertexBoneWeightsConst SubVertexWeights = SubMeshWeightAttrs.Get(SubVertexID);
-		for (UE::AnimationCore::FBoneWeight SingleBoneWeight : SubVertexWeights)
-		{
-			BoneWeightsToApply.Add(UE::AnimationCore::FBoneWeight(SingleBoneWeight.GetBoneIndex(), SingleBoneWeight.GetWeight()));
-		}
-
-		// remap the submesh vertex ID to the base mesh and then copy the weights over
-		int32 BaseVertexID = CleanedSubMesh.MapVertexToBaseMesh(SubVertexID);
-		BaseMeshWeightAttrs.Set(FVertexID(BaseVertexID), UE::AnimationCore::FBoneWeights::Create(BoneWeightsToApply, Settings));
-	}
-}
-
-FDynamicMesh3& FCleanedEditMesh::GetEditableMesh()
-{
-	return CleanedSubMeshDescription ? CleanedSubMesh.GetSubmesh() : OriginalDynamicMesh;
-}
-
-FMeshDescription& FCleanedEditMesh::GetEditableMeshDescription()
-{
-	return CleanedSubMeshDescription ? *CleanedSubMeshDescription.Get() : OriginalMeshDescription;
-}
-
-FDynamicMesh3& FCleanedEditMesh::GetOriginalMesh()
-{
-	return OriginalDynamicMesh;
-}
-
-FMeshDescription& FCleanedEditMesh::GetOriginalMeshDescription()
-{
-	return OriginalMeshDescription;
-}
-
-int32 FCleanedEditMesh::GetOriginalVertexID(const int32 InCleanedVertexId) const
-{
-	return CleanedSubMeshDescription ? CleanedSubMesh.MapVertexToBaseMesh(InCleanedVertexId) : InCleanedVertexId;
-}
 
 bool USkinWeightsPaintTool::HasActiveSelectionOnMainMesh()
 {
diff --git a/Engine/Plugins/Runtime/MeshModelingToolset/Source/MeshModelingToolsEditorOnly/Private/ToolMeshSelector.cpp b/Engine/Plugins/Runtime/MeshModelingToolset/Source/MeshModelingToolsEditorOnly/Private/ToolMeshSelector.cpp
--- a/Engine/Plugins/Runtime/MeshModelingToolset/Source/MeshModelingToolsEditorOnly/Private/ToolMeshSelector.cpp
+++ b/Engine/Plugins/Runtime/MeshModelingToolset/Source/MeshModelingToolsEditorOnly/Private/ToolMeshSelector.cpp
@@ -47,13 +47,10 @@ void UToolMeshSelector::InitialSetup(UWorld* InWorld, UInteractiveTool* InParent
 
 void UToolMeshSelector::SetMesh(
 	UPreviewMesh* InPreviewMesh,
-	const FMeshDescription* InMeshDescription,
 	const FTransform3d& InMeshTransform)
 {
 	// store the mesh this is operating on
 	PreviewMesh = InPreviewMesh;
-	// copy the mesh description for selection validation
-	MeshDescription = *InMeshDescription;
 
 	if (!ensure(World))
 	{
@@ -187,10 +184,8 @@ void UToolMeshSelector::Render(IToolsContextRenderAPI* RenderAPI)
 	PolygonSelectionMechanic->Render(RenderAPI);
 }
 
-const TArray<int32>& UToolMeshSelector::GetSelectedVertices(bool bRemapToOriginal)
+const TArray<int32>& UToolMeshSelector::GetSelectedVertices()
 {
-	using UE::Geometry::FNonManifoldMappingSupport;
-
 	SelectedVerticesInternal.Empty();
 	if (!(PolygonSelectionMechanic && PreviewMesh))
 	{
@@ -199,32 +194,13 @@ const TArray<int32>& UToolMeshSelector::GetSelectedVertices(bool bRemapToOrigina
 
 	const FGroupTopologySelection& Selection = PolygonSelectionMechanic->GetActiveSelection();
 	const FDynamicMesh3* DynamicMesh = PreviewMesh->GetMesh();
-	const FVertexArray& AllVertices = MeshDescription.Vertices();
-	const FNonManifoldMappingSupport NonManifoldMappingSupport(*DynamicMesh);
 
 	// validate and add vertices to the output array
-	auto AddVertices = [this, bRemapToOriginal, &AllVertices, &NonManifoldMappingSupport](const TSet<int32>& VerticesToAdd)
+	auto AddVertices = [this](const TSet<int32>& VerticesToAdd)
 		{
-			if (bRemapToOriginal)
-			{
-				// we have to make sure that the vertex ids are safe to use as PolygonSelectionMechanic does not act on the
-				// mesh description but on the dynamic mesh that can duplicate vertices when dealing with degenerate triangles.
-				// cf. FMeshDescriptionToDynamicMesh::Convert for more details.
-				for (const int32 VertexToAdd : VerticesToAdd)
-				{
-					const int32 SrcVertexID = NonManifoldMappingSupport.GetOriginalNonManifoldVertexID(VertexToAdd);
-					if (AllVertices.IsValid(SrcVertexID))
-					{
-						SelectedVerticesInternal.Add(SrcVertexID);
-					}
-				}
-			}
-			else
+			for (const int32 VertexToAdd : VerticesToAdd)
 			{
-				for (const int32 VertexToAdd : VerticesToAdd)
-				{
-					SelectedVerticesInternal.Add(VertexToAdd);
-				}
+				SelectedVerticesInternal.Add(VertexToAdd);
 			}
 		};
 
diff --git a/Engine/Plugins/Runtime/MeshModelingToolset/Source/MeshModelingToolsEditorOnly/Public/SkeletalMesh/SkeletalMeshEditingInterface.h b/Engine/Plugins/Runtime/MeshModelingToolset/Source/MeshModelingToolsEditorOnly/Public/SkeletalMesh/SkeletalMeshEditingInterface.h
new file mode 100644
--- /dev/null
+++ b/Engine/Plugins/Runtime/MeshModelingToolset/Source/MeshModelingToolsEditorOnly/Public/SkeletalMesh/SkeletalMeshEditingInterface.h
@@ -0,0 +1,166 @@
+﻿// Copyright Epic Games, Inc. All Rights Reserved.
+
+#pragma once
+
+#include "UObject/Interface.h"
+#include "UObject/WeakInterfacePtr.h"
+#include "SkeletalMeshNotifier.h"
+
+#include "HitProxies.h"
+#include "ToolContextInterfaces.h"
+#include "GenericPlatform/ICursor.h"
+#include "TargetInterfaces/MeshTargetInterfaceTypes.h"
+
+#include "SkeletalMeshEditingInterface.generated.h"
+
+#define UE_API MESHMODELINGTOOLSEDITORONLY_API
+
+class FSkeletalMeshToolNotifier;
+class UTransformProxy;
+class UInteractiveToolManager;
+class IGizmoStateTarget;
+class USkeletonModifier;
+struct FInputDeviceRay;
+class UToolTarget;
+
+/**
+ * USkeletalMeshEditingInterface
+ */
+
+UINTERFACE(MinimalAPI)
+class USkeletalMeshEditingInterface : public UInterface
+{
+	GENERATED_BODY()
+};
+
+/**
+ * ISkeletalMeshEditingInterface
+ */
+
+class ISkeletalMeshEditingInterface
+{
+	GENERATED_BODY()
+
+public:
+	UE_API ISkeletalMeshNotifier& GetNotifier();
+	UE_API bool NeedsNotification() const;
+
+	UE_API virtual TArray<FName> GetSelectedBones() const;
+	UE_API virtual TArray<FName> GetSelectedMorphTargets() const;
+	
+	UE_API void BindTo(TSharedPtr<ISkeletalMeshEditorBinding> InBinding);
+	UE_API void Unbind();
+
+	UE_API virtual TWeakObjectPtr<USkeletonModifier> GetModifier() const;
+	
+protected:
+	virtual void HandleSkeletalMeshModified(const TArray<FName>& BoneNames, const ESkeletalMeshNotifyType InNotifyType) = 0;
+	UE_API TOptional<FName> GetBoneName(HHitProxy* InHitProxy) const;
+	
+	TWeakPtr<ISkeletalMeshEditorBinding> Binding;
+
+private:
+	TUniquePtr<FSkeletalMeshToolNotifier> Notifier;
+	
+	friend FSkeletalMeshToolNotifier;
+};
+
+/**
+ * FSkeletalMeshToolNotifier
+ */
+
+class FSkeletalMeshToolNotifier: public ISkeletalMeshNotifier
+{
+public:
+	FSkeletalMeshToolNotifier(TWeakInterfacePtr<ISkeletalMeshEditingInterface> InInterface);
+	virtual void HandleNotification(const TArray<FName>& BoneNames, const ESkeletalMeshNotifyType InNotifyType) override;
+	
+protected:
+	TWeakInterfacePtr<ISkeletalMeshEditingInterface> Interface;
+};
+
+/**
+ * HBoneHitProxy
+ */
+
+struct HBoneHitProxy : public HHitProxy
+{
+	DECLARE_HIT_PROXY( UE_API )
+
+	int32 BoneIndex;
+	FName BoneName;
+
+	HBoneHitProxy(int32 InBoneIndex, FName InBoneName)
+		: HHitProxy(HPP_Foreground)
+		, BoneIndex(InBoneIndex)
+		, BoneName(InBoneName)
+	{}
+
+	virtual EMouseCursor::Type GetMouseCursor() override { return EMouseCursor::Crosshairs; }
+};
+
+/**
+ * USkeletalMeshGizmoContextObjectBase
+ */
+
+UCLASS(MinimalAPI, Abstract)
+class USkeletalMeshGizmoContextObjectBase : public UObject
+{
+	GENERATED_BODY()
+	
+public:
+	virtual USkeletalMeshGizmoWrapperBase* GetNewWrapper(UInteractiveToolManager* InToolManager, UObject* Outer = nullptr, IGizmoStateTarget* InStateTarget = nullptr)
+	PURE_VIRTUAL(UGizmoContextObject::GetNewWrapper, return nullptr;);
+};
+
+/**
+ * USkeletalMeshGizmoWrapperBase
+ */
+
+UCLASS(MinimalAPI, Abstract)
+class USkeletalMeshGizmoWrapperBase : public UObject
+{
+	GENERATED_BODY()
+
+protected:
+	using FGetTransform = TUniqueFunction<FTransform(void)>;
+	using FSetTransform = TUniqueFunction<void(const FTransform&)>;
+	
+public:
+	virtual void Initialize(const FTransform& InTransform = FTransform::Identity, const EToolContextCoordinateSystem& InTransformMode = EToolContextCoordinateSystem::Local)
+	PURE_VIRTUAL(USkeletalMeshGizmoWrapperBase::Initialize, return;)
+	
+	virtual void HandleBoneTransform(FGetTransform GetTransformFunc, FSetTransform SetTransformFunc)
+	PURE_VIRTUAL(USkeletalMeshGizmoWrapperBase::HandleBoneTransform, return;)
+
+	virtual void Clear() PURE_VIRTUAL(USkeletalMeshGizmoWrapperBase::Clear, return;)
+
+	virtual bool CanInteract() const
+	PURE_VIRTUAL(USkeletalMeshGizmoWrapperBase::CanInteract, return false;)
+	
+	virtual bool IsGizmoHit(const FInputDeviceRay& PressPos) const PURE_VIRTUAL(USkeletalMeshGizmoWrapperBase::IsGizmoHit, return false;)
+
+	UPROPERTY()
+	TWeakObjectPtr<USceneComponent> Component;
+};
+
+/**
+ * USkeletalMeshEditorContextObjectBase
+ */
+
+UCLASS(MinimalAPI, Abstract)
+class USkeletalMeshEditorContextObjectBase : public UObject
+{
+	GENERATED_BODY()
+
+public:
+	virtual EMeshLODIdentifier GetEditingLOD() PURE_VIRTUAL(USkeletalMeshEditorContextObjectBase::GetEditingLOD, return EMeshLODIdentifier::LOD0;)
+	virtual void HideSkeleton() PURE_VIRTUAL(USkeletalMeshEditorContextObjectBase::HideSkeleton, return;)
+	virtual void ShowSkeleton() PURE_VIRTUAL(USkeletalMeshEditorContextObjectBase::ShowSkeleton, return;)
+	virtual void ToggleBoneManipulation(bool bEnable) PURE_VIRTUAL(USkeletalMeshEditorContextObjectBase::ToggleBoneManipulation, return;)
+	virtual const TArray<FTransform>& GetComponentSpaceBoneTransforms(UToolTarget* InToolTarget) PURE_VIRTUAL(USkeletalMeshEditorContextObjectBase::GetComponentSpaceBoneTransforms, static TArray<FTransform> Dummy; return Dummy;)
+	virtual void BindTo(ISkeletalMeshEditingInterface* InEditingInterface) PURE_VIRTUAL(USkeletalMeshEditorContextObjectBase::BindTo, return;)
+	virtual void UnbindFrom(ISkeletalMeshEditingInterface* InEditingInterface) PURE_VIRTUAL(USkeletalMeshEditorContextObjectBase::UnbindFrom, return;)
+};
+
+#undef UE_API
diff --git a/Engine/Plugins/Runtime/MeshModelingToolset/Source/MeshModelingToolsEditorOnly/Public/SkeletalMesh/SkeletalMeshEditionInterface.h b/Engine/Plugins/Runtime/MeshModelingToolset/Source/MeshModelingToolsEditorOnly/Public/SkeletalMesh/SkeletalMeshEditionInterface.h
deleted file mode 100644
diff --git a/Engine/Plugins/Runtime/MeshModelingToolset/Source/MeshModelingToolsEditorOnly/Public/SkeletalMesh/SkeletonEditingTool.h b/Engine/Plugins/Runtime/MeshModelingToolset/Source/MeshModelingToolsEditorOnly/Public/SkeletalMesh/SkeletonEditingTool.h
--- a/Engine/Plugins/Runtime/MeshModelingToolset/Source/MeshModelingToolsEditorOnly/Public/SkeletalMesh/SkeletonEditingTool.h
+++ b/Engine/Plugins/Runtime/MeshModelingToolset/Source/MeshModelingToolsEditorOnly/Public/SkeletalMesh/SkeletonEditingTool.h
@@ -8,7 +8,7 @@
 
 #include "SkeletalMeshNotifier.h"
 #include "SkeletonModifier.h"
-#include "SkeletalMesh/SkeletalMeshEditionInterface.h"
+#include "SkeletalMesh/SkeletalMeshEditingInterface.h"
 #include "Engine/World.h"
 #include "Changes/ValueWatcher.h"
 #include "Mechanics/RectangleMarqueeMechanic.h"
@@ -146,7 +146,7 @@ public:
 	UE_API void UnParentBones();
 	UE_API void SnapBoneToComponentSelection(const bool bCreate);
 	
-	// ISkeletalMeshEditionInterface overrides
+	// ISkeletalMeshEditingInterface overrides
 	UE_API virtual TArray<FName> GetSelectedBones() const override;
 
 	UE_API const TArray<FName>& GetSelection() const;
@@ -157,7 +157,7 @@ public:
 	// IModifierToggleBehaviorTarget overrides
 	UE_API virtual void OnUpdateModifierState(int ModifierID, bool bIsOn) override;
 
-	// ISkeletalMeshEditionInterface overrides
+	// ISkeletalMeshEditingInterface overrides
 	UE_API virtual TWeakObjectPtr<USkeletonModifier> GetModifier() const override;
 
 	UE_API EEditingOperation GetOperation() const;
@@ -170,7 +170,7 @@ protected:
 	// UInteractiveTool overrides
 	UE_API virtual void RegisterActions(FInteractiveToolActionSet& ActionSet) override;
 
-	// ISkeletalMeshEditionInterface overrides
+	// ISkeletalMeshEditingInterface overrides
 	UE_API virtual void HandleSkeletalMeshModified(const TArray<FName>& BoneNames, const ESkeletalMeshNotifyType InNotifyType) override;
 
 	// Modifier functions
diff --git a/Engine/Plugins/Runtime/MeshModelingToolset/Source/MeshModelingToolsEditorOnly/Public/SkeletalMesh/SkinWeightsBindingTool.h b/Engine/Plugins/Runtime/MeshModelingToolset/Source/MeshModelingToolsEditorOnly/Public/SkeletalMesh/SkinWeightsBindingTool.h
--- a/Engine/Plugins/Runtime/MeshModelingToolset/Source/MeshModelingToolsEditorOnly/Public/SkeletalMesh/SkinWeightsBindingTool.h
+++ b/Engine/Plugins/Runtime/MeshModelingToolset/Source/MeshModelingToolsEditorOnly/Public/SkeletalMesh/SkinWeightsBindingTool.h
@@ -7,7 +7,7 @@
 #include "BoneContainer.h"
 #include "InteractiveToolBuilder.h"
 #include "ModelingOperators.h"
-#include "SkeletalMeshEditionInterface.h"
+#include "SkeletalMeshEditingInterface.h"
 #include "BaseTools/MultiSelectionMeshEditingTool.h"
 #include "Interfaces/Interface_BoneReferenceSkeletonProvider.h"
 
@@ -134,7 +134,7 @@ protected:
 	// updates the skin weights using the operator's result
 	UE_API bool UpdateSkinWeightsFromDynamicMesh(UE::Geometry::FDynamicMesh3& InResultMesh) const;
 	
-	// ISkeletalMeshEditionInterface
+	// ISkeletalMeshEditingInterface
 	UE_API virtual void HandleSkeletalMeshModified(const TArray<FName>& InBoneNames, const ESkeletalMeshNotifyType InNotifyType) override;
 
 	// Pending update function */
diff --git a/Engine/Plugins/Runtime/MeshModelingToolset/Source/MeshModelingToolsEditorOnly/Public/SkeletalMesh/SkinWeightsPaintTool.h b/Engine/Plugins/Runtime/MeshModelingToolset/Source/MeshModelingToolsEditorOnly/Public/SkeletalMesh/SkinWeightsPaintTool.h
--- a/Engine/Plugins/Runtime/MeshModelingToolset/Source/MeshModelingToolsEditorOnly/Public/SkeletalMesh/SkinWeightsPaintTool.h
+++ b/Engine/Plugins/Runtime/MeshModelingToolset/Source/MeshModelingToolsEditorOnly/Public/SkeletalMesh/SkinWeightsPaintTool.h
@@ -13,7 +13,7 @@
 #include "Misc/Optional.h"
 #include "Containers/Map.h"
 #include "DynamicMesh/DynamicMeshOctree3.h"
-#include "SkeletalMesh/SkeletalMeshEditionInterface.h"
+#include "SkeletalMesh/SkeletalMeshEditingInterface.h"
 #include "Engine/SkeletalMesh.h"
 #include "Selections/GeometrySelection.h"
 #include "TargetInterfaces/MeshTargetInterfaceTypes.h"
@@ -115,7 +115,7 @@ namespace SkinPaintTool
 	// data required to preview the skinning deformations as you paint
 	struct FSkinToolDeformer
 	{
-		void Initialize(const USkeletalMeshComponent* InSkelMeshComponent, const FMeshDescription* InMeshDescription);
+		void Initialize(const FReferenceSkeleton& InRefSkeleton, const TArray<FTransform>& PoseComponentSpace, const FDynamicMesh3& InMesh);
 
 		void SetAllVerticesToBeUpdated();
 
@@ -129,6 +129,8 @@ namespace SkinPaintTool
 		TSet<int32> VerticesWithModifiedWeights;
 		// position of all vertices in the reference pose
 		TArray<FVector> RefPoseVertexPositions;
+		// Store a ref to the Mesh so there is a way to iterate over valid verts
+		const FDynamicMesh3* Mesh;
 		// inverted, component space ref pose transform of each bone
 		TArray<FTransform> InvCSRefPoseTransforms;
 		// bones transforms used in last deformation update
@@ -138,8 +140,9 @@ namespace SkinPaintTool
 		// bone index to bone name
 		TArray<FName> BoneNames;
 		TMap<FName, BoneIndex> BoneNameToIndexMap;
-		// the skeletal mesh to get the current pose from
-		const USkeletalMeshComponent* Component;
+		
+		FReferenceSkeleton RefSkeleton;
+		
 	};
 
 	// store a sparse set of modifications to a set of vertex weights on a SINGLE bone
@@ -175,9 +178,8 @@ namespace SkinPaintTool
 	class FMeshSkinWeightsChange : public FToolCommandChange
 	{
 	public:
-		FMeshSkinWeightsChange(const EMeshLODIdentifier InLOD, const FName InSkinWeightProfile)
+		FMeshSkinWeightsChange(const FName InSkinWeightProfile)
 			: FToolCommandChange()
-			, LOD(InLOD)
 			, SkinWeightProfile(InSkinWeightProfile)
 		{}
 
@@ -196,7 +198,6 @@ namespace SkinPaintTool
 
 	private:
 		FMultiBoneWeightEdits AllWeightEdits;
-		EMeshLODIdentifier LOD = EMeshLODIdentifier::Default;
 		FName SkinWeightProfile = FSkeletalMeshAttributesShared::DefaultSkinWeightProfileName;
 		
 	};
@@ -206,8 +207,8 @@ namespace SkinPaintTool
 	{
 		// copy the initial weight values from the skeletal mesh
 		void InitializeSkinWeights(
-			const USkeletalMeshComponent* SkeletalMeshComponent,
-			const FMeshDescription* Mesh);
+			USkinWeightsPaintTool* Tool,
+			const FDynamicMesh3& Mesh);
 
 		// applies an edit to a single vertex weight on a single bone, then normalizes the remaining weights while
 		// keeping the edited weight intact (ie, adapts OTHER influences to achieve normalization)
@@ -217,7 +218,7 @@ namespace SkinPaintTool
 			float NewWeightValue,
 			FMultiBoneWeightEdits& WeightEdits);
 
-		void ApplyCurrentWeightsToMeshDescription(FMeshDescription* MeshDescription);
+		void ApplyCurrentWeightsToMesh(FDynamicMesh3& Mesh);
 		
 		static float GetWeightOfBoneOnVertex(
 			const int32 BoneIndex,
@@ -421,9 +422,6 @@ public:
 	UPROPERTY(Config)
 	FSkinWeightBrushConfig BrushConfigRelax;
 
-	// skin weight layer properties
-	UPROPERTY(EditAnywhere, Category = SkinWeightLayer, meta = (DisplayName = "Active LOD", GetOptions = GetTargetLODsFunc))
-	FName ActiveLOD = "LOD0";
 	UPROPERTY(EditAnywhere, Category = SkinWeightLayer, meta = (DisplayName = "Active Profile", GetOptions = GetTargetSkinWeightProfilesFunc))
 	FName ActiveSkinWeightProfile = FSkeletalMeshAttributesShared::DefaultSkinWeightProfileName;
 	
@@ -457,8 +455,6 @@ public:
 	FTransform SourcePreviewOffset = FTransform::Identity;
 	
 private:
-	UFUNCTION()
-	UE_API TArray<FName> GetTargetLODsFunc() const;
 	UFUNCTION()
 	UE_API TArray<FName> GetTargetSkinWeightProfilesFunc() const;
 	UFUNCTION()
@@ -589,7 +585,6 @@ private:
 		
 	// when selection is isolated, we hide the full mesh and show a submesh
 	// when islated selection is unhidden, we remap all changes from the submesh back to the full mesh
-	TSharedPtr<FMeshDescription> PartialMeshDescription = nullptr; // only non-null during isolated selection
 	TArray<int32> CurrentlyIsolatedTriangles;
 	UE::Geometry::FGeometrySelection IsolatedSelectionToRestoreVertices;
 	UE::Geometry::FGeometrySelection IsolatedSelectionToRestoreEdges;
@@ -610,31 +605,7 @@ public:
 	UE_API virtual FString ToString() const override;
 };
 
-// Why does this exist?
-// It is important for the mesh selection algorithms to operate on a mesh that is free from floating vertices,
-// but changing the topology of the mesh the tool is working on causes issue with the non-manifold mapping, UV attributes and possibly other things
-//
-// So, rather than track these issues down individually, we leave the original mesh topology intact, and operate on a cleaned submesh.
-// When the tool exits, we copy the weight edits from the submesh to the full mesh.
-// This means weights of the floating vertices cannot be edited, however they do not matter since they do not belong to any triangles they cannot be rendered
-struct FCleanedEditMesh
-{
-	UE_API FCleanedEditMesh(const FDynamicMesh3& InDynamicMesh, const FMeshDescription& InMeshDescription);
-	UE_API void CopyWeightsToOriginalMesh(const FName InProfile);
-	UE_API FDynamicMesh3& GetEditableMesh();
-	UE_API FMeshDescription& GetEditableMeshDescription();
-	UE_API FDynamicMesh3& GetOriginalMesh();
-	UE_API FMeshDescription& GetOriginalMeshDescription();
-
-	// returns the VertexID on the original mesh 
-	UE_API int32 GetOriginalVertexID(const int32 InCleanedVertexId) const;
-	
-private:
-	FDynamicMesh3 OriginalDynamicMesh;
-	FMeshDescription OriginalMeshDescription;
-	UE::Geometry::FDynamicSubmesh3 CleanedSubMesh;
-	TSharedPtr<FMeshDescription> CleanedSubMeshDescription = nullptr;
-};
+
 
 // An interactive tool for painting and editing skin weights.
 UCLASS(MinimalAPI)
@@ -643,6 +614,11 @@ class USkinWeightsPaintTool : public UDynamicMeshBrushTool, public ISkeletalMesh
 	GENERATED_BODY()
 
 public:
+	// Derived tools can optionally override these if they are invoked in editors that don't offer skeletal mesh editor context object
+	UE_API virtual EMeshLODIdentifier GetEditingLOD();
+	UE_API virtual const TArray<FTransform>& GetComponentSpaceBoneTransforms();
+	UE_API virtual void ToggleBoneManipulation(bool bEnable);
+	
 
 	// UBaseBrushTool overrides
 	virtual bool HasCancel() const override { return true; }
@@ -675,7 +651,7 @@ public:
 
 	// using when ToolChange is applied via Undo/Redo
 	UE_API void ExternalUpdateWeights(const int32 BoneIndex, const TMap<int32, float>& IndexValues);
-	UE_API void ExternalUpdateSkinWeightLayer(const EMeshLODIdentifier InLOD, const FName InSkinWeightProfile);
+	UE_API void ExternalUpdateSkinWeightLayer(const FName InSkinWeightProfile);
 	UE_API void ExternalAddInfluenceToVertices(const BoneIndex InfluenceToAdd, const TArray<VertexIndex>& Vertices);
 	UE_API void ExternalRemoveInfluenceFromVertices(const BoneIndex InfluenceToRemove, const TArray<VertexIndex>& Vertices);
 
@@ -763,10 +739,8 @@ public:
 	// get the tool properties
 	UE_API USkinWeightsPaintToolProperties* GetWeightToolProperties() const;
 
-	// get access to the mesh description for the mesh being edited
-	UE_API FCleanedEditMesh* GetCurrentCleanedEditMesh() const;
-	UE_API FMeshDescription* GetCurrentlyEditedMeshDescription() const;
-	UE_API FDynamicMesh3* GetCurrentlyEditedDynamicMesh() const;
+	// get access to the mesh being edited
+	UE_API FDynamicMesh3& GetMesh() {return EditedMesh;};
 
 	// HOW TO EDIT WEIGHTS WITH UNDO/REDO:
 	//
@@ -787,10 +761,7 @@ public:
 	UE_API void ApplyWeightEditsAsTransaction(const SkinPaintTool::FMultiBoneWeightEdits& WeightEdits, const FText& TransactionLabel);
 
 	// call this whenever the target mesh is modified
-	UE_API void UpdateCurrentlyEditedMesh(
-		const USkeletalMeshComponent* Component,
-		const FDynamicMesh3& InDynamicMesh,
-		const FMeshDescription& InMeshDescription);
+	UE_API void UpdateCurrentlyEditedMesh(const FDynamicMesh3& InDynamicMesh);
 
 	// Update the brush indicators visibility state
 	UE_API void UpdateBrushIndicators();
@@ -855,8 +826,7 @@ protected:
 	UE_API virtual void OnPropertyModified(UObject* ModifiedObject, FProperty* ModifiedProperty) override;
 	
 	// the currently edited mesh descriptions
-	mutable TMap<EMeshLODIdentifier, FCleanedEditMesh> EditedMeshes;
-	EMeshLODIdentifier CurrentlyEditedLOD;
+	FDynamicMesh3 EditedMesh;
 
 	// storage of vertex weights per bone 
 	SkinPaintTool::FSkinToolWeights Weights;
@@ -888,7 +858,7 @@ protected:
 	TArray<FName> SelectedBoneNames;
 	TArray<BoneIndex> SelectedBoneIndices;
 
-	// ISkeletalMeshEditionInterface
+	// ISkeletalMeshEditingInterface
 	UE_API virtual void HandleSkeletalMeshModified(const TArray<FName>& InBoneNames, const ESkeletalMeshNotifyType InNotifyType) override;
 
 	// the selection system for the main mesh
@@ -896,10 +866,9 @@ protected:
 	TObjectPtr<UToolMeshSelector> MeshSelector;
 
 	// skin weight layer
-	UE_API void OnActiveLODChanged();
 	UE_API void OnActiveSkinWeightProfileChanged();
 	UE_API void OnNewSkinWeightProfileChanged();
-	UE_API bool IsProfileValid(const FName InProfileName) const;
+	UE_API bool IsProfileValid(const FName InProfileName);
 
 	// global properties stored on initialization
 	UPROPERTY()
@@ -920,6 +889,8 @@ protected:
 	// editor state to restore when exiting the paint tool
 	FString PreviewProfileToRestore;
 
+	TArray<FTransform> DefaultComponentSpaceBoneTransforms;
+
 	friend SkinPaintTool::FSkinToolDeformer;
 };
 
diff --git a/Engine/Plugins/Runtime/MeshModelingToolset/Source/MeshModelingToolsEditorOnly/Public/ToolMeshSelector.h b/Engine/Plugins/Runtime/MeshModelingToolset/Source/MeshModelingToolsEditorOnly/Public/ToolMeshSelector.h
--- a/Engine/Plugins/Runtime/MeshModelingToolset/Source/MeshModelingToolsEditorOnly/Public/ToolMeshSelector.h
+++ b/Engine/Plugins/Runtime/MeshModelingToolset/Source/MeshModelingToolsEditorOnly/Public/ToolMeshSelector.h
@@ -55,7 +55,6 @@ public:
 	// passing in a null preview mesh will disable the selector
 	UE_API void SetMesh(
 		UPreviewMesh* InMesh,
-		const FMeshDescription* InMeshDescription,
 		const FTransform3d& InMeshTransform);
 
 	UE_API void UpdateAfterMeshDeformation();
@@ -71,7 +70,7 @@ public:
 	UE_API void Render(IToolsContextRenderAPI* RenderAPI);
 
 	// get a list of currently selected vertices (converting edges and faces to vertices)
-	UE_API const TArray<int32>& GetSelectedVertices(bool bRemapToOriginal = true);
+	UE_API const TArray<int32>& GetSelectedVertices();
 	UE_API bool IsAnyComponentSelected() const;
 	UE_API void GetSelectedTriangles(TArray<int32>& OutTriangleIndices) const;
 
@@ -92,7 +91,6 @@ private:
 	TObjectPtr<UWorld> World;
 	UPROPERTY()
 	TObjectPtr<UPreviewMesh> PreviewMesh;
-	FMeshDescription MeshDescription;
 	UPROPERTY()
 	TObjectPtr<UPolygonSelectionMechanic> PolygonSelectionMechanic;
 
diff --git a/Engine/Plugins/Runtime/MeshModelingToolset/Source/ModelingComponentsEditorOnly/Private/ToolTargets/SkeletalMeshComponentToolTarget.cpp b/Engine/Plugins/Runtime/MeshModelingToolset/Source/ModelingComponentsEditorOnly/Private/ToolTargets/SkeletalMeshComponentToolTarget.cpp
--- a/Engine/Plugins/Runtime/MeshModelingToolset/Source/ModelingComponentsEditorOnly/Private/ToolTargets/SkeletalMeshComponentToolTarget.cpp
+++ b/Engine/Plugins/Runtime/MeshModelingToolset/Source/ModelingComponentsEditorOnly/Private/ToolTargets/SkeletalMeshComponentToolTarget.cpp
@@ -124,8 +124,9 @@ const FMeshDescription* USkeletalMeshComponentReadOnlyToolTarget::GetMeshDescrip
 {
 	if (ensure(IsValid()))
 	{
-		const USkeletalMesh* SkeletalMesh = GetSkeletalMesh(); 
-		const int32 LODIndex = USkeletalMeshReadOnlyToolTarget::GetValidEditingLOD(SkeletalMesh, DefaultEditingLOD, GetMeshParams.bHaveRequestLOD, GetMeshParams.RequestLOD);
+		const USkeletalMesh* SkeletalMesh = GetSkeletalMesh();
+		EMeshLODIdentifier UseLOD = GetMeshParams.bHaveRequestLOD ? GetMeshParams.RequestLOD : EditingLOD;
+		const int32 LODIndex = USkeletalMeshReadOnlyToolTarget::GetValidEditingLOD(SkeletalMesh, UseLOD);
 	
 		return SkeletalMesh->GetMeshDescription(LODIndex);
 	}
@@ -179,8 +180,8 @@ void USkeletalMeshComponentToolTarget::CommitMeshDescription(const FCommitter& C
 	if (ensure(IsValid()))
 	{
 		USkeletalMesh* SkeletalMesh = GetSkeletalMesh();
-		
-		const int32 LODIndex = USkeletalMeshReadOnlyToolTarget::GetValidEditingLOD(SkeletalMesh, DefaultEditingLOD, CommitMeshParams.bHaveTargetLOD, CommitMeshParams.TargetLOD);
+		EMeshLODIdentifier UseLOD = CommitMeshParams.bHaveTargetLOD ? CommitMeshParams.TargetLOD : EditingLOD;
+		const int32 LODIndex = USkeletalMeshReadOnlyToolTarget::GetValidEditingLOD(SkeletalMesh, UseLOD);
 
 		// Unregister the component while we update its skeletal mesh
 		FComponentReregisterContext ComponentReregisterContext(Component.Get());
@@ -279,6 +280,7 @@ UToolTarget* USkeletalMeshComponentToolTargetFactory::BuildTarget(UObject* Sourc
 		Target = NewObject<USkeletalMeshComponentReadOnlyToolTarget>();
 	}	
 	Target->InitializeComponent(Cast<USkeletalMeshComponent>(SourceObject));
+	Target->EditingLOD = EditingLOD;
 	checkSlow(Target->Component.IsValid() && Requirements.AreSatisfiedBy(Target));
 
 	return Target;
@@ -296,4 +298,9 @@ bool USkeletalMeshComponentToolTargetFactory::CanWriteToSource(const UObject* So
 	return false;
 }
 
+void USkeletalMeshComponentToolTargetFactory::SetActiveEditingLOD(EMeshLODIdentifier NewEditingLOD)
+{
+	EditingLOD = NewEditingLOD;
+}
+
 #undef LOCTEXT_NAMESPACE
diff --git a/Engine/Plugins/Runtime/MeshModelingToolset/Source/ModelingComponentsEditorOnly/Private/ToolTargets/SkeletalMeshToolTarget.cpp b/Engine/Plugins/Runtime/MeshModelingToolset/Source/ModelingComponentsEditorOnly/Private/ToolTargets/SkeletalMeshToolTarget.cpp
--- a/Engine/Plugins/Runtime/MeshModelingToolset/Source/ModelingComponentsEditorOnly/Private/ToolTargets/SkeletalMeshToolTarget.cpp
+++ b/Engine/Plugins/Runtime/MeshModelingToolset/Source/ModelingComponentsEditorOnly/Private/ToolTargets/SkeletalMeshToolTarget.cpp
@@ -139,27 +139,25 @@ bool USkeletalMeshReadOnlyToolTarget::CommitMaterialSetUpdate(USkeletalMesh* Ske
 	return true;
 }
 
-int32 USkeletalMeshReadOnlyToolTarget::GetValidEditingLOD(const USkeletalMesh* SkeletalMesh, int32 DefaultEditingLOD, bool bHaveRequestLOD, EMeshLODIdentifier RequestedEditingLOD)
+int32 USkeletalMeshReadOnlyToolTarget::GetValidEditingLOD(const USkeletalMesh* SkeletalMesh, EMeshLODIdentifier RequestedEditingLOD)
 {
 	using namespace SkeletalMeshToolTargetLocals;
 
-	int32 ValidEditingLOD = DefaultEditingLOD;
+	int32 ValidEditingLOD = 0;
 
-	if (bHaveRequestLOD)
+	if (RequestedEditingLOD == EMeshLODIdentifier::MaxQuality)
 	{
-		if (RequestedEditingLOD == EMeshLODIdentifier::MaxQuality)
-		{
-			ValidEditingLOD = 0;
-		}
-		else if (RequestedEditingLOD == EMeshLODIdentifier::HiResSource)
-		{
-			DisplayCriticalWarningMessage(LOCTEXT("SkelmeshHiresLODUnavailable", "HiRes LOD not available on skeletal meshes. Falling back to using LOD 0."));
-		}
-		else if (RequestedEditingLOD != EMeshLODIdentifier::Default)
-		{
-			ValidEditingLOD = static_cast<int32>(RequestedEditingLOD);
-		}
+		ValidEditingLOD = 0;
+	}
+	else if (RequestedEditingLOD == EMeshLODIdentifier::HiResSource)
+	{
+		DisplayCriticalWarningMessage(LOCTEXT("SkelmeshHiresLODUnavailable", "HiRes LOD not available on skeletal meshes. Falling back to using LOD 0."));
+	}
+	else if (RequestedEditingLOD != EMeshLODIdentifier::Default)
+	{
+		ValidEditingLOD = static_cast<int32>(RequestedEditingLOD);
 	}
+
 	
 	const int32 MaxLOD = SkeletalMesh->GetLODNum() - 1;
 	if (ValidEditingLOD > MaxLOD)
@@ -175,7 +173,8 @@ const FMeshDescription* USkeletalMeshReadOnlyToolTarget::GetMeshDescription(cons
 {
 	if (ensure(IsValid()))
 	{
-		const int32 LODIndex = GetValidEditingLOD(SkeletalMesh.Get(), DefaultEditingLOD, GetMeshParams.bHaveRequestLOD, GetMeshParams.RequestLOD);
+		EMeshLODIdentifier UseLOD = GetMeshParams.bHaveRequestLOD ? GetMeshParams.RequestLOD : EditingLOD;
+		const int32 LODIndex = GetValidEditingLOD(SkeletalMesh.Get(),UseLOD);
 	
 		return SkeletalMesh->GetMeshDescription(LODIndex);
 	}
@@ -221,9 +220,11 @@ void USkeletalMeshToolTarget::CommitMeshDescription(const FCommitter& Committer,
 {
 	if (ensure(IsValid()))
 	{
-		const int32 LODIndex = GetValidEditingLOD(SkeletalMesh.Get(), DefaultEditingLOD, CommitMeshParams.bHaveTargetLOD, CommitMeshParams.TargetLOD);
+		USkeletalMesh* Mesh = GetSkeletalMesh();
+		EMeshLODIdentifier UseLOD = CommitMeshParams.bHaveTargetLOD ? CommitMeshParams.TargetLOD : EditingLOD;
+		const int32 LODIndex = GetValidEditingLOD(Mesh, UseLOD);
 		
-		CommitMeshDescription(SkeletalMesh.Get(), Committer, LODIndex);
+		CommitMeshDescription(Mesh, Committer, LODIndex);
 	}
 }
 
@@ -350,5 +351,10 @@ bool USkeletalMeshToolTargetFactory::CanWriteToSource(const UObject* SourceObjec
 	return false;
 }
 
+void USkeletalMeshToolTargetFactory::SetActiveEditingLOD(EMeshLODIdentifier NewEditingLOD)
+{
+	EditingLOD = NewEditingLOD;
+}
+
 #undef LOCTEXT_NAMESPACE
 
diff --git a/Engine/Plugins/Runtime/MeshModelingToolset/Source/ModelingComponentsEditorOnly/Public/ToolTargets/SkeletalMeshComponentToolTarget.h b/Engine/Plugins/Runtime/MeshModelingToolset/Source/ModelingComponentsEditorOnly/Public/ToolTargets/SkeletalMeshComponentToolTarget.h
--- a/Engine/Plugins/Runtime/MeshModelingToolset/Source/ModelingComponentsEditorOnly/Public/ToolTargets/SkeletalMeshComponentToolTarget.h
+++ b/Engine/Plugins/Runtime/MeshModelingToolset/Source/ModelingComponentsEditorOnly/Public/ToolTargets/SkeletalMeshComponentToolTarget.h
@@ -47,7 +47,7 @@ public:
 	UE_API virtual FMeshDescription GetEmptyMeshDescription() override;
 	virtual bool SupportsLODs() const override { return true; }
 	UE_API virtual TArray<EMeshLODIdentifier> GetAvailableLODs(bool bSkipAutoGenerated = true) const override;
-	virtual EMeshLODIdentifier GetMeshDescriptionLOD() const override { return static_cast<EMeshLODIdentifier>(DefaultEditingLOD); }
+	virtual EMeshLODIdentifier GetMeshDescriptionLOD() const override { return EditingLOD; }
 
 	
 	// IMaterialProvider implementation
@@ -67,8 +67,8 @@ public:
 	UE_API virtual FReferenceSkeleton GetSkeleton() const override;
 	
 protected:
-	// The LOD to edit if requested to get/commit at EMeshLODIdentifier::Default or if no specific LOD is requested.
-	int32 DefaultEditingLOD = 0;	
+	// LOD to edit, default is to edit LOD0
+	EMeshLODIdentifier EditingLOD = EMeshLODIdentifier::LOD0;	
 	
 	// So that the tool target factory can poke into Component.
 	friend class USkeletalMeshComponentReadOnlyToolTargetFactory;
@@ -135,6 +135,13 @@ public:
 	UE_API virtual UToolTarget* BuildTarget(UObject* SourceObject, const FToolTargetTypeRequirements& TargetTypeInfo) override;
 
 	static UE_API bool CanWriteToSource(const UObject* SourceObject);
+	
+	virtual EMeshLODIdentifier GetActiveEditingLOD() const { return EditingLOD; }
+	UE_API virtual void SetActiveEditingLOD(EMeshLODIdentifier NewEditingLOD);
+
+protected:
+	// LOD to edit, default is to edit LOD0
+	EMeshLODIdentifier EditingLOD = EMeshLODIdentifier::LOD0;
 };
 
 #undef UE_API
diff --git a/Engine/Plugins/Runtime/MeshModelingToolset/Source/ModelingComponentsEditorOnly/Public/ToolTargets/SkeletalMeshToolTarget.h b/Engine/Plugins/Runtime/MeshModelingToolset/Source/ModelingComponentsEditorOnly/Public/ToolTargets/SkeletalMeshToolTarget.h
--- a/Engine/Plugins/Runtime/MeshModelingToolset/Source/ModelingComponentsEditorOnly/Public/ToolTargets/SkeletalMeshToolTarget.h
+++ b/Engine/Plugins/Runtime/MeshModelingToolset/Source/ModelingComponentsEditorOnly/Public/ToolTargets/SkeletalMeshToolTarget.h
@@ -44,7 +44,7 @@ public:
 	UE_API virtual FMeshDescription GetEmptyMeshDescription() override;
 	virtual bool SupportsLODs() const override { return true; }
 	UE_API virtual TArray<EMeshLODIdentifier> GetAvailableLODs(bool bSkipAutoGenerated = true) const override;
-	virtual EMeshLODIdentifier GetMeshDescriptionLOD() const override { return static_cast<EMeshLODIdentifier>(DefaultEditingLOD); }
+	virtual EMeshLODIdentifier GetMeshDescriptionLOD() const override { return EditingLOD; }
 
 	// IMaterialProvider implementation
 	UE_API virtual int32 GetNumMaterials() const override;
@@ -63,7 +63,7 @@ protected:
 	TWeakObjectPtr<USkeletalMesh> SkeletalMesh = nullptr;
 
 	// The LOD to edit if requested to get/commit at EMeshLODIdentifier::Default or if no specific LOD is requested.
-	int32 DefaultEditingLOD = 0;	
+	EMeshLODIdentifier EditingLOD = EMeshLODIdentifier::LOD0;
 	
 	// So that the tool target factory can poke into Component.
 	friend class USkeletalMeshReadOnlyToolTargetFactory;
@@ -79,7 +79,7 @@ protected:
 		bool bPreferAssetMaterials);
 	static UE_API bool CommitMaterialSetUpdate(USkeletalMesh* SkeletalMesh,
 		const FComponentMaterialSet& MaterialSet, bool bApplyToAsset);
-	static UE_API int32 GetValidEditingLOD(const USkeletalMesh* SkeletalMesh, int32 DefaultEditingLOD, bool bHaveRequestLOD, EMeshLODIdentifier RequestedEditingLOD);	
+	static UE_API int32 GetValidEditingLOD(const USkeletalMesh* SkeletalMesh, EMeshLODIdentifier RequestedEditingLOD);	
 };
 
 
@@ -140,6 +140,13 @@ public:
 	UE_API virtual UToolTarget* BuildTarget(UObject* SourceObject, const FToolTargetTypeRequirements& TargetTypeInfo) override;
 
 	static UE_API bool CanWriteToSource(const UObject* SourceObject);
+
+	virtual EMeshLODIdentifier GetActiveEditingLOD() const { return EditingLOD; }
+	UE_API virtual void SetActiveEditingLOD(EMeshLODIdentifier NewEditingLOD);
+
+protected:
+	// LOD to edit, default is to edit LOD0
+	EMeshLODIdentifier EditingLOD = EMeshLODIdentifier::LOD0;
 };
 
 #undef UE_API
