[Skin weight paint tool] fixed minor SA issue

diff --git a/Engine/Plugins/Runtime/MeshModelingToolset/Source/MeshModelingToolsEditorOnly/Private/SkeletalMesh/SkinWeightsPaintTool.cpp b/Engine/Plugins/Runtime/MeshModelingToolset/Source/MeshModelingToolsEditorOnly/Private/SkeletalMesh/SkinWeightsPaintTool.cpp
--- a/Engine/Plugins/Runtime/MeshModelingToolset/Source/MeshModelingToolsEditorOnly/Private/SkeletalMesh/SkinWeightsPaintTool.cpp
+++ b/Engine/Plugins/Runtime/MeshModelingToolset/Source/MeshModelingToolsEditorOnly/Private/SkeletalMesh/SkinWeightsPaintTool.cpp
@@ -1670,7 +1670,7 @@ void FSkinToolWeights::UpdateIsBoneWeighted(BoneIndex BoneToUpdate)
 
 BoneIndex FSkinToolWeights::GetParentBoneToWeightTo(BoneIndex ChildBone)
 {
-	int32 ParentBoneIndex = ParentBoneIndex = Deformer.RefSkeleton.GetParentIndex(ChildBone);
+	int32 ParentBoneIndex = Deformer.RefSkeleton.GetParentIndex(ChildBone);
 
 	// are we at the root? (no parent)
 	if (ParentBoneIndex == INDEX_NONE)
