ModelingTools: Fix potential crash in numerical UI when transform proxy is not set.

diff --git a/Engine/Plugins/Runtime/MeshModelingToolset/Source/ModelingComponents/Private/Transforms/TransformGizmoDataBinder.cpp b/Engine/Plugins/Runtime/MeshModelingToolset/Source/ModelingComponents/Private/Transforms/TransformGizmoDataBinder.cpp
--- a/Engine/Plugins/Runtime/MeshModelingToolset/Source/ModelingComponents/Private/Transforms/TransformGizmoDataBinder.cpp
+++ b/Engine/Plugins/Runtime/MeshModelingToolset/Source/ModelingComponents/Private/Transforms/TransformGizmoDataBinder.cpp
@@ -471,10 +471,12 @@ void FTransformGizmoDataBinder::UnbindFromGizmo(UCombinedTransformGizmo* Gizmo,
 	// Note about this ensure: it might seem ok to unbind once the ActiveTarget is no longer around,
 	//  but it can be dangerous if the proxy is reused elsewhere instead of being immediately destroyed,
 	//  causing us to continue to trigger these callbacks.
-	if (ensure(Gizmo->ActiveTarget 
+	ensure(Gizmo->ActiveTarget
 		// The second condition here is for cases where we destroyed a gizmo without actually ever calling
 		//  SetActiveTarget on it, meaning that we never fully bound to it in the first place.
-		|| !BoundGizmos.Contains(Gizmo)))
+		|| !BoundGizmos.Contains(Gizmo));
+
+	if (Gizmo->ActiveTarget)
 	{
 		Gizmo->ActiveTarget->OnBeginTransformEdit.RemoveAll(this);
 		Gizmo->ActiveTarget->OnTransformChanged.RemoveAll(this);
