[SKM ModelingTools] fixed tool target not inheriting LOD from tool target factory

diff --git a/Engine/Plugins/Runtime/MeshModelingToolset/Source/ModelingComponentsEditorOnly/Private/ToolTargets/SkeletalMeshToolTarget.cpp b/Engine/Plugins/Runtime/MeshModelingToolset/Source/ModelingComponentsEditorOnly/Private/ToolTargets/SkeletalMeshToolTarget.cpp
--- a/Engine/Plugins/Runtime/MeshModelingToolset/Source/ModelingComponentsEditorOnly/Private/ToolTargets/SkeletalMeshToolTarget.cpp
+++ b/Engine/Plugins/Runtime/MeshModelingToolset/Source/ModelingComponentsEditorOnly/Private/ToolTargets/SkeletalMeshToolTarget.cpp
@@ -337,6 +337,7 @@ UToolTarget* USkeletalMeshToolTargetFactory::BuildTarget(UObject* SourceObject,
 		Target = NewObject<USkeletalMeshReadOnlyToolTarget>();
 	}
 	Target->SkeletalMesh = Cast<USkeletalMesh>(SourceObject);
+	Target->EditingLOD = EditingLOD;
 	checkSlow(Target->SkeletalMesh.IsValid() && Requirements.AreSatisfiedBy(Target));
 
 	return Target;
