[SkeletalMeshModelingTools] Fixed skin weight bind tool using a uninitialized skeleton

diff --git a/Engine/Plugins/Runtime/MeshModelingToolset/Source/MeshModelingToolsEditorOnly/Private/SkeletalMesh/SkinWeightsBindingTool.cpp b/Engine/Plugins/Runtime/MeshModelingToolset/Source/MeshModelingToolsEditorOnly/Private/SkeletalMesh/SkinWeightsBindingTool.cpp
--- a/Engine/Plugins/Runtime/MeshModelingToolset/Source/MeshModelingToolsEditorOnly/Private/SkeletalMesh/SkinWeightsBindingTool.cpp
+++ b/Engine/Plugins/Runtime/MeshModelingToolset/Source/MeshModelingToolsEditorOnly/Private/SkeletalMesh/SkinWeightsBindingTool.cpp
@@ -12,6 +12,7 @@
 #include "TargetInterfaces/MeshDescriptionProvider.h"
 #include "TargetInterfaces/MeshDescriptionCommitter.h"
 #include "TargetInterfaces/PrimitiveComponentBackedTarget.h"
+#include "TargetInterfaces/SkeletonProvider.h"
 #include "ToolSetupUtil.h"
 #include "ToolTargetManager.h"
 #include "Components/SkeletalMeshComponent.h"
@@ -49,7 +50,8 @@ const FToolTargetTypeRequirements& USkinWeightsBindingToolBuilder::GetTargetRequ
 		UMaterialProvider::StaticClass(),
 		UMeshDescriptionProvider::StaticClass(),
 		UMeshDescriptionCommitter::StaticClass(),
-		UPrimitiveComponentBackedTarget::StaticClass()
+		UPrimitiveComponentBackedTarget::StaticClass(),
+		USkeletonProvider::StaticClass()
 		});
 	return TypeRequirements;
 }
@@ -98,13 +100,8 @@ void USkinWeightsBindingTool::Setup()
 		return;
 	}
 	
-	const USkeletalMeshComponent* SkeletalMeshComponent = Cast<USkeletalMeshComponent>(UE::ToolTarget::GetTargetComponent(Targets[0]));
-	
-	if (SkeletalMeshComponent && SkeletalMeshComponent->GetSkeletalMeshAsset())
-	{
-		USkeletalMesh* SkeletalMesh = SkeletalMeshComponent->GetSkeletalMeshAsset();
-		ReferenceSkeleton = SkeletalMesh->GetRefSkeleton();
-	}
+	ISkeletonProvider* SkeletonProvider = CastChecked<ISkeletonProvider>(Targets[0]);
+	ReferenceSkeleton = SkeletonProvider->GetSkeleton();
 	
 	UE::ToolTarget::HideSourceObject(Targets[0]);
 	Preview = NewObject<UMeshOpPreviewWithBackgroundCompute>(this, "Preview");
