[SKMModelingTool] Fixed crash when entering skeleton tool if a prior change adds multiple root bones to the skeleton + cannot change widget mode when in skeleton tool

diff --git a/Engine/Plugins/Animation/SkeletalMeshModelingTools/Source/SkeletalMeshModelingTools/Private/SkeletalMeshModelingToolsEditorMode.cpp b/Engine/Plugins/Animation/SkeletalMeshModelingTools/Source/SkeletalMeshModelingTools/Private/SkeletalMeshModelingToolsEditorMode.cpp
--- a/Engine/Plugins/Animation/SkeletalMeshModelingTools/Source/SkeletalMeshModelingTools/Private/SkeletalMeshModelingToolsEditorMode.cpp
+++ b/Engine/Plugins/Animation/SkeletalMeshModelingTools/Source/SkeletalMeshModelingTools/Private/SkeletalMeshModelingToolsEditorMode.cpp
@@ -937,6 +937,8 @@ bool USkeletalMeshModelingToolsEditorMode::ShouldDrawWidget() const
 
 bool USkeletalMeshModelingToolsEditorMode::UsesTransformWidget() const
 {
+	// This function determines if the GetWidgetLocation() of this editor mode should be used
+	
 	if (GetCurrentEditingCache()->IsDynamicMeshSkeletonEnabled())
 	{
 		return true;
@@ -947,12 +949,14 @@ bool USkeletalMeshModelingToolsEditorMode::UsesTransformWidget() const
 
 bool USkeletalMeshModelingToolsEditorMode::UsesTransformWidget(UE::Widget::EWidgetMode CheckMode) const
 {
+	// This function determines if the specific widget mode can be used
+	
 	if (GetCurrentEditingCache()->IsDynamicMeshSkeletonEnabled())
 	{
 		return ShouldDrawWidget() && (CheckMode == UE::Widget::WM_Scale || CheckMode == UE::Widget::WM_Translate || CheckMode == UE::Widget::WM_Rotate);
 	}
 	
-	return false;
+	return Super::UsesTransformWidget(CheckMode);
 }
 
 FVector USkeletalMeshModelingToolsEditorMode::GetWidgetLocation() const
diff --git a/Engine/Plugins/Runtime/MeshModelingToolset/Source/SkeletalMeshModifiers/Private/SkeletonModifier.cpp b/Engine/Plugins/Runtime/MeshModelingToolset/Source/SkeletalMeshModifiers/Private/SkeletonModifier.cpp
--- a/Engine/Plugins/Runtime/MeshModelingToolset/Source/SkeletalMeshModifiers/Private/SkeletonModifier.cpp
+++ b/Engine/Plugins/Runtime/MeshModelingToolset/Source/SkeletalMeshModifiers/Private/SkeletonModifier.cpp
@@ -506,6 +506,8 @@ bool USkeletonModifier::CommitSkeletonToDynamicMesh()
 	TSharedPtr<FDynamicMesh3> OldMeshShared = MakeShared<FDynamicMesh3>();
 	DynamicMesh->ProcessMesh([&](const FDynamicMesh3& ReadMesh) { *OldMeshShared = ReadMesh; });
 
+	EnsureSingleRoot();
+	
 	DynamicMesh->EditMesh([&](FDynamicMesh3& Mesh)
 	{
 		const int32 NumBones = ReferenceSkeleton->GetNum();
@@ -607,30 +609,7 @@ bool USkeletonModifier::IsReferenceSkeletonValid(const bool bLog) const
 
 ESkeletalMeshModificationType USkeletonModifier::PreCommitSkeletalMesh()
 {
-	// before commiting, we have to reparent non-root bones with no parent as the animation pipeline
-	// doesn't support multi-roots
-	const TArray<FMeshBoneInfo>& BoneInfos = ReferenceSkeleton->GetRawRefBoneInfo();
-	if (!BoneInfos.IsEmpty())
-	{
-		TArray<FName> BonesToParent;
-		for (int32 Index = 1; Index < BoneInfos.Num(); Index++)
-		{
-			const FMeshBoneInfo& BoneInfo = BoneInfos[Index];
-			if (BoneInfo.ParentIndex == INDEX_NONE)
-			{
-				BonesToParent.Add(BoneInfo.Name);
-			}
-		}
-		
-		if (!BonesToParent.IsEmpty())
-		{
-			for (const FName& BoneName: BonesToParent)
-			{
-				UE_LOG(LogAnimation, Warning, TEXT("Skeleton Modifier: %s will be parented to the root bone before commiting."), *BoneName.ToString());
-			}
-			ParentBones(BonesToParent, {BoneInfos[0].Name});
-		}
-	}
+	EnsureSingleRoot();
 
 	// check topological change from this modifier
 	auto HasBoneIndexesChanged = [this]()
@@ -1101,6 +1080,34 @@ void USkeletonModifier::PostCommitSkeleton(const ESkeletonModificationType InSke
 	}
 }
 
+void USkeletonModifier::EnsureSingleRoot()
+{
+	// before commiting, we have to reparent non-root bones with no parent as the animation pipeline
+	// doesn't support multi-roots
+	const TArray<FMeshBoneInfo>& BoneInfos = ReferenceSkeleton->GetRawRefBoneInfo();
+	if (!BoneInfos.IsEmpty())
+	{
+		TArray<FName> BonesToParent;
+		for (int32 Index = 1; Index < BoneInfos.Num(); Index++)
+		{
+			const FMeshBoneInfo& BoneInfo = BoneInfos[Index];
+			if (BoneInfo.ParentIndex == INDEX_NONE)
+			{
+				BonesToParent.Add(BoneInfo.Name);
+			}
+		}
+		
+		if (!BonesToParent.IsEmpty())
+		{
+			for (const FName& BoneName: BonesToParent)
+			{
+				UE_LOG(LogAnimation, Warning, TEXT("Skeleton Modifier: %s will be parented to the root bone before commiting."), *BoneName.ToString());
+			}
+			ParentBones(BonesToParent, {BoneInfos[0].Name});
+		}
+	}	
+}
+
 bool USkeletonModifier::HasAnyOtherReferences(const bool bSkeletalMeshOnly) const
 {
 	USkeleton* Skeleton = SkeletalMesh->GetSkeleton();
diff --git a/Engine/Plugins/Runtime/MeshModelingToolset/Source/SkeletalMeshModifiers/Public/SkeletonModifier.h b/Engine/Plugins/Runtime/MeshModelingToolset/Source/SkeletalMeshModifiers/Public/SkeletonModifier.h
--- a/Engine/Plugins/Runtime/MeshModelingToolset/Source/SkeletalMeshModifiers/Public/SkeletonModifier.h
+++ b/Engine/Plugins/Runtime/MeshModelingToolset/Source/SkeletalMeshModifiers/Public/SkeletonModifier.h
@@ -284,6 +284,8 @@ private:
 	UE_API ESkeletonModificationType PreCommitSkeleton(const ESkeletalMeshModificationType InSkeletalMeshModifications) const;
 	UE_API void PostCommitSkeleton(const ESkeletonModificationType InSkeletonModifications) const;
 
+	UE_API void EnsureSingleRoot();
+	
 	// does the skeleton have any references other than the skeletal mesh currently being edited
 	UE_API bool HasAnyOtherReferences(const bool bSkeletalMeshOnly) const;
 
