[Skin Weight Paint tool] fixed brush indicator not updating change changing brush mode

diff --git a/Engine/Plugins/Runtime/MeshModelingToolset/Source/MeshModelingToolsEditorOnly/Private/SkeletalMesh/SkinWeightsPaintTool.cpp b/Engine/Plugins/Runtime/MeshModelingToolset/Source/MeshModelingToolsEditorOnly/Private/SkeletalMesh/SkinWeightsPaintTool.cpp
--- a/Engine/Plugins/Runtime/MeshModelingToolset/Source/MeshModelingToolsEditorOnly/Private/SkeletalMesh/SkinWeightsPaintTool.cpp
+++ b/Engine/Plugins/Runtime/MeshModelingToolset/Source/MeshModelingToolsEditorOnly/Private/SkeletalMesh/SkinWeightsPaintTool.cpp
@@ -2246,6 +2246,7 @@ void USkinWeightsPaintToolProperties::SetBrushMode(EWeightEditOperation InBrushM
 	BrushStrength = GetBrushConfig().Strength;
 	BrushFalloffAmount = GetBrushConfig().Falloff;
 
+	WeightTool->UpdateBrushIndicators();
 	WeightTool->SetFocusInViewport();
 }
 
