[SKMmodellingtools] change default component selection mode to face due to user feedback

diff --git a/Engine/Plugins/Runtime/MeshModelingToolset/Source/MeshModelingToolsEditorOnly/Private/SkeletalMesh/SkeletonEditingTool.cpp b/Engine/Plugins/Runtime/MeshModelingToolset/Source/MeshModelingToolsEditorOnly/Private/SkeletalMesh/SkeletonEditingTool.cpp
--- a/Engine/Plugins/Runtime/MeshModelingToolset/Source/MeshModelingToolsEditorOnly/Private/SkeletalMesh/SkeletonEditingTool.cpp
+++ b/Engine/Plugins/Runtime/MeshModelingToolset/Source/MeshModelingToolsEditorOnly/Private/SkeletalMesh/SkeletonEditingTool.cpp
@@ -304,9 +304,9 @@ void USkeletonEditingTool::SetupComponentsSelection()
 	SelectionMechanic->SetMarqueeSelectionUpdateType(EMarqueeSelectionUpdateType::OnRelease);
 	SelectionMechanic->Properties->bDisplayPolygroupReliantControls = false;
 
-	SelectionMechanic->Properties->bSelectVertices = true;
+	SelectionMechanic->Properties->bSelectVertices = false;
 	SelectionMechanic->Properties->bSelectEdges = false;
-	SelectionMechanic->Properties->bSelectFaces = false;
+	SelectionMechanic->Properties->bSelectFaces = true;
 	SelectionMechanic->SetShowSelectableCorners(true);
 
 	// adjust selection rendering for this context
diff --git a/Engine/Plugins/Runtime/MeshModelingToolset/Source/MeshModelingToolsEditorOnly/Public/SkeletalMesh/SkinWeightsPaintTool.h b/Engine/Plugins/Runtime/MeshModelingToolset/Source/MeshModelingToolsEditorOnly/Public/SkeletalMesh/SkinWeightsPaintTool.h
--- a/Engine/Plugins/Runtime/MeshModelingToolset/Source/MeshModelingToolsEditorOnly/Public/SkeletalMesh/SkinWeightsPaintTool.h
+++ b/Engine/Plugins/Runtime/MeshModelingToolset/Source/MeshModelingToolsEditorOnly/Public/SkeletalMesh/SkinWeightsPaintTool.h
@@ -380,7 +380,7 @@ public:
 
 	// are we selecting vertices, edges or faces
 	UPROPERTY(Config)
-	EComponentSelectionMode ComponentSelectionMode;
+	EComponentSelectionMode ComponentSelectionMode = EComponentSelectionMode::Faces;
 
 	// weight color properties
 	UPROPERTY(EditAnywhere, Config, Category = MeshDisplay)
