SkinWeight painting tool : prevent tool from crashing if the dynamic mesh does not have any skin weights

diff --git a/Engine/Plugins/Runtime/MeshModelingToolset/Source/MeshModelingToolsEditorOnly/Private/SkeletalMesh/SkinWeightsPaintTool.cpp b/Engine/Plugins/Runtime/MeshModelingToolset/Source/MeshModelingToolsEditorOnly/Private/SkeletalMesh/SkinWeightsPaintTool.cpp
--- a/Engine/Plugins/Runtime/MeshModelingToolset/Source/MeshModelingToolsEditorOnly/Private/SkeletalMesh/SkinWeightsPaintTool.cpp
+++ b/Engine/Plugins/Runtime/MeshModelingToolset/Source/MeshModelingToolsEditorOnly/Private/SkeletalMesh/SkinWeightsPaintTool.cpp
@@ -1194,57 +1194,64 @@ void FSkinToolWeights::InitializeSkinWeights(
 	// initialize current weights (using compact format: num_verts * max_influences)
 	using namespace UE::Geometry;
 	using namespace UE::AnimationCore;
-	FDynamicMeshVertexSkinWeightsAttribute* VertexSkinWeights = Mesh.Attributes()->GetSkinWeightsAttribute(Profile);
+
 	int32 MaxVertexIndex = Mesh.MaxVertexID();
 	TMap<BoneIndex, float> NormalizedWeights;
 	CurrentWeights.SetNum(MaxVertexIndex);
-	
-	for (int32 VertexIndex : Mesh.VertexIndicesItr())
+
+	if (FDynamicMeshVertexSkinWeightsAttribute* VertexSkinWeights = Mesh.Attributes()->GetSkinWeightsAttribute(Profile))
 	{
-		// we have to normalize here because there are edges cases where skeletal meshes are loaded with non-normalized weights.
-		// even when/if we fix these, it's best to be 100% sure that things are normalized since weight data can be generated by tools we don't control
-		NormalizedWeights.Reset();
-		int32 NumInfluences = 0;
-		FBoneWeights BoneWeights;
-		VertexSkinWeights->GetValue(VertexIndex, BoneWeights);
-		for (const FBoneWeight BoneWeight: BoneWeights)
+		for (int32 VertexIndex : Mesh.VertexIndicesItr())
 		{
-			if (!ensure(NumInfluences < MAX_TOTAL_INFLUENCES))
+			// we have to normalize here because there are edges cases where skeletal meshes are loaded with non-normalized weights.
+			// even when/if we fix these, it's best to be 100% sure that things are normalized since weight data can be generated by tools we don't control
+			NormalizedWeights.Reset();
+			int32 NumInfluences = 0;
+			FBoneWeights BoneWeights;
+			VertexSkinWeights->GetValue(VertexIndex, BoneWeights);
+			for (const FBoneWeight BoneWeight : BoneWeights)
 			{
-				// cut-off weights if they exceed max influences (this shouldn't happen)
-				break;
+				if (!ensure(NumInfluences < MAX_TOTAL_INFLUENCES))
+				{
+					// cut-off weights if they exceed max influences (this shouldn't happen)
+					break;
+				}
+				NormalizedWeights.Add(BoneWeight.GetBoneIndex(), BoneWeight.GetWeight());
+				++NumInfluences;
 			}
-			NormalizedWeights.Add(BoneWeight.GetBoneIndex(), BoneWeight.GetWeight());
-			++NumInfluences;
-		}
-		USkinWeightsPaintTool::NormalizeWeightMap(NormalizedWeights);
-		
-		// if there are no bone weights, default to root bone 
-		if (NumInfluences == 0)
-		{
-			const FVector& RefPoseVertexPosition = Deformer.RefPoseVertexPositions[VertexIndex];
-			const FTransform& InvRefPoseTransform = Deformer.InvCSRefPoseTransforms[RootBoneIndex];
-			const FVector& BoneLocalPositionInRefPose = InvRefPoseTransform.TransformPosition(RefPoseVertexPosition);
-			CurrentWeights[VertexIndex].Emplace(RootBoneIndex, BoneLocalPositionInRefPose, FullWeight);
-			continue;
-		}
+			USkinWeightsPaintTool::NormalizeWeightMap(NormalizedWeights);
 
-		// load into the main weights data structure
-		for (TTuple<BoneIndex, float>& SingleBoneWeight : NormalizedWeights)
-		{
-			BoneIndex BoneIndex = SingleBoneWeight.Key;
-			const float Weight = SingleBoneWeight.Value;
-			if (!ensure(Deformer.InvCSRefPoseTransforms.IsValidIndex(BoneIndex)))
+			// if there are no bone weights, default to root bone 
+			if (NumInfluences == 0)
+			{
+				const FVector& RefPoseVertexPosition = Deformer.RefPoseVertexPositions[VertexIndex];
+				const FTransform& InvRefPoseTransform = Deformer.InvCSRefPoseTransforms[RootBoneIndex];
+				const FVector& BoneLocalPositionInRefPose = InvRefPoseTransform.TransformPosition(RefPoseVertexPosition);
+				CurrentWeights[VertexIndex].Emplace(RootBoneIndex, BoneLocalPositionInRefPose, FullWeight);
+				continue;
+			}
+
+			// load into the main weights data structure
+			for (TTuple<BoneIndex, float>& SingleBoneWeight : NormalizedWeights)
 			{
-				UE_LOG(LogMeshModelingToolsEditor, Warning, TEXT("InitializeSkinWeights: Invalid bone index provided (%d); falling back to 0 (root) as bone index."), BoneIndex);
-				BoneIndex = 0;
+				BoneIndex BoneIndex = SingleBoneWeight.Key;
+				const float Weight = SingleBoneWeight.Value;
+				if (!ensure(Deformer.InvCSRefPoseTransforms.IsValidIndex(BoneIndex)))
+				{
+					UE_LOG(LogMeshModelingToolsEditor, Warning, TEXT("InitializeSkinWeights: Invalid bone index provided (%d); falling back to 0 (root) as bone index."), BoneIndex);
+					BoneIndex = 0;
+				}
+				const FVector& RefPoseVertexPosition = Deformer.RefPoseVertexPositions[VertexIndex];
+				const FTransform& InvRefPoseTransform = Deformer.InvCSRefPoseTransforms[BoneIndex];
+				const FVector& BoneLocalPositionInRefPose = InvRefPoseTransform.TransformPosition(RefPoseVertexPosition);
+				CurrentWeights[VertexIndex].Emplace(BoneIndex, BoneLocalPositionInRefPose, Weight);
 			}
-			const FVector& RefPoseVertexPosition = Deformer.RefPoseVertexPositions[VertexIndex];
-			const FTransform& InvRefPoseTransform = Deformer.InvCSRefPoseTransforms[BoneIndex];
-			const FVector& BoneLocalPositionInRefPose = InvRefPoseTransform.TransformPosition(RefPoseVertexPosition);
-			CurrentWeights[VertexIndex].Emplace(BoneIndex, BoneLocalPositionInRefPose, Weight);
 		}
 	}
+	else
+	{
+		UE_LOG(LogMeshModelingToolsEditor, Warning, TEXT("FSkinToolWeights::InitializeSkinWeights : failed to find skin weights on the target mesh, all initial weights will be empty"))
+	}
 	
 	// maintain duplicate weight map
 	PreChangeWeights = CurrentWeights;
