Ref skeleton poser : added a getter to get the reference skeleton

diff --git a/Engine/Plugins/Runtime/MeshModelingToolset/Source/MeshModelingToolsEditorOnly/Private/SkeletalMesh/RefSkeletonPoser.cpp b/Engine/Plugins/Runtime/MeshModelingToolset/Source/MeshModelingToolsEditorOnly/Private/SkeletalMesh/RefSkeletonPoser.cpp
--- a/Engine/Plugins/Runtime/MeshModelingToolset/Source/MeshModelingToolsEditorOnly/Private/SkeletalMesh/RefSkeletonPoser.cpp
+++ b/Engine/Plugins/Runtime/MeshModelingToolset/Source/MeshModelingToolsEditorOnly/Private/SkeletalMesh/RefSkeletonPoser.cpp
@@ -57,6 +57,11 @@ void URefSkeletonPoser::SetRefSkeleton(const FReferenceSkeleton& InRefSkeleton)
 	AdditiveTransforms.Reset();
 }
 
+const FReferenceSkeleton& URefSkeletonPoser::GetRefSkeleton() const
+{
+	return RefSkeleton;
+}
+
 void URefSkeletonPoser::ModifyBoneAdditiveTransform(int32 BoneIndex, TFunctionRef<void(FTransform&)> InModifyFunc)
 {
 	FTransform& TransformRef = AdditiveTransforms.FindOrAdd(BoneIndex);
diff --git a/Engine/Plugins/Runtime/MeshModelingToolset/Source/MeshModelingToolsEditorOnly/Public/SkeletalMesh/RefSkeletonPoser.h b/Engine/Plugins/Runtime/MeshModelingToolset/Source/MeshModelingToolsEditorOnly/Public/SkeletalMesh/RefSkeletonPoser.h
--- a/Engine/Plugins/Runtime/MeshModelingToolset/Source/MeshModelingToolsEditorOnly/Public/SkeletalMesh/RefSkeletonPoser.h
+++ b/Engine/Plugins/Runtime/MeshModelingToolset/Source/MeshModelingToolsEditorOnly/Public/SkeletalMesh/RefSkeletonPoser.h
@@ -18,6 +18,7 @@ class URefSkeletonPoser: public UObject
 
 public:
 	UE_API void SetRefSkeleton(const FReferenceSkeleton& InRefSkeleton);
+	UE_API const FReferenceSkeleton& GetRefSkeleton() const;
 
 	UE_API void ModifyBoneAdditiveTransform(int32 BoneIndex, TFunctionRef<void(FTransform&)> InModifyFunc);
 	UE_API void ClearBoneAdditiveTransform(int32 BoneIndex);
