[Skin weights paint tool] Fixed transfer weights not working + fixed weighted bone tracker not working when changing between partial mesh and full mesh

diff --git a/Engine/Plugins/Runtime/MeshModelingToolset/Source/MeshModelingToolsEditorOnly/Private/SkeletalMesh/SkinWeightsPaintTool.cpp b/Engine/Plugins/Runtime/MeshModelingToolset/Source/MeshModelingToolsEditorOnly/Private/SkeletalMesh/SkinWeightsPaintTool.cpp
--- a/Engine/Plugins/Runtime/MeshModelingToolset/Source/MeshModelingToolsEditorOnly/Private/SkeletalMesh/SkinWeightsPaintTool.cpp
+++ b/Engine/Plugins/Runtime/MeshModelingToolset/Source/MeshModelingToolsEditorOnly/Private/SkeletalMesh/SkinWeightsPaintTool.cpp
@@ -610,7 +610,10 @@ void UWeightToolTransferManager::TransferWeightsFromOtherMeshOrSubset()
 			
 			// create a sub-mesh from the selected triangles to filter the transfer so that it only copies from the selected components
 			const UE::Geometry::FDynamicSubmesh3 PartialSubMesh(&EditedMesh, ToIsolate);
-			TmpSourceCopy = PartialSubMesh.GetSubmesh();			
+			TmpSourceCopy = PartialSubMesh.GetSubmesh();
+			// By default, submesh only have matching attributes enabled, so we need to explicitly copy the ones we need
+			FDynamicMeshAttributeSet* SourceAttributes = TmpSourceCopy.Attributes();
+			SourceAttributes->CopyBoneAttributes(*EditedMesh.Attributes());
 			return &TmpSourceCopy;
 		}
 
@@ -623,6 +626,9 @@ void UWeightToolTransferManager::TransferWeightsFromOtherMeshOrSubset()
 		{
 			const UE::Geometry::FDynamicSubmesh3 PartialSubMesh(&SourceMeshOrig, ToIsolate);
 			TmpSourceCopy = PartialSubMesh.GetSubmesh();
+			// By default, submesh only have matching attributes enabled, so we need to explicitly copy the ones we need
+			FDynamicMeshAttributeSet* SourceAttributes = TmpSourceCopy.Attributes();
+			SourceAttributes->CopyBoneAttributes(*EditedMesh.Attributes());
 		}
 		else
 		{
@@ -1412,7 +1418,7 @@ void FSkinToolWeights::CreateWeightEditForVertex(
 	}
 }
 
-void FSkinToolWeights::ApplyCurrentWeightsToMesh(FDynamicMesh3& Mesh)
+void FSkinToolWeights::ApplyCurrentWeightsToMesh(FDynamicMesh3& Mesh, TFunction<int32(int32)> WeightsIndexToMeshIndex)
 {
 	using namespace UE::Geometry;
 	FDynamicMeshVertexSkinWeightsAttribute* VertexWeightAttrs = Mesh.Attributes()->GetSkinWeightsAttribute(Profile);
@@ -1423,17 +1429,20 @@ void FSkinToolWeights::ApplyCurrentWeightsToMesh(FDynamicMesh3& Mesh)
 	TArray<UE::AnimationCore::FBoneWeight> BoneWeightsToApply;
 	BoneWeightsToApply.Reserve(UE::AnimationCore::MaxInlineBoneWeightCount);
 
-	if (!ensure(CurrentWeights.Num() == Mesh.MaxVertexID()))
+	if (!WeightsIndexToMeshIndex)
 	{
-		// weights are out of sync with mesh you're trying to apply them to
-		return;
+		if (!ensure(CurrentWeights.Num() == Mesh.MaxVertexID()))
+		{
+			// weights are out of sync with mesh you're trying to apply them to
+			return;
+		}
 	}
 
-	for (int32 VertexIndex : Mesh.VertexIndicesItr())
+	for (int32 WeightsIndex = 0; WeightsIndex < CurrentWeights.Num(); WeightsIndex++)
 	{
 		BoneWeightsToApply.Reset();
 
-		const VertexWeights& VertexWeights = CurrentWeights[VertexIndex];
+		const VertexWeights& VertexWeights = CurrentWeights[WeightsIndex];
 		for (const FVertexBoneWeight& SingleBoneWeight : VertexWeights)
 		{
 			if (!ensure(SingleBoneWeight.BoneID != INDEX_NONE))
@@ -1443,6 +1452,11 @@ void FSkinToolWeights::ApplyCurrentWeightsToMesh(FDynamicMesh3& Mesh)
 			BoneWeightsToApply.Add(UE::AnimationCore::FBoneWeight(SingleBoneWeight.BoneID, SingleBoneWeight.Weight));
 		}
 
+		int32 VertexIndex = WeightsIndex;
+		if (WeightsIndexToMeshIndex)
+		{
+			VertexIndex = WeightsIndexToMeshIndex(WeightsIndex);
+		}
 		VertexWeightAttrs->SetValue(VertexIndex, UE::AnimationCore::FBoneWeights::Create(BoneWeightsToApply, Settings));
 	}
 }
@@ -4180,48 +4194,22 @@ void UWeightToolSelectionIsolator::RestoreFullMesh()
 		return;
 	}
 
-	FDynamicMesh3& PartialMesh = PartialSubMesh.GetSubmesh();
 	FDynamicMesh3& FullMesh = WeightTool->GetMesh();
-	FSkinToolWeights& Weights = WeightTool->GetWeights();
 
-	// apply partial mesh weights to partial mesh description
-	Weights.ApplyCurrentWeightsToMesh(PartialMesh);
+	// apply partial mesh weights to full mesh
+	{
+		auto PartialIndexToFullIndex = [this](int32 PartialIndex)
+			{
+				return PartialSubMesh.MapVertexToBaseMesh(PartialIndex);
+			};
+	
+		WeightTool->GetWeights().ApplyCurrentWeightsToMesh(FullMesh, PartialIndexToFullIndex);
+	}
 
 	// reinitialize with full mesh
 	// this resizes Weights to the full mesh size
 	WeightTool->UpdateCurrentlyEditedMesh(FullMesh);
 
-	using namespace UE::Geometry;
-	using namespace UE::AnimationCore;
-	const FName ActiveProfile = WeightTool->GetWeightToolProperties()->GetActiveSkinWeightProfile();
-	FDynamicMeshVertexSkinWeightsAttribute* PartialWeights = PartialMesh.Attributes()->GetSkinWeightsAttribute(ActiveProfile);
-	
-	for (int32 VertexIndexPartial : PartialSubMesh.GetSubmesh().VertexIndicesItr() )
-	{
-		int32 VertexIndexFull = PartialSubMesh.MapVertexToBaseMesh(VertexIndexPartial);
-		
-		// clear all the weights on this vertex
-		Weights.CurrentWeights[VertexIndexFull].Reset();
-		// replace with weights from partial mesh
-		FBoneWeights VertexWeightsPartial;	
-		PartialWeights->GetValue(VertexIndexPartial, VertexWeightsPartial);
-		for (const FBoneWeight& SingleBoneWeight : VertexWeightsPartial)
-		{
-			const FBoneIndexType BoneID = SingleBoneWeight.GetBoneIndex();
-			const float Weight = SingleBoneWeight.GetWeight();
-			const FVector& RefPoseVertexPosition = Weights.Deformer.RefPoseVertexPositions[VertexIndexFull];
-			const FTransform& InvRefPoseBoneTransform = Weights.Deformer.InvCSRefPoseTransforms[BoneID];
-			const FVector VertexInBoneSpace = InvRefPoseBoneTransform.TransformPosition(RefPoseVertexPosition);
-			// store the new influence on this vertex
-			Weights.CurrentWeights[VertexIndexFull].Emplace(BoneID, VertexInBoneSpace, Weight);
-		}
-	}
-
-	// sync both weight buffers
-	Weights.PreChangeWeights = Weights.CurrentWeights;
-	// apply full mesh weights to full mesh description
-	Weights.ApplyCurrentWeightsToMesh(WeightTool->GetMesh());
-
 	// restore selection (a convenience for the user that allows for easily adjusting the isolation by going back/forth as needed)
 	if (UPolygonSelectionMechanic* SelectionMechanic = WeightTool->GetMainMeshSelector()->GetSelectionMechanic())
 	{
diff --git a/Engine/Plugins/Runtime/MeshModelingToolset/Source/MeshModelingToolsEditorOnly/Public/SkeletalMesh/SkinWeightsPaintTool.h b/Engine/Plugins/Runtime/MeshModelingToolset/Source/MeshModelingToolsEditorOnly/Public/SkeletalMesh/SkinWeightsPaintTool.h
--- a/Engine/Plugins/Runtime/MeshModelingToolset/Source/MeshModelingToolsEditorOnly/Public/SkeletalMesh/SkinWeightsPaintTool.h
+++ b/Engine/Plugins/Runtime/MeshModelingToolset/Source/MeshModelingToolsEditorOnly/Public/SkeletalMesh/SkinWeightsPaintTool.h
@@ -218,7 +218,7 @@ namespace SkinPaintTool
 			float NewWeightValue,
 			FMultiBoneWeightEdits& WeightEdits);
 
-		void ApplyCurrentWeightsToMesh(FDynamicMesh3& Mesh);
+		void ApplyCurrentWeightsToMesh(FDynamicMesh3& Mesh, TFunction<int32(int32)> WeightsIndexToMeshIndex = {});
 		
 		static float GetWeightOfBoneOnVertex(
 			const int32 BoneIndex,
