[SkinWeightsPaintTool] Fixed removing weights from rigid weighted verts turning them into unweighted verts

diff --git a/Engine/Plugins/Runtime/MeshModelingToolset/Source/MeshModelingToolsEditorOnly/Private/SkeletalMesh/SkinWeightsPaintTool.cpp b/Engine/Plugins/Runtime/MeshModelingToolset/Source/MeshModelingToolsEditorOnly/Private/SkeletalMesh/SkinWeightsPaintTool.cpp
--- a/Engine/Plugins/Runtime/MeshModelingToolset/Source/MeshModelingToolsEditorOnly/Private/SkeletalMesh/SkinWeightsPaintTool.cpp
+++ b/Engine/Plugins/Runtime/MeshModelingToolset/Source/MeshModelingToolsEditorOnly/Private/SkeletalMesh/SkinWeightsPaintTool.cpp
@@ -1314,7 +1314,7 @@ void FSkinToolWeights::CreateWeightEditForVertex(
 	}
 
 	// assigning full weight to this vertex?
-	if (FMath::IsNearlyEqual(NewWeightValue, 1.f))
+	if (FMath::IsNearlyEqual(NewWeightValue, 1.f, MinimumWeightThreshold))
 	{
 		// in this case normalization is trivial, just assign the full weight directly and zero all others
 		constexpr float FullWeight = 1.0f;
@@ -1404,9 +1404,16 @@ void FSkinToolWeights::CreateWeightEditForVertex(
 		for (int32 i=0; i<ValuesToNormalize.Num(); ++i)
 		{
 			float NormalizedValue = 0.f;
-			if (AvailableTotal > MinimumWeightThreshold && Total > KINDA_SMALL_NUMBER)
+			if (ensure(AvailableTotal > MinimumWeightThreshold))
 			{
-				NormalizedValue = (ValuesToNormalize[i] / Total) * AvailableTotal;	
+				if (Total > KINDA_SMALL_NUMBER)
+				{
+					NormalizedValue = (ValuesToNormalize[i] / Total) * AvailableTotal;	
+				}
+				else
+				{
+					NormalizedValue = AvailableTotal / ValuesToNormalize.Num();
+				}
 			}
 			const int32 BoneIndex = RecordedBonesOnVertex[i];
 			const float NewWeight = NormalizedValue;
