Merge Actor - Approximate: Fix issue where meshes with a forced LOD would not be rendered during the scene capture phase, resulting in them being black in the final mesh.
* Static mesh assets with MinLOD > 0
* Static mesh components with ForceLOD > 0

diff --git a/Engine/Plugins/Runtime/MeshModelingToolset/Source/ModelingComponents/Private/Scene/WorldRenderCapture.cpp b/Engine/Plugins/Runtime/MeshModelingToolset/Source/ModelingComponents/Private/Scene/WorldRenderCapture.cpp
--- a/Engine/Plugins/Runtime/MeshModelingToolset/Source/ModelingComponents/Private/Scene/WorldRenderCapture.cpp
+++ b/Engine/Plugins/Runtime/MeshModelingToolset/Source/ModelingComponents/Private/Scene/WorldRenderCapture.cpp
@@ -497,8 +497,7 @@ void SetCommonShowFlags(FEngineShowFlags& ShowFlags, bool bAntiAliasing)
 	ShowFlags.DisableAdvancedFeatures();
 
 	// Disable some more features. These are sorted alphabetically.
-	ShowFlags.LOD = 0;
-	ShowFlags.MotionBlur = 0;
+	ShowFlags.SetMotionBlur(false);
 	ShowFlags.SetAtmosphere(false);
 	ShowFlags.SetBloom(false);
 	ShowFlags.SetCapsuleShadows(false);
