Merge Actor - Approximate: Fix issue where actors could be distance culled during the scene capture phase, resulting in them being black in the final mesh.

diff --git a/Engine/Plugins/Runtime/MeshModelingToolset/Source/ModelingComponents/Private/Scene/WorldRenderCapture.cpp b/Engine/Plugins/Runtime/MeshModelingToolset/Source/ModelingComponents/Private/Scene/WorldRenderCapture.cpp
--- a/Engine/Plugins/Runtime/MeshModelingToolset/Source/ModelingComponents/Private/Scene/WorldRenderCapture.cpp
+++ b/Engine/Plugins/Runtime/MeshModelingToolset/Source/ModelingComponents/Private/Scene/WorldRenderCapture.cpp
@@ -168,30 +168,9 @@ void FlushRHIThreadToUpdateTextureRenderTargetReference()
 
 FWorldRenderCapture::~FWorldRenderCapture()
 {
-	Shutdown();
+	RestoreComponentsMaxDrawDistance();
 }
 
-void FWorldRenderCapture::Shutdown()
-{
-	if (LinearRenderTexture != nullptr)
-	{
-		LinearRenderTexture->RemoveFromRoot();
-		LinearRenderTexture = nullptr;
-	}
-	if (GammaRenderTexture != nullptr)
-	{
-		GammaRenderTexture->RemoveFromRoot();
-		GammaRenderTexture = nullptr;
-	}
-	if (DepthRenderTexture != nullptr)
-	{
-		DepthRenderTexture->RemoveFromRoot();
-		DepthRenderTexture = nullptr;
-	}
-}
-
-
-
 void FWorldRenderCapture::SetWorld(UWorld* WorldIn)
 {
 	this->World = WorldIn;
@@ -207,12 +186,24 @@ void FWorldRenderCapture::SetVisibleComponents(const TArray<UActorComponent*>& C
 	SetVisibleActorsAndComponents({}, Components);
 }
 
+void FWorldRenderCapture::RestoreComponentsMaxDrawDistance()
+{
+	for (UPrimitiveComponent* PrimitiveComponent : ComponentsRestoreMaxDrawDistance)
+	{
+		PrimitiveComponent->SetCachedMaxDrawDistance(PrimitiveComponent->LDMaxDrawDistance);
+	}
+
+	ComponentsRestoreMaxDrawDistance.Reset();
+}
+
 void FWorldRenderCapture::SetVisibleActorsAndComponents(const TArray<AActor*>& Actors, const TArray<UActorComponent*>& Components)
 {
 	TRACE_CPUPROFILER_EVENT_SCOPE(FWorldRenderCapture::SetVisibleActorsAndComponents);
 
 	VisiblePrimitives.Reset();
 
+	RestoreComponentsMaxDrawDistance();
+
 	FBoxSphereBounds::Builder BoundsBuilder;
 
 	// Find all components that need to be included in rendering.
@@ -226,6 +217,12 @@ void FWorldRenderCapture::SetVisibleActorsAndComponents(const TArray<AActor*>& A
 			{				
 				VisiblePrimitives.Add(PrimitiveComponent->GetPrimitiveSceneId());
 
+				if (!PrimitiveComponent->bNeverDistanceCull)
+				{
+					ComponentsRestoreMaxDrawDistance.Add(PrimitiveComponent);
+					PrimitiveComponent->SetCachedMaxDrawDistance(0.0f);
+				}
+				
 				// Append bounds of visible components only
 				BoundsBuilder += PrimitiveComponent->Bounds;
 			}
@@ -295,19 +292,11 @@ UTextureRenderTarget2D* FWorldRenderCapture::GetRenderTexture(bool bLinear)
 {
 	if (RenderTextureDimensions != this->Dimensions)
 	{
-		if (LinearRenderTexture != nullptr)
-		{
-			LinearRenderTexture->RemoveFromRoot();
-			LinearRenderTexture = nullptr;
-		}
-		if (GammaRenderTexture != nullptr)
-		{
-			GammaRenderTexture->RemoveFromRoot();
-			GammaRenderTexture = nullptr;
-		}
+		LinearRenderTexture = nullptr;
+		GammaRenderTexture = nullptr;
 	}
 
-	UTextureRenderTarget2D** WhichTexture = (bLinear) ? &LinearRenderTexture : &GammaRenderTexture;
+	TObjectPtr<UTextureRenderTarget2D>* WhichTexture = (bLinear) ? &LinearRenderTexture : &GammaRenderTexture;
 
 	if ( *WhichTexture != nullptr )
 	{
@@ -325,7 +314,6 @@ UTextureRenderTarget2D* FWorldRenderCapture::GetRenderTexture(bool bLinear)
 	*WhichTexture = NewObject<UTextureRenderTarget2D>();
 	if (ensure(*WhichTexture))
 	{
-		(*WhichTexture)->AddToRoot();		// keep alive for GC
 		(*WhichTexture)->ClearColor = FLinearColor::Transparent;
 		(*WhichTexture)->TargetGamma = (bLinear) ? 1.0f : 2.2f;
 		(*WhichTexture)->InitCustomFormat(Width, Height, PF_FloatRGBA, false);
@@ -342,11 +330,7 @@ UTextureRenderTarget2D* FWorldRenderCapture::GetDepthRenderTexture()
 {
 	if (RenderTextureDimensions != this->Dimensions)
 	{
-		if (DepthRenderTexture != nullptr)
-		{
-			DepthRenderTexture->RemoveFromRoot();
-			DepthRenderTexture = nullptr;
-		}
+		DepthRenderTexture = nullptr;
 	}
 
 	if (DepthRenderTexture != nullptr)
@@ -361,7 +345,6 @@ UTextureRenderTarget2D* FWorldRenderCapture::GetDepthRenderTexture()
 	DepthRenderTexture = NewObject<UTextureRenderTarget2D>();
 	if (ensure(DepthRenderTexture))
 	{
-		DepthRenderTexture->AddToRoot();		// keep alive for GC
 		DepthRenderTexture->RenderTargetFormat = ETextureRenderTargetFormat::RTF_RGBA32f;
 		DepthRenderTexture->ClearColor = FLinearColor::Black;
 		DepthRenderTexture->InitAutoFormat(Width, Height);
@@ -497,7 +480,6 @@ void SetCommonShowFlags(FEngineShowFlags& ShowFlags, bool bAntiAliasing)
 	ShowFlags.DisableAdvancedFeatures();
 
 	// Disable some more features. These are sorted alphabetically.
-	ShowFlags.SetMotionBlur(false);
 	ShowFlags.SetAtmosphere(false);
 	ShowFlags.SetBloom(false);
 	ShowFlags.SetCapsuleShadows(false);
@@ -946,3 +928,15 @@ FViewMatrices UE::Geometry::GetRenderCaptureViewMatrices(const FFrame3d& ViewFra
 	return FViewMatrices(ViewInitOptions);
 }
 
+void FWorldRenderCapture::AddReferencedObjects(FReferenceCollector& Collector)
+{
+	Collector.AddReferencedObjects(ComponentsRestoreMaxDrawDistance);
+	Collector.AddReferencedObject(LinearRenderTexture);
+	Collector.AddReferencedObject(GammaRenderTexture);
+	Collector.AddReferencedObject(DepthRenderTexture);
+}
+
+FString FWorldRenderCapture::GetReferencerName() const
+{
+	return TEXT("FWorldRenderCapture");
+}
diff --git a/Engine/Plugins/Runtime/MeshModelingToolset/Source/ModelingComponents/Public/Scene/WorldRenderCapture.h b/Engine/Plugins/Runtime/MeshModelingToolset/Source/ModelingComponents/Public/Scene/WorldRenderCapture.h
--- a/Engine/Plugins/Runtime/MeshModelingToolset/Source/ModelingComponents/Public/Scene/WorldRenderCapture.h
+++ b/Engine/Plugins/Runtime/MeshModelingToolset/Source/ModelingComponents/Public/Scene/WorldRenderCapture.h
@@ -2,6 +2,7 @@
 
 #pragma once
 
+#include "UObject/GCObject.h"
 #include "SceneTypes.h"		// FPrimitiveComponentId
 #include "Image/SpatialPhotoSet.h"
 
@@ -9,6 +10,7 @@
 
 class UTextureRenderTarget2D;
 class AActor;
+class UPrimitiveComponent;
 
 namespace UE
 {
@@ -213,7 +215,7 @@ FViewMatrices GetRenderCaptureViewMatrices(const FFrame3d& ViewFrame, double Hor
  * specific viewpoint. Various types of rendering are supported, as defined by ERenderCaptureType.
  * Currently rendering an entire World, ie without an explicit list of Actors or Components, is not supported.
  */
-class FWorldRenderCapture
+class FWorldRenderCapture : public FGCObject
 {
 public:
 
@@ -297,6 +299,11 @@ UE_API PRAGMA_ENABLE_DEPRECATION_WARNINGS
 	 */
 	UE_API void SetEnableWriteDebugImage(bool bEnable, int32 ImageCounter = -1, FString FolderName = FString());
 
+	// ~Begin FGCObject Interface
+	virtual void AddReferencedObjects(FReferenceCollector& Collector) override;
+	virtual FString GetReferencerName() const override;
+	// ~End FGCObject Interface
+
 protected:
 	UWorld* World = nullptr;
 
@@ -306,9 +313,9 @@ protected:
 	FImageDimensions Dimensions = FImageDimensions(128, 128);
 
 	// Temporary textures used as render targets. We explicitly prevent this from being GC'd internally
-	UTextureRenderTarget2D* LinearRenderTexture = nullptr;
-	UTextureRenderTarget2D* GammaRenderTexture = nullptr;
-	UTextureRenderTarget2D* DepthRenderTexture = nullptr;
+	TObjectPtr<UTextureRenderTarget2D> LinearRenderTexture = nullptr;
+	TObjectPtr<UTextureRenderTarget2D> GammaRenderTexture = nullptr;
+	TObjectPtr<UTextureRenderTarget2D> DepthRenderTexture = nullptr;
 
 	FImageDimensions RenderTextureDimensions;
 	UE_API UTextureRenderTarget2D* GetRenderTexture(bool bLinear);
@@ -364,6 +371,9 @@ protected:
 	int32 DebugImageCounter = -1;
 	FString DebugImageFolderName = TEXT("WorldRenderCapture");
 	UE_API void WriteDebugImage(const FImageAdapter& ResultImageOut, const FString& ImageTypeName);
+
+	TArray<TObjectPtr<UPrimitiveComponent>> ComponentsRestoreMaxDrawDistance;
+	void RestoreComponentsMaxDrawDistance();
 };
 
 
