Modeling Tools: fixes bug which created random gizmo orientation in local space in the PolyEdit tool

diff --git a/Engine/Plugins/Runtime/MeshModelingToolset/Source/MeshModelingTools/Private/EditMeshPolygonsTool.cpp b/Engine/Plugins/Runtime/MeshModelingToolset/Source/MeshModelingTools/Private/EditMeshPolygonsTool.cpp
--- a/Engine/Plugins/Runtime/MeshModelingToolset/Source/MeshModelingTools/Private/EditMeshPolygonsTool.cpp
+++ b/Engine/Plugins/Runtime/MeshModelingToolset/Source/MeshModelingTools/Private/EditMeshPolygonsTool.cpp
@@ -1378,6 +1378,7 @@ bool UEditMeshPolygonsTool::GetWorldSpaceFocusPoint(const FRay& WorldRay, FVecto
 void UEditMeshPolygonsTool::OnSelectionModifiedEvent()
 {
 	bSelectionStateDirty = true;
+	LastGeometryFrame = FFrame3d();
 }
 
 
