Modeling Mode: add slight delay to appearance of working material during displace tool computations

diff --git a/Engine/Plugins/Runtime/MeshModelingToolset/Source/MeshModelingTools/Private/DisplaceMeshTool.cpp b/Engine/Plugins/Runtime/MeshModelingToolset/Source/MeshModelingTools/Private/DisplaceMeshTool.cpp
--- a/Engine/Plugins/Runtime/MeshModelingToolset/Source/MeshModelingTools/Private/DisplaceMeshTool.cpp
+++ b/Engine/Plugins/Runtime/MeshModelingToolset/Source/MeshModelingTools/Private/DisplaceMeshTool.cpp
@@ -1475,14 +1475,13 @@ void UDisplaceMeshTool::StartComputation()
 		SubdivideTask = new FAsyncTaskExecuterWithAbort<TModelingOpTask<FDynamicMeshOperator>>(Subdivider->MakeNewOperator());
 		SubdivideTask->StartBackgroundTask();
 		bNeedsSubdivided = false;
-		DynamicMeshComponent->SetOverrideRenderMaterial(ToolSetupUtil::GetDefaultWorkingMaterial(GetToolManager()));
 	}
 	if (bNeedsDisplaced && DisplaceTask)
 	{
 		DisplaceTask->CancelAndDelete();
 		DisplaceTask = nullptr;
-		DynamicMeshComponent->SetOverrideRenderMaterial(ToolSetupUtil::GetDefaultWorkingMaterial(GetToolManager()));
 	}
+	ComputationStartTime = FPlatformTime::Seconds();
 	AdvanceComputation();
 }
 
@@ -1490,6 +1489,12 @@ void UDisplaceMeshTool::AdvanceComputation()
 {
 	using namespace DisplaceMeshToolLocals;
 
+	const float ElapsedTime = (float)(FPlatformTime::Seconds() - ComputationStartTime);
+	if (ElapsedTime > 0.75f && DisplaceTask && !DisplaceTask->IsDone())
+	{
+		DynamicMeshComponent->SetOverrideRenderMaterial(ToolSetupUtil::GetDefaultWorkingMaterial(GetToolManager()));
+	}
+
 	if (SubdivideTask && SubdivideTask->IsDone())
 	{
 		TUniquePtr<FDynamicMeshOperator> SubOp = SubdivideTask->GetTask().ExtractOperator();
diff --git a/Engine/Plugins/Runtime/MeshModelingToolset/Source/MeshModelingTools/Public/DisplaceMeshTool.h b/Engine/Plugins/Runtime/MeshModelingToolset/Source/MeshModelingTools/Public/DisplaceMeshTool.h
--- a/Engine/Plugins/Runtime/MeshModelingToolset/Source/MeshModelingTools/Public/DisplaceMeshTool.h
+++ b/Engine/Plugins/Runtime/MeshModelingToolset/Source/MeshModelingTools/Public/DisplaceMeshTool.h
@@ -358,6 +358,8 @@ private:
 	UE_API void AdvanceComputation();
 	UE_API void ValidateSubdivisions();
 
+	double ComputationStartTime = 0;
+
 	UE::Geometry::FAsyncTaskExecuterWithAbort<UE::Geometry::TModelingOpTask<UE::Geometry::FDynamicMeshOperator>>* SubdivideTask = nullptr;
 	bool bNeedsSubdivided = true;
 	UE::Geometry::FAsyncTaskExecuterWithAbort<UE::Geometry::TModelingOpTask<UE::Geometry::FDynamicMeshOperator>>* DisplaceTask = nullptr;
