ModelingTools: Fixed StaticMeshSelector only working on exactly StaticMeshComponents.

The restriction was to filter out InstancedStaticMeshComponents which our selector does not support. That is overly restrictive for its purpose and preventing minor variations of StaticMeshComponent from being eligible for MeshElementSelection. This may open up other issues for more involved StaticMeshComponent variations but we can address those with a more holistic solution as they come in.

diff --git a/Engine/Plugins/Runtime/MeshModelingToolset/Source/ModelingComponentsEditorOnly/Private/Selection/StaticMeshSelector.cpp b/Engine/Plugins/Runtime/MeshModelingToolset/Source/ModelingComponentsEditorOnly/Private/Selection/StaticMeshSelector.cpp
--- a/Engine/Plugins/Runtime/MeshModelingToolset/Source/ModelingComponentsEditorOnly/Private/Selection/StaticMeshSelector.cpp
+++ b/Engine/Plugins/Runtime/MeshModelingToolset/Source/ModelingComponentsEditorOnly/Private/Selection/StaticMeshSelector.cpp
@@ -12,6 +12,7 @@
 #include "Changes/MeshVertexChange.h"
 
 #include "Components/StaticMeshComponent.h"
+#include "Components/InstancedStaticMeshComponent.h"
 #include "Engine/StaticMesh.h"
 #include "MeshDescriptionToDynamicMesh.h"
 #include "DynamicMeshToMeshDescription.h"
@@ -308,9 +309,9 @@ bool FStaticMeshComponentSelectorFactory::CanBuildForTarget(FGeometryIdentifier
 		UStaticMeshComponent* Component = TargetIdentifier.GetAsComponentType<UStaticMeshComponent>();
 		if (Component != nullptr )
 		{
-			// Ensure that the Component is an SMComponent and not an ISMC or other subclass that might have other behaviors
-			// that we can't know if the Selector will properly support
-			if (ExactCast<UStaticMeshComponent>(Component) == nullptr)
+			// Ensure that the Component is not an InstancedStaticMeshComponent.
+			// The StaticMeshSelector does not properly support that component.
+			if (Cast<UInstancedStaticMeshComponent>(Component))
 			{
 				return false;
 			}
