ModelingTools: Fix potential crash in sculpt tool from changing the TriangleROIArray out from under an asynchronous octree update.

diff --git a/Engine/Plugins/Runtime/MeshModelingToolset/Source/MeshModelingTools/Private/MeshVertexSculptTool.cpp b/Engine/Plugins/Runtime/MeshModelingToolset/Source/MeshModelingTools/Private/MeshVertexSculptTool.cpp
--- a/Engine/Plugins/Runtime/MeshModelingToolset/Source/MeshModelingTools/Private/MeshVertexSculptTool.cpp
+++ b/Engine/Plugins/Runtime/MeshModelingToolset/Source/MeshModelingTools/Private/MeshVertexSculptTool.cpp
@@ -659,6 +659,8 @@ void UMeshVertexSculptTool::OnCancelStroke()
 
 	delete ActiveVertexChange;
 	ActiveVertexChange = nullptr;
+
+	LongTransactions.Close(GetToolManager());
 }
 
 // The first part of UpdateROI, which updates TriangleROIArray to be triangles
@@ -673,6 +675,9 @@ void UMeshVertexSculptTool::UpdateRangeQueryTriBuffer(const FFrame3d& LocalFrame
 	{
 		// If we're requiring connectivity, and we didn't hit a triangle to start with, then we shouldn't
 		//  move any triangles.
+		// Make sure TriangleROIArray is not in use, as we're about to clear it
+		WaitForPendingStampUpdateConst();
+		
 		VertexROI.Reset();
 		TriangleROIArray.Reset();
 		SymmetricVertexROI.Reset();
@@ -909,7 +914,8 @@ void UMeshVertexSculptTool::UpdateRangeQueryTriBuffer(const FFrame3d& LocalFrame
 				}
 			}
 		}
-
+		// Make sure TriangleROIArray is not in use, as we're about to change it
+		WaitForPendingStampUpdateConst();
 		TriangleROIBuilder.SwapValuesWith(TriangleROIArray);
 	}
 
@@ -949,6 +955,8 @@ void UMeshVertexSculptTool::UpdateRangeQueryTriBuffer(const FFrame3d& LocalFrame
 			}
 		}
 		VertexROIBuilder.SwapValuesWith(VertexROI);
+		// Make sure TriangleROIArray is not in use, as we're about to change it
+		WaitForPendingStampUpdateConst();
 		TriangleROIBuilder.SwapValuesWith(TriangleROIArray);
 	}
 #endif
@@ -1520,17 +1528,16 @@ void UMeshVertexSculptTool::OnTick(float DeltaTime)
 		// begin octree rebuild calculation
 		if (!MeshVertexSculptToolLocals::DisableOctreeUpdates->GetBool())
 		{
+			// We've run into cases where we got to here before the previous octree future finished
+			//  executing, so make sure that the previous one is done.
+			WaitForPendingStampUpdateConst();
+
 			StampUpdateOctreeFuture = Async(VertexSculptToolAsyncExecTarget, [this]()
-				{
-					TRACE_CPUPROFILER_EVENT_SCOPE(VtxSculptTool_Tick_OctreeReinsert);
-					Octree.ReinsertTrianglesParallel(TriangleROIArray, OctreeUpdateTempBuffer, OctreeUpdateTempFlagBuffer);
-				});
-			bStampUpdatePending = true;
-			StampUpdateOctreeFuture.Then([this](TFuture<void> Future)
-				{
-					bStampUpdatePending = false;
-					bOctreeUpdated = true;
-				});
+			{
+				TRACE_CPUPROFILER_EVENT_SCOPE(VtxSculptTool_Tick_OctreeReinsert);
+				Octree.ReinsertTrianglesParallel(TriangleROIArray, OctreeUpdateTempBuffer, OctreeUpdateTempFlagBuffer);
+				bOctreeUpdated = true;
+			});
 		}
 
 		// Prepare list of triangles to process after all stamps are applied
@@ -1650,7 +1657,7 @@ void UMeshVertexSculptTool::OnTick(float DeltaTime)
 
 void UMeshVertexSculptTool::WaitForPendingStampUpdateConst() const
 {
-	if (bStampUpdatePending)
+	if (StampUpdateOctreeFuture.IsValid() && !StampUpdateOctreeFuture.IsReady())
 	{
 		StampUpdateOctreeFuture.Wait();
 	}
diff --git a/Engine/Plugins/Runtime/MeshModelingToolset/Source/MeshModelingTools/Public/MeshVertexSculptTool.h b/Engine/Plugins/Runtime/MeshModelingToolset/Source/MeshModelingTools/Public/MeshVertexSculptTool.h
--- a/Engine/Plugins/Runtime/MeshModelingToolset/Source/MeshModelingTools/Public/MeshVertexSculptTool.h
+++ b/Engine/Plugins/Runtime/MeshModelingToolset/Source/MeshModelingTools/Public/MeshVertexSculptTool.h
@@ -378,6 +378,7 @@ protected:
 	TArray<uint32> OctreeUpdateTempBuffer;
 	TArray<bool> OctreeUpdateTempFlagBuffer;
 	TFuture<void> StampUpdateOctreeFuture;
+	UE_DEPRECATED(5.7, "Check StampUpdateOctreeFuture directly instead")
 	bool bStampUpdatePending = false;
 
 	UE_DEPRECATED(5.7, "Vertex sculpt now uses a const version of WaitForPendingStampUpdate")
