ModelingTools: Fix spacing mode in sculpt tools, which behaved poorly when moving across ridges.

diff --git a/Engine/Plugins/Runtime/MeshModelingToolset/Source/MeshModelingTools/Private/Sculpting/MeshSculptToolBase.cpp b/Engine/Plugins/Runtime/MeshModelingToolset/Source/MeshModelingTools/Private/Sculpting/MeshSculptToolBase.cpp
--- a/Engine/Plugins/Runtime/MeshModelingToolset/Source/MeshModelingTools/Private/Sculpting/MeshSculptToolBase.cpp
+++ b/Engine/Plugins/Runtime/MeshModelingToolset/Source/MeshModelingTools/Private/Sculpting/MeshSculptToolBase.cpp
@@ -39,7 +39,8 @@ namespace MeshSculptBaseLocals
 	const FString& StrokePointSetID(TEXT("StrokePointSet"));
 	const FString& StrokeLineSetID(TEXT("StrokeLineSet"));
 
-	static FAutoConsoleVariable EnableMouseTrail(TEXT("modeling.Sculpting.EnableMouseTrail"), false, TEXT("Enable visualizing the path the mouse takes while sculpting."));
+	// TODO: This may be useful to have again
+	// static FAutoConsoleVariable EnableMouseTrail(TEXT("modeling.Sculpting.EnableMouseTrail"), false, TEXT("Enable visualizing the path the mouse takes while sculpting."));
 }
 
 void FBrushToolRadius::InitializeWorldSizeRange(TInterval<float> Range, bool bValidateWorldRadius)
@@ -333,131 +334,6 @@ void UMeshSculptToolBase::OnTick(float DeltaTime)
 
 	UpdateHoverStamp(GetBrushFrameWorld());
 
-	if (MeshSculptBaseLocals::EnableMouseTrail->GetBool() && StrokePath.Num() > 1)
-	{
-		TArray<FVector3d> UnProjectedStampPoints;
-
-		double Spacing = FMath::Clamp(BrushProperties->Spacing * BrushProperties->BrushSize.GetWorldRadius() * 2 , 0.1, 100000);
-
-		TArray<UE::Geometry::FPolyline3d> SpacingSections;
-		int32 SectionIndex = 0;
-		SpacingSections.Add(UE::Geometry::FPolyline3d());
-		SpacingSections[SectionIndex].AppendVertex(StrokePath[0].Position);
-		UnProjectedStampPoints.Add(StrokePath[0].Position);
-		for (int32 Index = 0; Index < StrokePath.Num() - 1; Index++)
-		{
-			FVector3d Span = StrokePath[Index + 1].Position - StrokePath[Index].Position;
-			FVector3d SpanDir = Span.GetSafeNormal();
-
-			double RemainingSpace = FMath::Max(0, Spacing - SpacingSections[SectionIndex].Length());
-			double OverflowSpace = Span.Length() - RemainingSpace;
-			if (OverflowSpace < 0)
-			{
-				SpacingSections[SectionIndex].AppendVertex(StrokePath[Index + 1].Position);
-			}
-			else
-			{
-				SpacingSections[SectionIndex].AppendVertex(StrokePath[Index].Position + (SpanDir * RemainingSpace));
-				
-				while (OverflowSpace > Spacing)
-				{
-					FVector3d TerminalPosition = SpacingSections[SectionIndex].End();
-					SpacingSections.Add(UE::Geometry::FPolyline3d());
-					SectionIndex++;
-					UnProjectedStampPoints.Add(TerminalPosition);
-					SpacingSections[SectionIndex].AppendVertex(TerminalPosition);
-					SpacingSections[SectionIndex].AppendVertex(TerminalPosition + (SpanDir * Spacing));
-					OverflowSpace -= Spacing;
-				}
-				
-				FVector3d TerminalPosition = SpacingSections[SectionIndex].End();
-				SpacingSections.Add(UE::Geometry::FPolyline3d());
-				SectionIndex++;
-				UnProjectedStampPoints.Add(TerminalPosition);
-				SpacingSections[SectionIndex].AppendVertex(TerminalPosition);
-				SpacingSections[SectionIndex].AppendVertex(TerminalPosition + (SpanDir * OverflowSpace));
-			}
-		}
-
-		if (!StrokeGeometry)
-		{
-			// Set up all the components we need to visualize things.
-			StrokeGeometry = NewObject<UPreviewGeometry>();
-			StrokeGeometry->CreateInWorld(TargetWorld, FTransform());
-
-			// These visualize the current spline edges that would be extracted
-			StrokeGeometry->AddPointSet(MeshSculptBaseLocals::StrokePointSetID);
-			StrokeGeometry->AddLineSet(MeshSculptBaseLocals::StrokeLineSetID);
-		}
-
-		ULineSetComponent* LineSet = StrokeGeometry->FindLineSet(MeshSculptBaseLocals::StrokeLineSetID);
-		if (LineSet)
-		{
-			LineSet->Clear();
-			for (int32 StrokeSectionIndex = 0; StrokeSectionIndex < SpacingSections.Num(); ++StrokeSectionIndex)
-			{
-				FColor SectionColor = (StrokeSectionIndex % 2) ? FColor::Red : FColor::Cyan;
-				for (int32 SegmentIndex = 0; SegmentIndex < SpacingSections[StrokeSectionIndex].SegmentCount(); ++SegmentIndex)
-				{
-
-					LineSet->AddLine(SpacingSections[StrokeSectionIndex].GetSegment(SegmentIndex).StartPoint(),
-									 SpacingSections[StrokeSectionIndex].GetSegment(SegmentIndex).EndPoint(),
-									 SectionColor, 3.0);
-				}
-			}
-		}
-
-		UPointSetComponent* PointSet = StrokeGeometry->FindPointSet(MeshSculptBaseLocals::StrokePointSetID);
-		if (PointSet)
-		{
-			PointSet->Clear();
-			for (int32 UnProjectedStampPointsIndex = 0; UnProjectedStampPointsIndex < UnProjectedStampPoints.Num(); ++UnProjectedStampPointsIndex)
-			{
-				FRay3d CameraToPathRay(CameraState.Position, (UnProjectedStampPoints[UnProjectedStampPointsIndex] - CameraState.Position).GetSafeNormal());
-
-				TUniquePtr<FMeshSculptBrushOp>& UseBrushOp = GetActiveBrushOp();
-				ESculptBrushOpTargetType TargetType = UseBrushOp->GetBrushTargetType();
-
-				FVector3d ProjectedPosition;
-				FVector3d ProjectedNormal;
-				int32 HitTriangle = IndexConstants::InvalidID;
-
-				switch (TargetType)
-				{
-				case ESculptBrushOpTargetType::SculptMesh:
-					ProjectWorldRayOnSculptMesh(CameraToPathRay, true, ProjectedPosition, ProjectedNormal, HitTriangle);
-					break;
-				case ESculptBrushOpTargetType::TargetMesh:
-					ProjectWorldRayOnTargetMesh(CameraToPathRay, true, ProjectedPosition, ProjectedNormal, HitTriangle);
-					break;
-				case ESculptBrushOpTargetType::ActivePlane:
-					ProjectWorldRayOnActivePlane(CameraToPathRay, ProjectedPosition, ProjectedNormal);
-					break;
-				}				
-
-				FColor PointColor = FColor::Green;
-				PointSet->AddPoint(ProjectedPosition, PointColor, 6.0, 0.0);
-			}
-		}
-	}
-	else
-	{
-		if (StrokeGeometry)
-		{
-			ULineSetComponent* LineSet = StrokeGeometry->FindLineSet(MeshSculptBaseLocals::StrokeLineSetID);
-			if (LineSet)
-			{
-				LineSet->Clear();
-			}
-			UPointSetComponent* PointSet = StrokeGeometry->FindPointSet(MeshSculptBaseLocals::StrokePointSetID);
-			if (PointSet)
-			{
-				PointSet->Clear();
-			}
-		}
-	}
-
-
 	// always using HoverStamp here because it's position should always be up-to-date for the current apply-stamp...
 	BrushIndicator->Update((float)GetCurrentBrushRadius(),
 		HoverStamp.WorldFrame.ToFTransform(), 1.0f - (float)GetCurrentBrushFalloff());
@@ -759,20 +635,6 @@ bool UMeshSculptToolBase::HitTest(const FRay& Ray, FHitResult& OutHit)
 	return false;
 }
 
-
-void UMeshSculptToolBase::ClearStrokePath()
-{
-	StrokePath.Empty();
-}
-
-void UMeshSculptToolBase::AppendStrokePath(const FVector3d& NextPoint)
-{
-	if(StrokePath.Num() == 0 || !FMath::IsNearlyZero((StrokePath.Last().Position - NextPoint).SquaredLength()))
-	{
-		StrokePath.Add({ .Timestamp = ActiveStrokeTime, .Position = NextPoint, .Pressure = GetActivePressure() });
-	}
-}
-
 void UMeshSculptToolBase::OnBeginDrag(const FRay& WorldRay)
 {
 	SaveActiveStrokeModifiers();
@@ -782,10 +644,14 @@ void UMeshSculptToolBase::OnBeginDrag(const FRay& WorldRay)
 	{
 		bInStroke = true;
 		ResetStrokeTime();
+		PendingInputRays.Reset();
+		PreviousInputRay.Reset();
+		PreviousUnfilteredHitPosition.Reset();
 
 		UpdateBrushTargetPlaneFromHit(WorldRay, OutHit);
 
 		// initialize first stamp
+		PendingInputRays.Add({ .Ray = WorldRay, .Timestamp = ActiveStrokeTime, .Pressure = GetActivePressure() });
 		PendingStampRay = WorldRay;
 		bIsStampPending = true;
 
@@ -801,29 +667,8 @@ void UMeshSculptToolBase::OnUpdateDrag(const FRay& WorldRay)
 {
 	if (InStroke())
 	{
+		PendingInputRays.Add({ .Ray = WorldRay, .Timestamp = ActiveStrokeTime, .Pressure = GetActivePressure() });
 		PendingStampRay = WorldRay;
-
-		TUniquePtr<FMeshSculptBrushOp>& UseBrushOp = GetActiveBrushOp();
-		ESculptBrushOpTargetType TargetType = UseBrushOp->GetBrushTargetType();
-
-		FVector3d ProjectedPosition;
-		FVector3d ProjectedNormal;
-		int32 HitTriangle = IndexConstants::InvalidID;
-
-		switch (TargetType)
-		{
-		case ESculptBrushOpTargetType::SculptMesh:
-			ProjectWorldRayOnSculptMesh(WorldRay, true, ProjectedPosition, ProjectedNormal, HitTriangle);
-			break;
-		case ESculptBrushOpTargetType::TargetMesh:
-			ProjectWorldRayOnTargetMesh(WorldRay, true, ProjectedPosition, ProjectedNormal, HitTriangle);
-			break;
-		case ESculptBrushOpTargetType::ActivePlane:
-			ProjectWorldRayOnActivePlane(WorldRay, ProjectedPosition, ProjectedNormal);
-			break;
-		}
-
-		AppendStrokePath(ProjectedPosition);
 	}
 }
 
@@ -1104,7 +949,6 @@ void UMeshSculptToolBase::ProcessPerTickStamps(TFunction<bool(const FRay& StampR
 
 				UpdateStampPosition(StampRay.StampRay);
 				ExecuteStampOperation(Index, StampRay.StampRay);
-				StampsDuringStroke++;
 			}
 			PostExecuteStampsOperation();
 		}
@@ -1131,6 +975,8 @@ void UMeshSculptToolBase::UpdateStampPendingState()
 	{
 	case EMeshSculptStrokeType::Airbrush:
 	{
+		// TODO: Airbrush mode should actually process PendingInputRays instead of PendingStampRay
+		//  so that it properly handles cases where we get multiple drag updates in between a tick.
 		bool bFlowStampPending = false;
 		float UseStampFlowRate = GetStampTemporalFlowRate();
 		if (UseStampFlowRate >= 1.0)
@@ -1159,120 +1005,182 @@ void UMeshSculptToolBase::UpdateStampPendingState()
 		break;
 	case EMeshSculptStrokeType::Spacing:
 	{
-		// Save our active pressure since we'll be manipulating it during the next block
-		TGuardValue<double> CachedActivePressure(ActivePressure, ActivePressure);
+		// We do two things in this mode: we upsample any input rays that are too far apart, and we filter
+		//  out any hit positions that are close together.
 
-		TArray<FTimedStrokePoint> UnProjectedStampPoints;
-		if (!StrokePath.IsEmpty())
+		if (PendingInputRays.IsEmpty())
 		{
-			UnProjectedStampPoints.Add(StrokePath[0]);
+			bIsStampPending = false;
+			return;
 		}
-		if (StrokePath.Num() > 1)
+		bIsStampPending = true;
+		
+		auto GetHitPosition = [this](const FRay& WorldRay, FVector3d& HitPositionOut)->bool
 		{
-			TArray<FTimedStrokePoint> UnconsumedPathPoints;
+			TUniquePtr<FMeshSculptBrushOp>& UseBrushOp = GetActiveBrushOp();
+			ESculptBrushOpTargetType TargetType = UseBrushOp->GetBrushTargetType();
 
-			int32 SectionIndex = 0;
-			double SpacingSectionConsumedLength = 0.0;
+			bool bHit = false;
+			FVector3d HitNormal;
+			int32 HitTriangle = IndexConstants::InvalidID;
+			switch (TargetType)
+			{
+			case ESculptBrushOpTargetType::SculptMesh:
+				bHit = ProjectWorldRayOnSculptMesh(WorldRay, true, HitPositionOut, HitNormal, HitTriangle);
+				break;
+			case ESculptBrushOpTargetType::TargetMesh:
+				bHit = ProjectWorldRayOnTargetMesh(WorldRay, true, HitPositionOut, HitNormal, HitTriangle);
+				break;
+			case ESculptBrushOpTargetType::ActivePlane:
+				bHit = ProjectWorldRayOnActivePlane(WorldRay, HitPositionOut, HitNormal);
+				break;
+			}
 
-			UnconsumedPathPoints.Add(StrokePath[0]);
+			return bHit;
+		};
+		
+		// Start by upsampling into PendingStampRays. We will filter them further below.
+
+		if (BrushProperties->RayDensity <= KINDA_SMALL_NUMBER)
+		{
+			// No upsampling needed, just use the same rays
+			Algo::Transform(PendingInputRays, PendingStampRays, [](const FInputRayData& InputRay)
+				{ return FStampRayData{ .StampRay = InputRay.Ray, .Pressure = InputRay.Pressure }; });
+		}
+		else
+		{
+			FFrame3d CameraPlane(CameraState.Position + CameraState.Forward(), CameraState.Forward());
 
-			for (int32 Index = 0; Index < StrokePath.Num() - 1; Index++)
+			// At ray density 1.0, if we snap the ray from the left side of the screen to the right
+			//  side, we will say that we want to interpolate this many samples
+			constexpr double NumSamplesForScreenWidth = 200;
+			double ScreenWidth = CameraState.bIsOrthographic ? CameraState.OrthoWorldCoordinateWidth
+				// Width of screen at 1 unit away from camera
+				: FMath::Tan(FMath::DegreesToRadians(FMath::Clamp(CameraState.HorizontalFOVDegrees / 2, 0.f, 179.f))) * 2;
+			double CameraPlaneSpacing = ScreenWidth / NumSamplesForScreenWidth;
+
+			auto GetCameraPlaneIntersection = [this, &CameraPlane, &GetHitPosition](const FRay& WorldRay, FVector3d& IntersectionOut) -> bool
 			{
-				// Compute average spacing over stroke section - needed for dynamic spacing as the pressure changes over the path.
-				double Spacing = 1.0;
+				// This is simpler if we assume that we're shooting rays from behind the camera plane. But
+				//  we'll do the proper thing and imagine that we may be getting VR rays. We'll find the
+				//  hit location and then project back from there.
+				FVector3d HitPosition;
+				if (!GetHitPosition(WorldRay, HitPosition))
 				{
-					ActivePressure = StrokePath[Index].Pressure;
-					double SpacingAtBeginning = FMath::Clamp(BrushProperties->Spacing * GetActiveBrushRadius() * 2, 0.1, 100000);
-					ActivePressure = StrokePath[Index + 1].Pressure;
-					double SpacingAtEnding = FMath::Clamp(BrushProperties->Spacing * GetActiveBrushRadius() * 2, 0.1, 100000);
-					Spacing = (SpacingAtBeginning + SpacingAtEnding) * 0.5;
+					return false;
 				}
+				FVector3d Direction = HitPosition - CameraState.Position;
+				if (!Direction.Normalize())
+				{
+					return false;
+				}
+				return CameraPlane.RayPlaneIntersection(CameraState.Position, Direction, 2, IntersectionOut);
+			};
 
-				FVector3d Span = StrokePath[Index + 1].Position - StrokePath[Index].Position; // Vector representing our stroke segment
-				FVector3d SpanDir = Span.GetSafeNormal(); // The direction of our stroke segment
-
-				double RemainingSpace = FMath::Max(0, Spacing - SpacingSectionConsumedLength); // How much space do we have left before a stamp happens
-				double OverflowSpace = Span.Length() - RemainingSpace; // If positive, our stroke section meets or exceeds the remaining space for a stamp to happen
-
-				if (OverflowSpace < 0)
+			// Appends to PendingStampRays. Does not include start, but does include end.
+			auto UpsampleSegment = [this, &GetCameraPlaneIntersection, CameraPlaneSpacing](const FInputRayData& Start, const FInputRayData& End)
+			{
+				FVector3d CameraPlaneEnd, CameraPlaneStart;
+				if (!GetCameraPlaneIntersection(End.Ray, CameraPlaneEnd))
 				{
-					SpacingSectionConsumedLength += Span.Length();
-					UnconsumedPathPoints.Add(StrokePath[Index + 1]);
+					return;
 				}
-				else
+				if (!GetCameraPlaneIntersection(Start.Ray, CameraPlaneStart))
 				{
-					double SpanAlpha = FMath::IsNearlyZero(Span.Length()) ? 0.0 : RemainingSpace / Span.Length(); // How far along, within the current stroke segment, are we?
-					double OverflowAlphaIncrement = FMath::IsNearlyZero(Span.Length()) ? 0.0 : Spacing / Span.Length(); // For our current stroke segment, what is the percentage of it's length in terms of a spacing interval?
-					
-					FVector3d TerminalPosition = FMath::LerpStable(StrokePath[Index].Position, StrokePath[Index + 1].Position, SpanAlpha);
-					SpacingSectionConsumedLength += RemainingSpace;
+					// Can at least place the end
+					PendingStampRays.Add({ .StampRay = End.Ray, .Pressure = End.Pressure });
+					return;
+				}
 
-					UnProjectedStampPoints.Add({ .Position = TerminalPosition, .Pressure = FMath::LerpStable(StrokePath[Index].Pressure, StrokePath[Index + 1].Pressure, SpanAlpha) });
+				FVector3d CameraPlaneSegmentVector = CameraPlaneEnd - CameraPlaneStart;
+				double DistInCameraPlane = CameraPlaneSegmentVector.Length();
+				if (DistInCameraPlane == 0)
+				{
+					return;
+				}
 
-					while (OverflowSpace > Spacing) // Handle cases where we have so much line we inject multiple points...
+				int32 NumSubSegments = FMath::CeilToInt(DistInCameraPlane / CameraPlaneSpacing);
+				for (int i = 1; i <= NumSubSegments; ++i) // 1-based seems simpler here
+				{
+					double TAlongSegment = (double)i / NumSubSegments;
+
+					FVector3d PositionInCameraPlane = CameraPlaneStart + TAlongSegment * CameraPlaneSegmentVector;
+					FVector3d Direction = PositionInCameraPlane - CameraState.Position;
+					
+					if (ensure(Direction.Normalize()))
 					{
-						SpanAlpha += OverflowAlphaIncrement;
-						TerminalPosition = FMath::LerpStable(StrokePath[Index].Position, StrokePath[Index + 1].Position, SpanAlpha);
-						UnProjectedStampPoints.Add({ .Position = TerminalPosition, .Pressure = FMath::LerpStable(StrokePath[Index].Pressure, StrokePath[Index + 1].Pressure, SpanAlpha) });
-						OverflowSpace -= Spacing;
+						double InterpolatedPressure = FMath::Lerp(Start.Pressure, End.Pressure, TAlongSegment);
+						PendingStampRays.Add({ .StampRay = FRay(CameraState.Position, Direction), .Pressure = InterpolatedPressure });
 					}
-
-					ensure(SpanAlpha <= 1.0);
-
-					SpacingSectionConsumedLength = OverflowSpace;
-					UnconsumedPathPoints.Empty();
-					UnconsumedPathPoints.Add(
-						{
-								.Timestamp = FMath::LerpStable(StrokePath[Index].Timestamp, StrokePath[Index + 1].Timestamp, SpanAlpha),
-								.Position = TerminalPosition,
-								.Pressure = FMath::LerpStable(StrokePath[Index].Pressure, StrokePath[Index + 1].Pressure, SpanAlpha)
-						});
-					UnconsumedPathPoints.Add({ StrokePath[Index + 1] });
 				}
+			};
+
+			if (!PreviousInputRay.IsSet())
+			{
+				PendingStampRays.Add({ .StampRay = PendingInputRays[0].Ray, .Pressure = PendingInputRays[0].Pressure });
+			}
+			else
+			{
+				UpsampleSegment(PreviousInputRay.GetValue(), PendingInputRays[0]);
 			}
 
-			StrokePath = UnconsumedPathPoints;
-		}
+			for (int32 PendingInputIndex = 1; PendingInputIndex < PendingInputRays.Num(); ++PendingInputIndex)
+			{
+				UpsampleSegment(PendingInputRays[PendingInputIndex - 1], PendingInputRays[PendingInputIndex]);
+			}
+		}//end upscaling
+		PreviousInputRay = PendingInputRays.Last();
+		PendingInputRays.Reset();
+		
+		// Filter our results
 
-		// We want to ensure that when starting the stroke, we stamp right at the beginning of the path
-		int StartIndex = 1;
-		if (StampsDuringStroke == 0)
+		TGuardValue<double> CachedActivePressure(ActivePressure, ActivePressure);
+		auto GetRadiusFromPressure = [this](double PressureIn)
 		{
-			StartIndex = 0;
-		}
+			ActivePressure = PressureIn;
+			return GetActiveBrushRadius();
+		};
 
-		for (int32 UnProjectedStampPointsIndex = StartIndex; UnProjectedStampPointsIndex < UnProjectedStampPoints.Num(); ++UnProjectedStampPointsIndex)
+		int32 DestIndex = 0;
+		for (int32 SourceIndex = 0; SourceIndex < PendingStampRays.Num(); ++SourceIndex)
 		{
-			FRay3d CameraToPathRay(CameraState.Position, (UnProjectedStampPoints[UnProjectedStampPointsIndex].Position - CameraState.Position).GetSafeNormal());
-
-			TUniquePtr<FMeshSculptBrushOp>& UseBrushOp = GetActiveBrushOp();
-			ESculptBrushOpTargetType TargetType = UseBrushOp->GetBrushTargetType();
-
-			FVector3d ProjectedPosition = FVector3d::ZeroVector;
-			FVector3d ProjectedNormal = FVector3d::ZeroVector;
-			int32 HitTriangle = IndexConstants::InvalidID;
+			FVector3d HitPosition;
+			auto PassItem = [SourceIndex, &DestIndex, this]()
+			{
+				if (SourceIndex != DestIndex) 
+				{
+					PendingStampRays[DestIndex] = PendingStampRays[SourceIndex]; 
+				}
+				++DestIndex;
+			};
 
-			switch (TargetType)
+			if (!GetHitPosition(PendingStampRays[SourceIndex].StampRay, HitPosition))
 			{
-			case ESculptBrushOpTargetType::SculptMesh:
-				ProjectWorldRayOnSculptMesh(CameraToPathRay, true, ProjectedPosition, ProjectedNormal, HitTriangle);
-				break;
-			case ESculptBrushOpTargetType::TargetMesh:
-				ProjectWorldRayOnTargetMesh(CameraToPathRay, true, ProjectedPosition, ProjectedNormal, HitTriangle);
-				break;
-			case ESculptBrushOpTargetType::ActivePlane:
-				ProjectWorldRayOnActivePlane(CameraToPathRay, ProjectedPosition, ProjectedNormal);
-				break;
+				// It's not our job to filter out misses, only hits so that we can get spacing
+				PassItem();
+				continue;
 			}
 
-			PendingStampRays.Add({ .StampRay = FRay3d(CameraState.Position, (ProjectedPosition - CameraState.Position).GetSafeNormal()), .Pressure = UnProjectedStampPoints[UnProjectedStampPointsIndex].Pressure });
-		}
+			double HitRadius = GetRadiusFromPressure(PendingStampRays[SourceIndex].Pressure);
+			if (!PreviousUnfilteredHitPosition.IsSet())
+			{
+				PassItem();
+				PreviousUnfilteredHitPosition = HitPosition;
+				PreviousUnfilteredHitRadius = HitRadius;
+				continue;
+			}
 
-		if (!PendingStampRays.IsEmpty())
-		{
-			bIsStampPending = true;
+			double DistanceToNotOverlap = PreviousUnfilteredHitRadius + HitRadius;
+			double Spacing = BrushProperties->Spacing * DistanceToNotOverlap;
+			if (FVector3d::Distance(HitPosition, PreviousUnfilteredHitPosition.GetValue()) >= Spacing)
+			{
+				PassItem();
+				PreviousUnfilteredHitPosition = HitPosition;
+				PreviousUnfilteredHitRadius = HitRadius;
+				continue;
+			}
 		}
-
+		PendingStampRays.SetNum(DestIndex);
 		return;
 	}
 		break;
@@ -1288,10 +1196,6 @@ void UMeshSculptToolBase::UpdateStampPendingState()
 	default:
 		break;
 	}
-
-
-
-
 }
 
 
@@ -1301,8 +1205,6 @@ void UMeshSculptToolBase::ResetStrokeTime()
 	LastFlowTimeStamp = 0;
 	ActiveStrokePathArcLen = 0;
 	LastSpacingTimestamp = 0;
-	StampsDuringStroke = 0;
-	ClearStrokePath();
 }
 
 void UMeshSculptToolBase::AccumulateStrokeTime(float DeltaTime)
diff --git a/Engine/Plugins/Runtime/MeshModelingToolset/Source/MeshModelingTools/Public/Sculpting/MeshSculptToolBase.h b/Engine/Plugins/Runtime/MeshModelingToolset/Source/MeshModelingTools/Public/Sculpting/MeshSculptToolBase.h
--- a/Engine/Plugins/Runtime/MeshModelingToolset/Source/MeshModelingTools/Public/Sculpting/MeshSculptToolBase.h
+++ b/Engine/Plugins/Runtime/MeshModelingToolset/Source/MeshModelingTools/Public/Sculpting/MeshSculptToolBase.h
@@ -133,8 +133,16 @@ public:
 	UPROPERTY(EditAnywhere, Category = Brush, meta = (DisplayName = "Flow", UIMin = "0.0", UIMax = "1.0", ClampMin = "0.0", ClampMax = "1.0", DisplayPriority = 7, HideEditConditionToggle, EditConditionHides, EditCondition = "bShowFlowRate && StrokeType == EMeshSculptStrokeType::Airbrush"))
 	float FlowRate = 1.0f;
 
+	/** When moving the mouse quickly in Spacing mode, determines how initial rays are upsampled before being filtered by StampSpacing. */
+	UPROPERTY(EditAnywhere, Category = Brush, AdvancedDisplay, meta = (UIMin = "0.0", UIMax = "2.0", ClampMin = "0.0", ClampMax = "10.0",
+		DisplayPriority = 8, HideEditConditionToggle, EditConditionHides,
+		EditCondition = "bShowSpacing && StrokeType == EMeshSculptStrokeType::Spacing"))
+	float RayDensity = 1.0f;
+
 	/** Space out stamp centers at distances Spacing*BrushDiameter along the stroke (so Spacing of 1 means that stamps will be adjacent but non-overlapping). Zero spacing means continuous stamps. */
-	UPROPERTY(EditAnywhere, Category = Brush, meta = (UIMin = "0.0", UIMax = "4.0", ClampMin = "0.0", ClampMax = "1000.0", DisplayPriority = 8, HideEditConditionToggle, EditConditionHides, EditCondition = "bShowSpacing && StrokeType == EMeshSculptStrokeType::Spacing"))
+	UPROPERTY(EditAnywhere, Category = Brush, meta = (UIMin = "0.0", UIMax = "1.0", ClampMin = "0.0", ClampMax = "1000.0", 
+		DisplayPriority = 8, DisplayName = "Stamp Spacing", HideEditConditionToggle, EditConditionHides, 
+		EditCondition = "bShowSpacing && StrokeType == EMeshSculptStrokeType::Spacing"))
 	float Spacing = .05f;
 
 	/** Lazy brush smooths out the brush path by averaging the cursor positions */
@@ -528,19 +536,22 @@ protected:
 	virtual bool GetInSmoothingStroke() const { return bSmoothing; }
 	virtual bool GetInInvertStroke() const { return bInvert; }
 
-	struct FTimedStrokePoint
+private:
+	struct FInputRayData
 	{
+		FRay Ray;
 		double Timestamp;
-		FVector3d Position;		
 		double Pressure;
 	};
-	TArray<FTimedStrokePoint> StrokePath;
-
-
+	TArray<FInputRayData> PendingInputRays;
+	// If this is not the first time we're processing PendingInputRays during a stroke, the previous ray
+	TOptional<FInputRayData> PreviousInputRay;
+	TOptional<FVector3d> PreviousUnfilteredHitPosition;
+	double PreviousUnfilteredHitRadius = 0;
+	
 	UPROPERTY()
 	TObjectPtr<UPreviewGeometry> StrokeGeometry;
-	UE_API virtual void ClearStrokePath();
-	UE_API virtual void AppendStrokePath(const FVector3d& NextPoint);
+protected:
 
 	// when in a stroke, this function determines when a new stamp should be emitted, based on spacing and flow rate settings
 	UE_API virtual void UpdateStampPendingState();
@@ -550,7 +561,6 @@ protected:
 	double ActiveStrokePathArcLen = 0.0;
 	int LastFlowTimeStamp = 0;
 	int LastSpacingTimestamp = 0;
-	int StampsDuringStroke = 0;
 	UE_API virtual void ResetStrokeTime();
 	UE_API virtual void AccumulateStrokeTime(float DeltaTime);
 
@@ -561,13 +571,17 @@ protected:
 	bool bIsStampPending = false;
 	FRay PendingStampRay;
 
+private:
 	struct FStampRayData
 	{
 		FRay StampRay;
 		double Pressure;
 	};
-
+	// Sequence of stamps to apply. Initialized by UpdateStampPendingState and immediately
+	//  applied, though potentially with filtering
 	TArray<FStampRayData> PendingStampRays;
+protected:
+
 	FSculptBrushStamp HoverStamp;
 	FSculptBrushStamp CurrentStamp;
 	FSculptBrushStamp LastStamp;
