Remesh Tool: add relative density weight map attribute

diff --git a/Engine/Plugins/Runtime/MeshModelingToolset/Source/MeshModelingTools/Private/RemeshMeshTool.cpp b/Engine/Plugins/Runtime/MeshModelingToolset/Source/MeshModelingTools/Private/RemeshMeshTool.cpp
--- a/Engine/Plugins/Runtime/MeshModelingToolset/Source/MeshModelingTools/Private/RemeshMeshTool.cpp
+++ b/Engine/Plugins/Runtime/MeshModelingToolset/Source/MeshModelingTools/Private/RemeshMeshTool.cpp
@@ -32,6 +32,28 @@ using namespace UE::Geometry;
 
 #define LOCTEXT_NAMESPACE "URemeshMeshTool"
 
+/*
+ * Tool properties
+ */
+
+void URemeshMeshToolProperties::InitializeAttributes(const TArray<FName>& InAttributeNames)
+{
+	AttributeNames.Init(FString(), InAttributeNames.Num());
+	for (int32 AttributeIndex = 0; AttributeIndex < InAttributeNames.Num(); ++AttributeIndex)
+	{
+		AttributeNames[AttributeIndex] = InAttributeNames[AttributeIndex].ToString();
+	}
+}
+
+void URemeshMeshToolProperties::ValidateSelectedAttribute()
+{
+	const int32 FoundIndex = AttributeNames.IndexOfByKey(SelectedAttribute);
+	if (FoundIndex == INDEX_NONE)
+	{
+		SelectedAttribute = (AttributeNames.Num() > 0) ? AttributeNames[0] : TEXT("");
+	}
+}
+
 /*
  * ToolBuilder
  */
@@ -134,6 +156,26 @@ void URemeshMeshTool::Setup()
 	BasicProperties->TargetTriangleCount = (OriginalMesh->TriangleCount() < 5000) ? 5000 : OriginalMesh->TriangleCount();
 	BasicProperties->TargetEdgeLength = FRemeshMeshOp::CalculateTargetEdgeLength(OriginalMesh.Get(), BasicProperties->TargetTriangleCount, InitialMeshArea);
 
+	// Attribute properties
+	BasicProperties->WatchProperty(BasicProperties->SelectedAttribute,
+		[this](const FString& NewName)
+		{
+			BasicProperties->ValidateSelectedAttribute();
+		});
+
+	TArray<FName> AttributeNames;
+	if (FDynamicMeshAttributeSet* Attributes = OriginalMesh->Attributes())
+	{
+		const int32 NumLayers = Attributes->NumWeightLayers();
+		AttributeNames.SetNum(NumLayers);
+		for (int32 LayerIdx = 0; LayerIdx < NumLayers; ++LayerIdx)
+		{
+			AttributeNames[LayerIdx] = Attributes->GetWeightLayer(LayerIdx)->GetName();
+		}
+	}
+	BasicProperties->InitializeAttributes(AttributeNames);
+	BasicProperties->ValidateSelectedAttribute();
+
 	// add properties to GUI
 	AddToolPropertySource(BasicProperties);
 	AddToolPropertySource(MeshStatisticsProperties);
@@ -246,6 +288,50 @@ TUniquePtr<FDynamicMeshOperator> URemeshMeshTool::MakeNewOperator()
 	Op->ProjectionTarget = nullptr;
 	Op->ProjectionTargetSpatial = nullptr;
 
+
+	if (BasicProperties->bUseWeightMap && OriginalMesh->HasAttributes())
+	{
+		int32 WeightMapAttributeLayerIndex = -1;
+		if (const FDynamicMeshAttributeSet* const Attribs = OriginalMesh->Attributes())
+		{
+			for (int32 AttributeIndex = 0; AttributeIndex < Attribs->NumWeightLayers(); ++AttributeIndex)
+			{
+				const FDynamicMeshWeightAttribute* const AttributeLayer = Attribs->GetWeightLayer(AttributeIndex);
+
+				if (AttributeLayer->GetName() == BasicProperties->SelectedAttribute)
+				{
+					WeightMapAttributeLayerIndex = AttributeIndex;
+					break;
+				}
+			}
+		}
+
+		if (WeightMapAttributeLayerIndex >= 0)
+		{
+			Op->CustomEdgeLengthScaleF = [WeightMapAttributeLayerIndex, RelativeDensity = BasicProperties->RelativeDensity](const FDynamicMesh3& Mesh, int VertexA, int VertexB) -> double
+			{
+				if (Mesh.HasAttributes() && Mesh.Attributes()->GetWeightLayer(WeightMapAttributeLayerIndex))
+				{
+					const FDynamicMeshWeightAttribute* const AttributeLayer = Mesh.Attributes()->GetWeightLayer(WeightMapAttributeLayerIndex);
+
+					float WeightValueA;
+					AttributeLayer->GetValue(VertexA, &WeightValueA);
+					WeightValueA = FMath::Clamp(WeightValueA, 0.0f, 1.0f);
+
+					float WeightValueB;
+					AttributeLayer->GetValue(VertexB, &WeightValueB);
+					WeightValueB = FMath::Clamp(WeightValueB, 0.0f, 1.0f);
+
+					const double AvgWeight = 0.5 * (WeightValueA + WeightValueB);
+
+					return FMath::Pow(0.5, -RelativeDensity * AvgWeight);
+				}
+
+				return 1.0;
+			};
+		}
+	}
+
 	return Op;
 }
 
diff --git a/Engine/Plugins/Runtime/MeshModelingToolset/Source/MeshModelingTools/Public/RemeshMeshTool.h b/Engine/Plugins/Runtime/MeshModelingToolset/Source/MeshModelingTools/Public/RemeshMeshTool.h
--- a/Engine/Plugins/Runtime/MeshModelingToolset/Source/MeshModelingTools/Public/RemeshMeshTool.h
+++ b/Engine/Plugins/Runtime/MeshModelingToolset/Source/MeshModelingTools/Public/RemeshMeshTool.h
@@ -100,6 +100,33 @@ public:
 	UPROPERTY(EditAnywhere, Category = Remeshing, AdvancedDisplay, meta = (EditCondition = "bReprojectConstraints", UIMin = "0", UIMax = "180", ClampMin = "0", ClampMax = "180"))
 	float BoundaryCornerAngleThreshold;
 
+
+	/** Whether to use a vertex weight map layer to control relative mesh density */
+	UPROPERTY(EditAnywhere, Category = WeightMap)
+	bool bUseWeightMap;
+
+	/** Name of the vertex weight map layer that will be used to adjust relative mesh density */
+	UPROPERTY(EditAnywhere, Category = WeightMap, meta = (DisplayName = "Attribute", GetOptions = GetAttributeNames, EditCondition = bUseWeightMap))
+	FString SelectedAttribute;
+
+	UFUNCTION()
+	const TArray<FString>& GetAttributeNames() { return AttributeNames; };
+	TArray<FString> AttributeNames;
+
+	/** Amount to adjust density in regions where the vertex weight map is greater than zero. A positive value here results in a denser mesh, a negative value results in a coarser one */
+	UPROPERTY(EditAnywhere, Category = WeightMap, meta = (EditCondition = bUseWeightMap, UIMin = "-2.0", UIMax = "2.0"))
+	float RelativeDensity = 0.f;
+
+	/**
+	 * Set up the list of available attributes
+	 */
+	void InitializeAttributes(const TArray<FName>& InAttributeNames);
+
+	/**
+	 * If the attribute selection is not valid, use the attribute at index 0 or empty if there are no attributes
+	 */
+	void ValidateSelectedAttribute();
+
 };
 
 
