Remesh Tool: set maxium initial triangle count to 100k

diff --git a/Engine/Plugins/Runtime/MeshModelingToolset/Source/MeshModelingTools/Private/RemeshMeshTool.cpp b/Engine/Plugins/Runtime/MeshModelingToolset/Source/MeshModelingTools/Private/RemeshMeshTool.cpp
--- a/Engine/Plugins/Runtime/MeshModelingToolset/Source/MeshModelingTools/Private/RemeshMeshTool.cpp
+++ b/Engine/Plugins/Runtime/MeshModelingToolset/Source/MeshModelingTools/Private/RemeshMeshTool.cpp
@@ -152,8 +152,9 @@ void URemeshMeshTool::Setup()
 
 	// set properties defaults
 
-	// arbitrary threshold of 5000 tris seems reasonable?
-	BasicProperties->TargetTriangleCount = (OriginalMesh->TriangleCount() < 5000) ? 5000 : OriginalMesh->TriangleCount();
+	constexpr int32 MinInitialTriangles = 5000;
+	constexpr int32 MaxInitialTriangles = 100'000;
+	BasicProperties->TargetTriangleCount = FMath::Clamp(OriginalMesh->TriangleCount(), MinInitialTriangles, MaxInitialTriangles);
 	BasicProperties->TargetEdgeLength = FRemeshMeshOp::CalculateTargetEdgeLength(OriginalMesh.Get(), BasicProperties->TargetTriangleCount, InitialMeshArea);
 
 	// Attribute properties
