[SKMModelingTools] Added skin weight locking to skin weight paint tool

diff --git a/Engine/Plugins/Animation/SkeletalMeshModelingTools/Source/SkeletalMeshModelingTools/Private/DetailCustomization/SkinWeightDetailCustomization.cpp b/Engine/Plugins/Animation/SkeletalMeshModelingTools/Source/SkeletalMeshModelingTools/Private/DetailCustomization/SkinWeightDetailCustomization.cpp
--- a/Engine/Plugins/Animation/SkeletalMeshModelingTools/Source/SkeletalMeshModelingTools/Private/DetailCustomization/SkinWeightDetailCustomization.cpp
+++ b/Engine/Plugins/Animation/SkeletalMeshModelingTools/Source/SkeletalMeshModelingTools/Private/DetailCustomization/SkinWeightDetailCustomization.cpp
@@ -1580,7 +1580,14 @@ TSharedRef<SWidget> SVertexWeightItem::GenerateWidgetForColumn(const FName& Colu
 	if (ColumnName == ColumnName_Bone)
 	{
 		const FName BoneName = ParentTable->Tool.IsValid() ? ParentTable->Tool->GetBoneNameFromIndex(Element->BoneIndex) : NAME_None;
-		return SNew(STextBlock).Text(FText::FromName(BoneName));
+		return SNew(STextBlock)
+			.Text_Lambda([Tool = ParentTable->Tool, BoneName, BoneIndex = Element->BoneIndex]()
+				{
+					bool bLocked = Tool.IsValid() ? Tool->GetFinalLockedBoneIndices().Contains(BoneIndex) : false;
+					return FText::Format(LOCTEXT("VertexWeightBoneNameFormat", "{0}{1}"),
+						FText::FromName(BoneName),
+						bLocked ? LOCTEXT("VertexWeightBoneNameLockedText", " (Locked)") : FText());
+				});
 	}
 
 	if (ColumnName == ColumnName_Weight)
diff --git a/Engine/Plugins/Animation/SkeletalMeshModelingTools/Source/SkeletalMeshModelingTools/Private/SReferenceSkeletonTree.cpp b/Engine/Plugins/Animation/SkeletalMeshModelingTools/Source/SkeletalMeshModelingTools/Private/SReferenceSkeletonTree.cpp
deleted file mode 100644
diff --git a/Engine/Plugins/Animation/SkeletalMeshModelingTools/Source/SkeletalMeshModelingTools/Private/SReferenceSkeletonTree.h b/Engine/Plugins/Animation/SkeletalMeshModelingTools/Source/SkeletalMeshModelingTools/Private/SReferenceSkeletonTree.h
deleted file mode 100644
diff --git a/Engine/Plugins/Animation/SkeletalMeshModelingTools/Source/SkeletalMeshModelingTools/Private/SkeletalMeshEditingCache.cpp b/Engine/Plugins/Animation/SkeletalMeshModelingTools/Source/SkeletalMeshModelingTools/Private/SkeletalMeshEditingCache.cpp
--- a/Engine/Plugins/Animation/SkeletalMeshModelingTools/Source/SkeletalMeshModelingTools/Private/SkeletalMeshEditingCache.cpp
+++ b/Engine/Plugins/Animation/SkeletalMeshModelingTools/Source/SkeletalMeshModelingTools/Private/SkeletalMeshEditingCache.cpp
@@ -6,6 +6,7 @@
 #include "CanvasTypes.h"
 #include "PreviewMesh.h"
 #include "SkeletalDebugRendering.h"
+#include "SkeletalMeshEditorUtils.h"
 #include "Components/SKMBackedDynaMeshComponent.h"
 #include "Preferences/PersonaOptions.h"
 #include "SkeletalMesh/RefSkeletonPoser.h"
@@ -104,6 +105,14 @@ void USkeletalMeshEditingCache::Spawn(UWorld* World, USkeletalMeshComponent* InS
 	bCacheBoneManipulation = true;
 
 	SetSkeletalMeshSkeletonDrawMode(CacheSkeletonDrawMode);
+
+	if (Delegates.OnGetSkeletalMeshSkeletonNotifierDelegate.IsBound())
+	{
+		TSharedPtr<ISkeletalMeshNotifier> SkeletalMeshSkeletonNotifier = Delegates.OnGetSkeletalMeshSkeletonNotifierDelegate.Execute();
+					
+		SkeletalMeshSkeletonNotifierBindScope.Reset(
+			new FSkeletalMeshNotifierBindScope(GetNotifier(), SkeletalMeshSkeletonNotifier));
+	}
 	
 	ChangeCountWatcher.Initialize(
 		[this]()
@@ -135,6 +144,24 @@ void USkeletalMeshEditingCache::Spawn(UWorld* World, USkeletalMeshComponent* InS
 				bEnableDynamicMeshSkeleton = SkeletonChangeCount > 0;
 				RefSkeletonPoser->SetRefSkeleton(GetEditingMeshComponent()->GetRefSkeleton());
 				UpdateSelectedBoneIndices();
+				if (bEnableDynamicMeshSkeleton)
+				{
+					SkeletalMeshSkeletonNotifierBindScope.Reset();
+				}
+				else
+				{
+					if (Delegates.OnGetSkeletalMeshSkeletonNotifierDelegate.IsBound())
+					{
+						using namespace UE::SkeletalMeshEditorUtils;
+						TSharedPtr<ISkeletalMeshNotifier> SkeletalMeshSkeletonNotifier = Delegates.OnGetSkeletalMeshSkeletonNotifierDelegate.Execute();
+					
+						SkeletalMeshSkeletonNotifierBindScope.Reset(
+							new FSkeletalMeshNotifierBindScope(GetNotifier(), SkeletalMeshSkeletonNotifier));
+
+						SkeletalMeshSkeletonNotifier->HandleNotification({}, ESkeletalMeshNotifyType::HierarchyChanged);
+						SkeletalMeshSkeletonNotifier->HandleNotification(GetSelectedBones(), ESkeletalMeshNotifyType::BonesSelected);
+					}	
+				}
 			},
 		GetEditingMeshComponent()->GetSkeletonChangeTracker().GetChangeCount());
 
@@ -225,6 +252,7 @@ EMeshLODIdentifier USkeletalMeshEditingCache::GetLOD() const
 
 void USkeletalMeshEditingCache::Destroy()
 {
+	SkeletalMeshSkeletonNotifierBindScope.Reset();
 	HostActor->Destroy();
 	PreviewMesh->Disconnect();
 	
diff --git a/Engine/Plugins/Animation/SkeletalMeshModelingTools/Source/SkeletalMeshModelingTools/Private/SkeletalMeshEditingCache.h b/Engine/Plugins/Animation/SkeletalMeshModelingTools/Source/SkeletalMeshModelingTools/Private/SkeletalMeshEditingCache.h
--- a/Engine/Plugins/Animation/SkeletalMeshModelingTools/Source/SkeletalMeshModelingTools/Private/SkeletalMeshEditingCache.h
+++ b/Engine/Plugins/Animation/SkeletalMeshModelingTools/Source/SkeletalMeshModelingTools/Private/SkeletalMeshEditingCache.h
@@ -44,10 +44,12 @@ class USkeletalMeshEditingCache: public UObject
 public:
 	DECLARE_DELEGATE_OneParam(FToggleSkeletalMeshBoneManipulation, bool);
 	DECLARE_DELEGATE_RetVal(bool, FIsSkeletalMeshBoneManipulationEnabled);
+	DECLARE_DELEGATE_RetVal(TSharedPtr<ISkeletalMeshNotifier>, FOnGetSkeletalMeshSkeletonNotifier);
 	struct FDelegates
 	{
 		FToggleSkeletalMeshBoneManipulation ToggleSkeletalMeshBoneManipulationDelegate;
 		FIsSkeletalMeshBoneManipulationEnabled IsSkeletalMeshBoneManipulationEnabledDelegate;
+		FOnGetSkeletalMeshSkeletonNotifier OnGetSkeletalMeshSkeletonNotifierDelegate;
 	};
 
 	void Spawn(UWorld* World, USkeletalMeshComponent* SkeletalMeshComponent, EMeshLODIdentifier LOD, const FDelegates& InDelegates);
@@ -175,6 +177,7 @@ protected:
 	void DeformPreviewMesh(const TArray<FTransform>& ComponentSpaceTransforms, const TMap<FName, float>& MorphTargetWeights);
 	
 	TSharedPtr<FSkeletalMeshEditingCacheNotifier> Notifier;
+	TUniquePtr<FSkeletalMeshNotifierBindScope> SkeletalMeshSkeletonNotifierBindScope;
 	
 	TArray<FName> SelectedBoneNames;
 	TArray<int32> SelectedBoneIndices;
diff --git a/Engine/Plugins/Animation/SkeletalMeshModelingTools/Source/SkeletalMeshModelingTools/Private/SkeletalMeshEditorUtils.cpp b/Engine/Plugins/Animation/SkeletalMeshModelingTools/Source/SkeletalMeshModelingTools/Private/SkeletalMeshEditorUtils.cpp
--- a/Engine/Plugins/Animation/SkeletalMeshModelingTools/Source/SkeletalMeshModelingTools/Private/SkeletalMeshEditorUtils.cpp
+++ b/Engine/Plugins/Animation/SkeletalMeshModelingTools/Source/SkeletalMeshModelingTools/Private/SkeletalMeshEditorUtils.cpp
@@ -9,7 +9,6 @@
 #include "SkeletalMeshEditingCache.h"
 #include "SkeletalMeshModelingToolsEditorMode.h"
 #include "SkeletonModifier.h"
-#include "SReferenceSkeletonTree.h"
 #include "Widgets/Docking/SDockTab.h"
 
 #include UE_INLINE_GENERATED_CPP_BY_NAME(SkeletalMeshEditorUtils)
@@ -54,46 +53,6 @@ USkeletalMeshEditorContextObject* UE::SkeletalMeshEditorUtils::GetEditorContextO
 	return ToolsContext->ContextObjectStore->FindContext<USkeletalMeshEditorContextObject>();
 }
 
-UE::SkeletalMeshEditorUtils::FSkeletalMeshNotifierBindScope::FSkeletalMeshNotifierBindScope(
-	TWeakPtr<ISkeletalMeshNotifier> InNotifierA,
-	TWeakPtr<ISkeletalMeshNotifier> InNotifierB
-	) : NotifierA(InNotifierA)
-	, NotifierB(InNotifierB)
-{
-
-	DelegateToRemoveFromNotifierA = InNotifierA.Pin()->Delegate().AddLambda(
-	[InNotifierB](const TArray<FName>& BoneNames, const ESkeletalMeshNotifyType InNotifyType)
-	{
-		if (InNotifierB.IsValid())
-		{
-			InNotifierB.Pin()->HandleNotification(BoneNames, InNotifyType);
-		}
-	});
-	
-	DelegateToRemoveFromNotifierB = InNotifierB.Pin()->Delegate().AddLambda(
-	[InNotifierA](const TArray<FName>& BoneNames, const ESkeletalMeshNotifyType InNotifyType)
-	{
-		if (InNotifierA.IsValid())
-		{
-			InNotifierA.Pin()->HandleNotification(BoneNames, InNotifyType);
-		}
-	});
-}
-
-UE::SkeletalMeshEditorUtils::FSkeletalMeshNotifierBindScope::~FSkeletalMeshNotifierBindScope()
-{
-	if (DelegateToRemoveFromNotifierA.IsValid() && NotifierA.IsValid())
-	{
-		check(NotifierA.Pin()->Delegate().Remove(DelegateToRemoveFromNotifierA));
-	}
-
-	if (DelegateToRemoveFromNotifierB.IsValid() && NotifierB.IsValid())
-	{
-		check(NotifierB.Pin()->Delegate().Remove(DelegateToRemoveFromNotifierB));
-	}
-}
-
-
 
 void USkeletalMeshEditorContextObject::Register(UInteractiveToolManager* InToolManager)
 {
@@ -113,7 +72,6 @@ void USkeletalMeshEditorContextObject::Unregister(UInteractiveToolManager* InToo
 	InToolManager->GetContextObjectStore()->RemoveContextObject(this);
 
 	EditorBindings.Reset();
-	TreeBindings.Reset();
 	
 	bRegistered = false;
 }
@@ -123,7 +81,6 @@ void USkeletalMeshEditorContextObject::Init(USkeletalMeshModelingToolsEditorMode
 	EditorMode = InEditorMode;
 	
 	EditorBindings.Reset();
-	TreeBindings.Reset();
 }
 
 EMeshLODIdentifier USkeletalMeshEditorContextObject::GetEditingLOD()
@@ -175,7 +132,7 @@ void USkeletalMeshEditorContextObject::BindTo(ISkeletalMeshEditingInterface* InE
 	}
 	
 	BindEditor(InEditingInterface);
-	BindRefSkeletonTree(InEditingInterface);
+	EditorMode->BindToolSkeletonTree(InEditingInterface);
 }
 
 void USkeletalMeshEditorContextObject::UnbindFrom(ISkeletalMeshEditingInterface* InEditingInterface)
@@ -185,8 +142,8 @@ void USkeletalMeshEditorContextObject::UnbindFrom(ISkeletalMeshEditingInterface*
 		return;
 	}
 
+	EditorMode->UnbindToolSkeletonTree();
 	UnbindEditor(InEditingInterface);
-	UnbindRefSkeletonTree(InEditingInterface);
 }
 
 
@@ -196,7 +153,7 @@ USkeletalMeshEditorContextObject::FBindData USkeletalMeshEditorContextObject::Bi
 	TSharedPtr<ISkeletalMeshNotifier> InOtherNotifier)
 {
 	FBindData BindData;
-	BindData.BindScope.Reset(new UE::SkeletalMeshEditorUtils::FSkeletalMeshNotifierBindScope(InInterface->GetNotifier(), InOtherNotifier));
+	BindData.BindScope.Reset(new FSkeletalMeshNotifierBindScope(InInterface->GetNotifier(), InOtherNotifier));
 	
 	return BindData;
 
@@ -241,110 +198,6 @@ void USkeletalMeshEditorContextObject::UnbindEditor(ISkeletalMeshEditingInterfac
 	}
 }
 
-void USkeletalMeshEditorContextObject::BindRefSkeletonTree(ISkeletalMeshEditingInterface* InEditingInterface)
-{
-	if (!InEditingInterface || !EditorMode.IsValid())
-	{
-		return;
-	}
-	
-	const TWeakObjectPtr<USkeletonModifier> Modifier = InEditingInterface->GetModifier();
-	if (!Modifier.IsValid())
-	{
-		return;
-	}
-	
-	const TSharedPtr<FTabManager> TabManager = EditorMode->GetAssociatedTabManager();
-	TSharedPtr<SDockTab> SkeletonTab = TabManager->FindExistingLiveTab(GetSkeletonTreeTabId());
-	if (SkeletonTab.IsValid())
-	{
-		DefaultSkeletonWidget = SkeletonTab->GetContent();
-	}
-	else
-	{
-		SkeletonTab = TabManager->TryInvokeTab(GetSkeletonTreeTabId());
-	}
-
-	check(SkeletonTab.IsValid());
-
-	SAssignNew(RefSkeletonWidget, SBorder)
-		.Padding(4.f)
-		.BorderImage(FAppStyle::GetBrush("ToolPanel.GroupBorder"))
-	[
-		SAssignNew(RefSkeletonTree, SReferenceSkeletonTree)
-			.Modifier(Modifier)
-	];
-
-	// switch between SSkeletonTree and SReferenceSkeletonTree
-	SkeletonTab->SetContent(RefSkeletonWidget.ToSharedRef());
-
-	TreeBindings.Emplace(InEditingInterface, BindInterfaceTo(InEditingInterface, RefSkeletonTree->GetNotifier()));
-
-	// unbind notifications on the initial SSkeletonTree to avoid useless notifications
-	if(const TSharedPtr<ISkeletalMeshEditorBinding> Binding = EditorMode->GetModeBinding())
-	{
-		const TArray<FName> SelectedBones = Binding->GetSelectedBones();
-		
-		if (FBindData* EditorBindData = EditorBindings.Find(InEditingInterface))
-		{
-			UnbindInterfaceFrom(*EditorBindData);
-			EditorBindings.Remove(InEditingInterface);
-			
-			Binding->GetNotifier()->HandleNotification({}, ESkeletalMeshNotifyType::BonesSelected);
-		}
-
-		RefSkeletonTree->GetNotifier()->HandleNotification(SelectedBones, ESkeletalMeshNotifyType::BonesSelected);
-	}
-}
-
-void USkeletalMeshEditorContextObject::UnbindRefSkeletonTree(ISkeletalMeshEditingInterface* InEditingInterface)
-{
-	if (FBindData* BindData = TreeBindings.Find(InEditingInterface))
-	{
-		TArray<FName> SelectedBones;
-		if (RefSkeletonTree.IsValid())
-		{
-			RefSkeletonTree->GetSelectedBoneNames(SelectedBones);
-			UnbindInterfaceFrom(*BindData);
-		}
-
-		if (EditorMode.IsValid())
-		{
-			const TSharedPtr<FTabManager> TabManager = EditorMode->GetAssociatedTabManager();
-			if (TabManager.IsValid())
-			{
-				const TSharedPtr<SDockTab> SkeletonTab = TabManager->FindExistingLiveTab(GetSkeletonTreeTabId());
-				if (SkeletonTab.IsValid())
-				{
-					if (SkeletonTab->GetContent() == RefSkeletonWidget)
-					{
-						SkeletonTab->SetContent(DefaultSkeletonWidget.IsValid() ? DefaultSkeletonWidget.ToSharedRef() : SNullWidget::NullWidget);
-					}
-				}	
-			}
-
-			// update the default skeleton widget with the new RefSkeleton and set current selection 
-			if (const TSharedPtr<ISkeletalMeshEditorBinding> Binding = EditorMode->GetModeBinding())
-			{
-				TSharedPtr<ISkeletalMeshNotifier> Notifier = Binding->GetNotifier();
-				Notifier->HandleNotification(SelectedBones, ESkeletalMeshNotifyType::HierarchyChanged);
-				Notifier->HandleNotification(SelectedBones, ESkeletalMeshNotifyType::BonesSelected);
-			}
-		}
-		
-		RefSkeletonTree.Reset();
-		RefSkeletonWidget.Reset();
-		
-		TreeBindings.Remove(InEditingInterface);
-	}
-}
-
-const FName& USkeletalMeshEditorContextObject::GetSkeletonTreeTabId()
-{
-	static const FName SkeletonTreeId(TEXT("SkeletonTreeView"));
-	return SkeletonTreeId;
-}
-
 TSharedPtr<ISkeletalMeshEditorBinding> USkeletalMeshEditorContextObject::GetBinding() const
 {
 	return EditorMode.IsValid() ? EditorMode->GetModeBinding() : nullptr;
diff --git a/Engine/Plugins/Animation/SkeletalMeshModelingTools/Source/SkeletalMeshModelingTools/Private/SkeletalMeshEditorUtils.h b/Engine/Plugins/Animation/SkeletalMeshModelingTools/Source/SkeletalMeshModelingTools/Private/SkeletalMeshEditorUtils.h
--- a/Engine/Plugins/Animation/SkeletalMeshModelingTools/Source/SkeletalMeshModelingTools/Private/SkeletalMeshEditorUtils.h
+++ b/Engine/Plugins/Animation/SkeletalMeshModelingTools/Source/SkeletalMeshModelingTools/Private/SkeletalMeshEditorUtils.h
@@ -28,18 +28,6 @@ namespace SkeletalMeshEditorUtils
 
 	UE_API USkeletalMeshEditorContextObject* GetEditorContextObject(UInteractiveToolsContext* InToolsContext);
 
-
-	struct FSkeletalMeshNotifierBindScope 
-	{
-		FSkeletalMeshNotifierBindScope(TWeakPtr<ISkeletalMeshNotifier> InNotifierA,  TWeakPtr<ISkeletalMeshNotifier> InNotifierB);
-		~FSkeletalMeshNotifierBindScope();
-		
-		TWeakPtr<ISkeletalMeshNotifier> NotifierA;
-		TWeakPtr<ISkeletalMeshNotifier> NotifierB;
-
-		FDelegateHandle	DelegateToRemoveFromNotifierA;	
-		FDelegateHandle	DelegateToRemoveFromNotifierB;	
-	};
 }
 	
 }
@@ -77,17 +65,13 @@ public:
 private:
 	TWeakObjectPtr<USkeletalMeshModelingToolsEditorMode> EditorMode = nullptr;
 
-	UPROPERTY()
-	ESkeletonDrawMode SkeletonDrawMode = ESkeletonDrawMode::Default;
-
 	UE_API TSharedPtr<ISkeletalMeshEditorBinding> GetBinding() const;
 
 	struct FBindData
 	{
-		TUniquePtr<UE::SkeletalMeshEditorUtils::FSkeletalMeshNotifierBindScope> BindScope;
+		TUniquePtr<FSkeletalMeshNotifierBindScope> BindScope;
 	};
 	TMap< ISkeletalMeshEditingInterface*, FBindData > EditorBindings;
-	TMap< ISkeletalMeshEditingInterface*, FBindData > TreeBindings;
 
 	static UE_API FBindData BindInterfaceTo(
 		ISkeletalMeshEditingInterface* InInterface,
@@ -99,14 +83,6 @@ private:
 	UE_API void BindEditor(ISkeletalMeshEditingInterface* InEditingInterface);
 	UE_API void UnbindEditor(ISkeletalMeshEditingInterface* InEditingInterface);
 
-	UE_API void BindRefSkeletonTree(ISkeletalMeshEditingInterface* InEditingInterface);
-	UE_API void UnbindRefSkeletonTree(ISkeletalMeshEditingInterface* InEditingInterface);
-
-	static UE_API const FName& GetSkeletonTreeTabId();
-
-	TSharedPtr<SWidget> RefSkeletonWidget;
-	TSharedPtr<SReferenceSkeletonTree> RefSkeletonTree;
-	TSharedPtr<SWidget> DefaultSkeletonWidget;
 	
 	bool bRegistered = false;
 };
diff --git a/Engine/Plugins/Animation/SkeletalMeshModelingTools/Source/SkeletalMeshModelingTools/Private/SkeletalMeshModelingToolsCommands.cpp b/Engine/Plugins/Animation/SkeletalMeshModelingTools/Source/SkeletalMeshModelingTools/Private/SkeletalMeshModelingToolsCommands.cpp
--- a/Engine/Plugins/Animation/SkeletalMeshModelingTools/Source/SkeletalMeshModelingTools/Private/SkeletalMeshModelingToolsCommands.cpp
+++ b/Engine/Plugins/Animation/SkeletalMeshModelingTools/Source/SkeletalMeshModelingTools/Private/SkeletalMeshModelingToolsCommands.cpp
@@ -11,19 +11,6 @@ void FSkeletalMeshModelingToolsCommands::RegisterCommands()
 {
 	UI_COMMAND(ToggleEditingToolsMode, "Enable Editing Tools", "Toggles editing tools on or off.", EUserInterfaceActionType::ToggleButton, FInputChord());
 	
-	UI_COMMAND(NewBone, "New Bone", "Add new Bone.", EUserInterfaceActionType::Button, FInputChord(EKeys::N, EModifierKey::Control));
-	UI_COMMAND(RemoveBone, "Remove Bone", "Remove selected bone(s).", EUserInterfaceActionType::Button, FInputChord(EKeys::Delete), FInputChord(EKeys::BackSpace));
-	UI_COMMAND(UnParentBone, "Unparent Bone", "Unparent selected bone(s).", EUserInterfaceActionType::Button, FInputChord(EKeys::P, EModifierKey::Shift));
-
-	UI_COMMAND(RenameBone, "Rename Bone", "Rename the selected bone.", EUserInterfaceActionType::Button, FInputChord(EKeys::F2));
-	UI_COMMAND(SearchAndReplace, "Search and Replace", "Search and Replace the selected bone(s) name(s).", EUserInterfaceActionType::Button, FInputChord());
-	UI_COMMAND(AddPrefix, "Add Prefix", "Add a prefix to the selected bone(s) name(s).", EUserInterfaceActionType::Button, FInputChord());
-	UI_COMMAND(AddSuffix, "Add Suffix", "Add a suffix to the selected bone(s) name(s).", EUserInterfaceActionType::Button, FInputChord());
-	
-	UI_COMMAND(CopyBones, "Copy Bone(s)", "Copy selected bone(s).", EUserInterfaceActionType::Button, FInputChord(EKeys::C, EModifierKey::Control));
-	UI_COMMAND(PasteBones, "Paste Bone(s)", "Paste selected bone(s).", EUserInterfaceActionType::Button, FInputChord(EKeys::V, EModifierKey::Control));
-	UI_COMMAND(DuplicateBones, "Duplicate Bone(s)", "Duplicate selected bone(s).", EUserInterfaceActionType::Button, FInputChord(EKeys::D, EModifierKey::Control));
-	
 	UI_COMMAND(NewMorphTarget, "New Morph Target", "Add new Morph Target.", EUserInterfaceActionType::Button, FInputChord());
 }
 
diff --git a/Engine/Plugins/Animation/SkeletalMeshModelingTools/Source/SkeletalMeshModelingTools/Private/SkeletalMeshModelingToolsCommands.h b/Engine/Plugins/Animation/SkeletalMeshModelingTools/Source/SkeletalMeshModelingTools/Private/SkeletalMeshModelingToolsCommands.h
--- a/Engine/Plugins/Animation/SkeletalMeshModelingTools/Source/SkeletalMeshModelingTools/Private/SkeletalMeshModelingToolsCommands.h
+++ b/Engine/Plugins/Animation/SkeletalMeshModelingTools/Source/SkeletalMeshModelingTools/Private/SkeletalMeshModelingToolsCommands.h
@@ -26,20 +26,6 @@ public:
 	// Editing tools commands
 	TSharedPtr<FUICommandInfo> ToggleEditingToolsMode;
 
-	// Skeleton Editing commands
-	TSharedPtr<FUICommandInfo> NewBone;
-	TSharedPtr<FUICommandInfo> RemoveBone;
-	TSharedPtr<FUICommandInfo> UnParentBone;
-	TSharedPtr<FUICommandInfo> RenameBone;
-	TSharedPtr<FUICommandInfo> SearchAndReplace;
-	TSharedPtr<FUICommandInfo> AddPrefix;
-	TSharedPtr<FUICommandInfo> AddSuffix;
-
-	// Copy & Paste commands
-	TSharedPtr<FUICommandInfo> CopyBones;
-	TSharedPtr<FUICommandInfo> PasteBones;
-	TSharedPtr<FUICommandInfo> DuplicateBones;
-
 	// Morph Target Editing commands
 	TSharedPtr<FUICommandInfo> NewMorphTarget;
 };
diff --git a/Engine/Plugins/Animation/SkeletalMeshModelingTools/Source/SkeletalMeshModelingTools/Private/SkeletalMeshModelingToolsEditorMode.cpp b/Engine/Plugins/Animation/SkeletalMeshModelingTools/Source/SkeletalMeshModelingTools/Private/SkeletalMeshModelingToolsEditorMode.cpp
--- a/Engine/Plugins/Animation/SkeletalMeshModelingTools/Source/SkeletalMeshModelingTools/Private/SkeletalMeshModelingToolsEditorMode.cpp
+++ b/Engine/Plugins/Animation/SkeletalMeshModelingTools/Source/SkeletalMeshModelingTools/Private/SkeletalMeshModelingToolsEditorMode.cpp
@@ -205,6 +205,13 @@ void FSkeletalMeshModelingToolsEditorModeNotifier::HandleNotification(const TArr
 		return;
 	}
 
+	switch (InNotifyType)
+	{
+	case ESkeletalMeshNotifyType::BonesSelected:
+		EditorMode->SetSelectedBones(BoneNames);
+		break;
+	}
+	
 	Notify(BoneNames, InNotifyType);
 }
 
@@ -436,6 +443,33 @@ TSharedPtr<FTabManager> USkeletalMeshModelingToolsEditorMode::GetAssociatedTabMa
 	return nullptr;
 }
 
+void USkeletalMeshModelingToolsEditorMode::BindToolSkeletonTree(ISkeletalMeshEditingInterface* InEditingInterface)
+{
+	if (InEditingInterface->CanModifySkeleton())
+	{
+		// If the tool owns/modifies a skeleton, it should directly interface with the tool skeleton tree, so unbind notifier for the editing cache skeleton
+		EditingCacheNotifierBindScope.Reset();
+		GetCurrentEditingCache()->GetNotifier()->HandleNotification({}, ESkeletalMeshNotifyType::BonesSelected);
+	}
+	
+	TypedToolkit.Pin()->BindToolSkeletonTree(InEditingInterface);
+}
+
+void USkeletalMeshModelingToolsEditorMode::UnbindToolSkeletonTree()
+{
+	TypedToolkit.Pin()->UnbindToolSkeletonTree();
+	if (!EditingCacheNotifierBindScope.IsValid())
+	{
+		EditingCacheNotifierBindScope.Reset(
+			new FSkeletalMeshNotifierBindScope(
+				GetModeBinding()->GetNotifier(),
+				GetCurrentEditingCache()->GetNotifier()));
+
+		GetCurrentEditingCache()->GetNotifier()->HandleNotification({}, ESkeletalMeshNotifyType::HierarchyChanged);
+		GetCurrentEditingCache()->GetNotifier()->HandleNotification(GetSelectedBones(), ESkeletalMeshNotifyType::BonesSelected);
+	}
+}
+
 const TArray<FTransform>& USkeletalMeshModelingToolsEditorMode::GetComponentSpaceBoneTransforms() const
 {
 	return GetCurrentEditingCache()->GetComponentSpaceBoneTransforms();
@@ -451,9 +485,14 @@ USkeletonModifier* USkeletalMeshModelingToolsEditorMode::GetSkeletonReader()
 	return SkeletonReader;
 }
 
+void USkeletalMeshModelingToolsEditorMode::SetSelectedBones(const TArray<FName>& InSelectedBones)
+{
+	SelectedBones = InSelectedBones;
+}
+
 TArray<FName> USkeletalMeshModelingToolsEditorMode::GetSelectedBones() const 
 {
-	return GetCurrentEditingCache()->GetSelectedBones();
+	return SelectedBones;
 }
 
 
@@ -1062,6 +1101,7 @@ void USkeletalMeshModelingToolsEditorMode::SetEditorBinding(const TWeakPtr<ISkel
 	SkeletalMesh->GetOnMeshChanged().AddUObject(this, &USkeletalMeshModelingToolsEditorMode::HandleSkeletalMeshChanged);
 	SkeletonReader = NewObject<USkeletonModifier>();
 	SkeletonReader->SetSkeletalMesh(SkeletalMesh.Get());
+	SkeletonReader->SetReadOnly(true);
 	
 	if (USkeletalMeshEditorContextObject* ContextObject = UE::SkeletalMeshEditorUtils::GetEditorContextObject(GetInteractiveToolsContext()))
 	{
@@ -1166,11 +1206,19 @@ void USkeletalMeshModelingToolsEditorMode::RecreateEditingCache(EMeshLODIdentifi
 	USkeletalMeshEditingCache::FDelegates Delegates;
 	Delegates.ToggleSkeletalMeshBoneManipulationDelegate.BindUObject(this, &USkeletalMeshModelingToolsEditorMode::ToggleSkeletalMeshBoneManipulation);
 	Delegates.IsSkeletalMeshBoneManipulationEnabledDelegate.BindUObject(this, &USkeletalMeshModelingToolsEditorMode::IsSkeletalMeshBoneManipulationEnabled);
+	Delegates.OnGetSkeletalMeshSkeletonNotifierDelegate.BindLambda([this]()
+		{
+			// Skeletal Mesh Skeleton bone selection state is governed by the skeleton tree in the base editor
+			return GetEditorBinding()->GetNotifier();
+		});
 	CurrentEditingCache->Spawn(PreviewWorld, GetSkelMeshComponent(), InLOD, Delegates);
 
 
 	using namespace UE::SkeletalMeshEditorUtils;
 	EditingCacheNotifierBindScope.Reset(new FSkeletalMeshNotifierBindScope(GetModeBinding()->GetNotifier(), CurrentEditingCache->GetNotifier()));
+
+	CurrentEditingCache->GetNotifier()->HandleNotification({}, ESkeletalMeshNotifyType::HierarchyChanged);
+	CurrentEditingCache->GetNotifier()->HandleNotification(GetSelectedBones(), ESkeletalMeshNotifyType::BonesSelected);
 }
 
 USkeletalMeshEditingCache* USkeletalMeshModelingToolsEditorMode::GetCurrentEditingCache() const
diff --git a/Engine/Plugins/Animation/SkeletalMeshModelingTools/Source/SkeletalMeshModelingTools/Private/SkeletalMeshModelingToolsEditorMode.h b/Engine/Plugins/Animation/SkeletalMeshModelingTools/Source/SkeletalMeshModelingTools/Private/SkeletalMeshModelingToolsEditorMode.h
--- a/Engine/Plugins/Animation/SkeletalMeshModelingTools/Source/SkeletalMeshModelingTools/Private/SkeletalMeshModelingToolsEditorMode.h
+++ b/Engine/Plugins/Animation/SkeletalMeshModelingTools/Source/SkeletalMeshModelingTools/Private/SkeletalMeshModelingToolsEditorMode.h
@@ -15,11 +15,6 @@
 
 class FSkeletalMeshModelingToolsEditorModeToolkit;
 
-namespace UE::SkeletalMeshEditorUtils
-{
-	struct FSkeletalMeshNotifierBindScope;
-}
-
 class USkeletalMeshEditingCache;
 //class FStylusStateTracker;
 class UEdModeInteractiveToolsContext;
@@ -147,6 +142,7 @@ public:
 	const TArray<FTransform>& GetComponentSpaceBoneTransforms() const;
 	void ToggleBoneManipulation(bool bEnable);
 	USkeletonModifier* GetSkeletonReader();
+	void SetSelectedBones(const TArray<FName>& InSelectedBones);
 	TArray<FName> GetSelectedBones() const;
 	
 	FName GetEditingMorphTarget();
@@ -165,6 +161,10 @@ public:
 	void HandleSetMorphTargetAutoFill(FName MorphTarget, bool bAutoFill, float PreviousOverrideWeight);
 
 	TSharedPtr<FTabManager> GetAssociatedTabManager();
+
+	void BindToolSkeletonTree(ISkeletalMeshEditingInterface* InEditingInterface);
+	void UnbindToolSkeletonTree();
+	
 	FSimpleMulticastDelegate& OnInitialized() { return OnInitializedDelegate; }
 
 
@@ -182,15 +182,17 @@ private:
 	TWeakPtr<ISkeletalMeshEditorBinding> EditorBinding;
 	TWeakObjectPtr<USkeletalMesh> SkeletalMesh;
 
-	TUniquePtr<UE::SkeletalMeshEditorUtils::FSkeletalMeshNotifierBindScope> EditorBindScope;
-	
 	UPROPERTY()
 	TObjectPtr<USkeletalMeshEditingCache> CurrentEditingCache = nullptr;
 
 	void RecreateEditingCache(EMeshLODIdentifier InLOD);
 
-	TUniquePtr<UE::SkeletalMeshEditorUtils::FSkeletalMeshNotifierBindScope> EditingCacheNotifierBindScope;
+	TUniquePtr<FSkeletalMeshNotifierBindScope> EditingCacheNotifierBindScope;
 
+	// Represents the collection of bones currently selected in the viewport for whichever skeleton that is currently active
+	// Helps maintain consistent bone selection state as different skeleton tree widget activates
+	TArray<FName> SelectedBones;
+	
 	FName EditingMorphTarget = NAME_None;
 	
 	// Stylus support is currently disabled; this is left in for reference if/when it is brought back
diff --git a/Engine/Plugins/Animation/SkeletalMeshModelingTools/Source/SkeletalMeshModelingTools/Private/SkeletalMeshModelingToolsEditorModeToolkit.cpp b/Engine/Plugins/Animation/SkeletalMeshModelingTools/Source/SkeletalMeshModelingTools/Private/SkeletalMeshModelingToolsEditorModeToolkit.cpp
--- a/Engine/Plugins/Animation/SkeletalMeshModelingTools/Source/SkeletalMeshModelingTools/Private/SkeletalMeshModelingToolsEditorModeToolkit.cpp
+++ b/Engine/Plugins/Animation/SkeletalMeshModelingTools/Source/SkeletalMeshModelingTools/Private/SkeletalMeshModelingToolsEditorModeToolkit.cpp
@@ -30,7 +30,7 @@
 #include "BaseTools/MeshSurfacePointTool.h"
 #include "SPrimaryButton.h"
 #include "SkeletonModifier.h"
-#include "SReferenceSkeletonTree.h"
+#include "SkeletalMesh/SReferenceSkeletonTree.h"
 #include "Components/SKMBackedDynaMeshComponent.h"
 #include "Editor/EditorWidgets/Public/SEnumCombo.h"
 #include "Toolkits/AssetEditorModeUILayer.h"
@@ -217,7 +217,7 @@ void FSkeletalMeshModelingToolsEditorModeToolkit::InvokeUI()
 void FSkeletalMeshModelingToolsEditorModeToolkit::ShutdownUI()
 {
 	//Skeleton
-	SkeletonNotifierBindScope.Reset();
+	DynamicMeshSkeletonTreeNotifierBindScope.Reset();
 	
 	if (ModeUILayer.IsValid())
 	{
@@ -245,7 +245,7 @@ void FSkeletalMeshModelingToolsEditorModeToolkit::ShutdownUI()
 TArray<FName> FSkeletalMeshModelingToolsEditorModeToolkit::GetSelectedBonesForDynamicMeshSkeleton()
 {
 	TArray<FName> SelectedBones;
-	RefSkeletonTree->GetSelectedBoneNames(SelectedBones);
+	DynamicMeshSkeletonTree->GetSelectedBoneNames(SelectedBones);
 	return SelectedBones;
 }
 
@@ -541,15 +541,32 @@ void FSkeletalMeshModelingToolsEditorModeToolkit::OnEditorModeInitialized()
 	int32 EnumIndex = ApplyModeEnums.IndexOfByKey(EditorMode->GetApplyMode());
 	ApplyModeComboBox->SetSelectedItem(ApplyModeStrings[EnumIndex]);
 
-	SAssignNew(RefSkeletonWidget, SBorder)
+	// Create a Skeleton Tree for dynamic mesh skeleton
+	SAssignNew(DynamicMeshSkeletonWidget, SBorder)
 	.Padding(4.f)
 	.BorderImage(FAppStyle::GetBrush("ToolPanel.GroupBorder"))
 	[
-		SAssignNew(RefSkeletonTree, SReferenceSkeletonTree)
+		SAssignNew(DynamicMeshSkeletonTree, SReferenceSkeletonTree)
 			.Modifier(EditorMode->GetSkeletonReader())
-			.bIsReadOnly(true)
 	];
 
+	// Save the default Skeleton Tree
+	if (ensure(ModeUILayer.IsValid()))
+	{
+		TSharedPtr<FAssetEditorModeUILayer> ModeUILayerPtr = ModeUILayer.Pin();
+		TSharedPtr<FTabManager> TabManager = ModeUILayerPtr->GetTabManager();
+
+
+		const FName& SkeletonTreeId = FPersonaTabs::SkeletonTreeViewID;
+		TSharedPtr<SDockTab> SkeletonTab = TabManager->FindExistingLiveTab(SkeletonTreeId);
+		if (!SkeletonTab.IsValid())
+		{
+			SkeletonTab = TabManager->TryInvokeTab(SkeletonTreeId);
+		}
+		DefaultSkeletonWidget = SkeletonTab->GetContent();
+	}
+	
+
 	SMorphTargetManager::FMorphTargetManagerDelegates Delegates;
 
 	Delegates.OnGetMorphTargets.BindUObject(EditorMode.Get(), &USkeletalMeshModelingToolsEditorMode::GetMorphTargets);
@@ -572,14 +589,6 @@ void FSkeletalMeshModelingToolsEditorModeToolkit::OnEditorModeInitialized()
 			.Delegates(Delegates);
 	MorphTargetManagerWidget->SetContent(MorphTargetManager.ToSharedRef());
 
-	// Make sure viewport selection is routed to the default skeleton tree by default
-	SkeletonNotifierBindScope.Reset(
-		new UE::SkeletalMeshEditorUtils::FSkeletalMeshNotifierBindScope(
-			EditorMode->GetModeBinding()->GetNotifier(),
-			EditorMode->GetEditorBinding()->GetNotifier()));
-
-
-	
 }
 
 void FSkeletalMeshModelingToolsEditorModeToolkit::ShowDynamicMeshSkeletonTree()
@@ -589,49 +598,36 @@ void FSkeletalMeshModelingToolsEditorModeToolkit::ShowDynamicMeshSkeletonTree()
 		TSharedPtr<FAssetEditorModeUILayer> ModeUILayerPtr = ModeUILayer.Pin();
 		TSharedPtr<FTabManager> TabManager = ModeUILayerPtr->GetTabManager();
 
-		// Skeleton Tree
+		const FName& SkeletonTreeId = FPersonaTabs::SkeletonTreeViewID;
+		TSharedPtr<SDockTab> SkeletonTab = TabManager->FindExistingLiveTab(SkeletonTreeId);
+		if (!SkeletonTab.IsValid())
 		{
-			const FName& SkeletonTreeId = FPersonaTabs::SkeletonTreeViewID;
-			TSharedPtr<SDockTab> SkeletonTab = TabManager->FindExistingLiveTab(SkeletonTreeId);
-			if (!SkeletonTab.IsValid())
-			{
-				SkeletonTab = TabManager->TryInvokeTab(SkeletonTreeId);
-			}
-		
-			if (SkeletonTab->GetContent() != RefSkeletonWidget)
-			{
-				DefaultSkeletonWidget = SkeletonTab->GetContent();
-			
-				// switch between SSkeletonTree and SReferenceSkeletonTree
-				SkeletonTab->SetContent(RefSkeletonWidget.ToSharedRef());
-
-				
-				// Make sure viewport selection is routed to the ref skeleton tree
-				SkeletonNotifierBindScope.Reset(
-					new UE::SkeletalMeshEditorUtils::FSkeletalMeshNotifierBindScope(
-						EditorMode->GetModeBinding()->GetNotifier(),
-						RefSkeletonTree->GetNotifier()));
-				
-				TArray<FName> SelectedBones = EditorMode->GetSelectedBones();
-				
-				EditorMode->GetEditorBinding()->GetNotifier()->HandleNotification({}, ESkeletalMeshNotifyType::BonesSelected);
-				RefSkeletonTree->GetNotifier()->HandleNotification({}, ESkeletalMeshNotifyType::HierarchyChanged);
-				RefSkeletonTree->GetNotifier()->HandleNotification(SelectedBones, ESkeletalMeshNotifyType::BonesSelected);
-
-
-				const TSharedRef<FUICommandList>& CommandList = EditorMode->GetEditor()->GetToolkitCommands();
-				DefaultResetBoneTransformsAction = *(CommandList->GetActionForCommand(EditorMode->GetEditor()->GetResetBoneTransformsCommand()));
-				DefaultResetAllBoneTransformsAction = *(CommandList->GetActionForCommand(EditorMode->GetEditor()->GetResetAllBonesTransformsCommand()));
-
-				// Reset transforms
-				static constexpr bool bSelectedOnly = true;
-				CommandList->MapAction(EditorMode->GetEditor()->GetResetBoneTransformsCommand(),
-					FExecuteAction::CreateSP(this, &FSkeletalMeshModelingToolsEditorModeToolkit::ResetDynamicMeshBoneTransforms ,bSelectedOnly));
-	
-				CommandList->MapAction(EditorMode->GetEditor()->GetResetAllBonesTransformsCommand(),
-					FExecuteAction::CreateSP(this, &FSkeletalMeshModelingToolsEditorModeToolkit::ResetDynamicMeshBoneTransforms ,!bSelectedOnly));
-			}
+			SkeletonTab = TabManager->TryInvokeTab(SkeletonTreeId);
 		}
+
+		SkeletonTab->SetContent(DynamicMeshSkeletonWidget.ToSharedRef());
+
+		// Make sure viewport selection is routed to the dynamic mesh skeleton tree
+		DynamicMeshSkeletonTreeNotifierBindScope.Reset(
+			new FSkeletalMeshNotifierBindScope(
+				EditorMode->GetModeBinding()->GetNotifier(),
+				DynamicMeshSkeletonTree->GetNotifier()));
+
+		TArray<FName> SelectedBones = EditorMode->GetSelectedBones();
+		DynamicMeshSkeletonTree->GetNotifier()->HandleNotification({}, ESkeletalMeshNotifyType::HierarchyChanged);
+		DynamicMeshSkeletonTree->GetNotifier()->HandleNotification(SelectedBones, ESkeletalMeshNotifyType::BonesSelected);
+
+		const TSharedRef<FUICommandList>& CommandList = EditorMode->GetEditor()->GetToolkitCommands();
+		DefaultResetBoneTransformsAction = *(CommandList->GetActionForCommand(EditorMode->GetEditor()->GetResetBoneTransformsCommand()));
+		DefaultResetAllBoneTransformsAction = *(CommandList->GetActionForCommand(EditorMode->GetEditor()->GetResetAllBonesTransformsCommand()));
+
+		// Reset transforms
+		static constexpr bool bSelectedOnly = true;
+		CommandList->MapAction(EditorMode->GetEditor()->GetResetBoneTransformsCommand(),
+			FExecuteAction::CreateSP(this, &FSkeletalMeshModelingToolsEditorModeToolkit::ResetDynamicMeshBoneTransforms, bSelectedOnly));
+
+		CommandList->MapAction(EditorMode->GetEditor()->GetResetAllBonesTransformsCommand(),
+			FExecuteAction::CreateSP(this, &FSkeletalMeshModelingToolsEditorModeToolkit::ResetDynamicMeshBoneTransforms, !bSelectedOnly));
 	}
 }
 
@@ -642,45 +638,78 @@ void FSkeletalMeshModelingToolsEditorModeToolkit::ShowSkeletalMeshSkeletonTree()
 		TSharedPtr<FAssetEditorModeUILayer> ModeUILayerPtr = ModeUILayer.Pin();
 		TSharedPtr<FTabManager> TabManager = ModeUILayerPtr->GetTabManager();
 
-		// Skeleton Tree
+		const FName& SkeletonTreeViewID = FPersonaTabs::SkeletonTreeViewID;
+		TSharedPtr<SDockTab> SkeletonTab = TabManager->FindExistingLiveTab(SkeletonTreeViewID);
+		if (!SkeletonTab.IsValid())
 		{
-			const FName& SkeletonTreeViewID = FPersonaTabs::SkeletonTreeViewID;
-			TSharedPtr<SDockTab> SkeletonTab = TabManager->FindExistingLiveTab(SkeletonTreeViewID);
-			if (!SkeletonTab.IsValid())
-			{
-				SkeletonTab = TabManager->TryInvokeTab(SkeletonTreeViewID);
-			}
-
-			// switch back to default widget if it was replaced before
-			if (DefaultSkeletonWidget)
-			{
-				SkeletonTab->SetContent(DefaultSkeletonWidget.ToSharedRef());
-			}
+			SkeletonTab = TabManager->TryInvokeTab(SkeletonTreeViewID);
 		}
 
-		// Make sure viewport selection is routed to the default skeleton tree
-		SkeletonNotifierBindScope.Reset(
-			new UE::SkeletalMeshEditorUtils::FSkeletalMeshNotifierBindScope(
-				EditorMode->GetModeBinding()->GetNotifier(),
-				EditorMode->GetEditorBinding()->GetNotifier()));
+		SkeletonTab->SetContent(DefaultSkeletonWidget.ToSharedRef());
 
-		TArray<FName> SelectedBones = EditorMode->GetSelectedBones();
-		RefSkeletonTree->GetNotifier()->HandleNotification({}, ESkeletalMeshNotifyType::BonesSelected);
-		EditorMode->GetEditorBinding()->GetNotifier()->HandleNotification({}, ESkeletalMeshNotifyType::HierarchyChanged);
-		EditorMode->GetEditorBinding()->GetNotifier()->HandleNotification(SelectedBones, ESkeletalMeshNotifyType::BonesSelected);
+		DynamicMeshSkeletonTreeNotifierBindScope.Reset();
 
-		// Rebind commands it was replaced before
-		if (DefaultSkeletonWidget)
-		{
-			const TSharedRef<FUICommandList>& CommandList = EditorMode->GetEditor()->GetToolkitCommands();
+		// Rebind commands if it was replaced before
+		const TSharedRef<FUICommandList>& CommandList = EditorMode->GetEditor()->GetToolkitCommands();
 
-			// Reset transforms
-			CommandList->MapAction(EditorMode->GetEditor()->GetResetBoneTransformsCommand(), DefaultResetBoneTransformsAction);
-			CommandList->MapAction(EditorMode->GetEditor()->GetResetAllBonesTransformsCommand(), DefaultResetAllBoneTransformsAction);
-		}
+		// Reset transforms
+		CommandList->MapAction(EditorMode->GetEditor()->GetResetBoneTransformsCommand(), DefaultResetBoneTransformsAction);
+		CommandList->MapAction(EditorMode->GetEditor()->GetResetAllBonesTransformsCommand(), DefaultResetAllBoneTransformsAction);
 	}	
 }
 
+bool FSkeletalMeshModelingToolsEditorModeToolkit::BindToolSkeletonTree(ISkeletalMeshEditingInterface* InEditingInterface)
+{
+	bWasUsingDynamicMeshSkeletonBeforeToolActivation = EditorMode->GetCurrentEditingCache()->IsDynamicMeshSkeletonEnabled();
+
+	TSharedPtr<SWidget> ToolSkeletonTree = InEditingInterface->GetCustomSkeletonTreeWidget(EditorMode->GetModeBinding());
+
+	if (!ToolSkeletonTree.IsValid())
+	{
+		return false;
+	}
+	
+	if (!ModeUILayer.IsValid())
+	{
+		return false;
+	}
+	
+	TSharedPtr<FAssetEditorModeUILayer> ModeUILayerPtr = ModeUILayer.Pin();
+	TSharedPtr<FTabManager> TabManager = ModeUILayerPtr->GetTabManager();
+	
+	const FName& SkeletonTreeId = FPersonaTabs::SkeletonTreeViewID;
+	TSharedPtr<SDockTab> SkeletonTab = TabManager->FindExistingLiveTab(SkeletonTreeId);
+	if (!SkeletonTab.IsValid())
+	{
+		SkeletonTab = TabManager->TryInvokeTab(SkeletonTreeId);
+	}
+
+	SAssignNew(ToolSkeletonHost, SBorder)
+		.Padding(4.f)
+		.BorderImage(FAppStyle::GetBrush("ToolPanel.GroupBorder"))
+	[
+		ToolSkeletonTree.ToSharedRef()	
+	];
+
+	SkeletonTab->SetContent(ToolSkeletonHost.ToSharedRef());
+
+	return true;
+}
+
+void FSkeletalMeshModelingToolsEditorModeToolkit::UnbindToolSkeletonTree()
+{
+	if (bWasUsingDynamicMeshSkeletonBeforeToolActivation)
+	{
+		ShowDynamicMeshSkeletonTree();
+	}
+	else
+	{
+		ShowSkeletalMeshSkeletonTree();
+	}
+
+	ToolSkeletonHost.Reset();
+}
+
 void FSkeletalMeshModelingToolsEditorModeToolkit::RefreshMorphTargetManager()
 {
 	MorphTargetManager->RefreshList();
diff --git a/Engine/Plugins/Animation/SkeletalMeshModelingTools/Source/SkeletalMeshModelingTools/Private/SkeletalMeshModelingToolsEditorModeToolkit.h b/Engine/Plugins/Animation/SkeletalMeshModelingTools/Source/SkeletalMeshModelingTools/Private/SkeletalMeshModelingToolsEditorModeToolkit.h
--- a/Engine/Plugins/Animation/SkeletalMeshModelingTools/Source/SkeletalMeshModelingTools/Private/SkeletalMeshModelingToolsEditorModeToolkit.h
+++ b/Engine/Plugins/Animation/SkeletalMeshModelingTools/Source/SkeletalMeshModelingTools/Private/SkeletalMeshModelingToolsEditorModeToolkit.h
@@ -41,6 +41,10 @@ public:
 	void ShowDynamicMeshSkeletonTree();
 	void ShowSkeletalMeshSkeletonTree();
 
+	bool BindToolSkeletonTree(ISkeletalMeshEditingInterface* InTool);
+	void UnbindToolSkeletonTree();
+
+	
 	void RefreshMorphTargetManager();
 private:
 	void PostNotification(const FText& Message);
@@ -96,14 +100,18 @@ private:
 	TSharedPtr<STextComboBox> ApplyModeComboBox;
 	
 	TSharedPtr<SWidget> DefaultSkeletonWidget;
-	TSharedPtr<SBorder> RefSkeletonWidget;
-	TSharedPtr<SReferenceSkeletonTree> RefSkeletonTree;
+	TSharedPtr<SBorder> DynamicMeshSkeletonWidget;
+	TSharedPtr<SReferenceSkeletonTree> DynamicMeshSkeletonTree;
+
+	bool bWasUsingDynamicMeshSkeletonBeforeToolActivation = false;
+	TSharedPtr<SBorder> ToolSkeletonHost;
+	
 
 	TSharedPtr<SWidget> DefaultMorphTargetWidget; 
 	TSharedPtr<SBorder> MorphTargetManagerWidget; 
 	TSharedPtr<SMorphTargetManager> MorphTargetManager; 
 
-	TUniquePtr<UE::SkeletalMeshEditorUtils::FSkeletalMeshNotifierBindScope> SkeletonNotifierBindScope;
+	TUniquePtr<FSkeletalMeshNotifierBindScope> DynamicMeshSkeletonTreeNotifierBindScope;
 
 	TWeakObjectPtr<USkeletalMeshModelingToolsEditorMode> EditorMode;
 
diff --git a/Engine/Plugins/Runtime/GeometryProcessing/Source/DynamicMesh/Private/Operations/SmoothBoneWeights.cpp b/Engine/Plugins/Runtime/GeometryProcessing/Source/DynamicMesh/Private/Operations/SmoothBoneWeights.cpp
--- a/Engine/Plugins/Runtime/GeometryProcessing/Source/DynamicMesh/Private/Operations/SmoothBoneWeights.cpp
+++ b/Engine/Plugins/Runtime/GeometryProcessing/Source/DynamicMesh/Private/Operations/SmoothBoneWeights.cpp
@@ -198,6 +198,119 @@ bool TSmoothBoneWeights<BoneIndexType, BoneWeightType>::SmoothWeightsAtVertex(co
 	return true;
 }
 
+template <typename BoneIndexType, typename BoneWeightType>
+bool TSmoothBoneWeights<BoneIndexType, BoneWeightType>::SmoothWeightsAtVertex(
+	const int32 VertexID,
+	const BoneWeightType VertexFalloff,
+	const TSet<BoneIndexType>& LockedBones,
+	TMap<BoneIndexType, BoneWeightType>& OutFinalWeights)
+{
+	using namespace SmoothBoneWeightsLocals;
+	
+	// const UE::Geometry::FNonManifoldMappingSupport NonManifoldMappingSupport(*SourceMesh);
+
+	// for each vertex in the stamp...
+	constexpr int32 AvgNumNeighbors = 8;
+	using VertexNeighborWeights = TArray<BoneWeightType, TInlineAllocator<AvgNumNeighbors>>;
+	TArray<int32> AllNeighborVertices;
+	TMap<BoneIndexType, VertexNeighborWeights> WeightsOnAllNeighbors;
+	// const int32 SrcVertexID = NonManifoldMappingSupport.GetOriginalNonManifoldVertexID(VertexID);
+
+	TMap<BoneIndexType, BoneWeightType> LockedWeights;
+	BoneWeightType TotalLockedWeight = 0;
+	for (int32 Index = 0; Index < DataSource->GetBoneNum(VertexID); ++Index)
+	{
+		BoneIndexType BoneIndex = DataSource->GetBoneIndex(VertexID, Index);
+		if (LockedBones.Contains(BoneIndex))
+		{
+			const BoneWeightType Weight = DataSource->GetBoneWeight(VertexID, Index);
+			LockedWeights.Add(BoneIndex, Weight);
+			TotalLockedWeight += Weight;
+		}
+	}
+	
+	// get list of all neighboring vertices, AND this vertex
+	AllNeighborVertices.Add(VertexID);
+	for (const int32 NeighborVertexID : SourceMesh->VtxVerticesItr(VertexID))
+	{
+		AllNeighborVertices.Add(NeighborVertexID);
+	}
+
+	// get all weights above a given threshold across ALL neighbors (including self)
+	for (const int32 NeighborVertexID : AllNeighborVertices)
+	{
+		for (int32 Index = 0; Index < DataSource->GetBoneNum(NeighborVertexID); ++Index)
+		{
+			BoneIndexType BoneIndex = DataSource->GetBoneIndex(NeighborVertexID, Index);
+			if (!LockedBones.Contains(BoneIndex))
+			{
+				const BoneWeightType Weight = DataSource->GetBoneWeight(NeighborVertexID, Index);
+				if (Weight > MinimumWeightThreshold)
+				{
+					VertexNeighborWeights& BoneWeights = WeightsOnAllNeighbors.FindOrAdd(BoneIndex);
+					BoneWeights.Add(Weight);
+				}
+			}
+			
+		}
+	}
+
+	// calculate single average weight of each bone on all the neighbors, excluding locked ones
+	OutFinalWeights.Reset();
+	for (const TTuple<BoneIndexType, VertexNeighborWeights>& NeighborWeights : WeightsOnAllNeighbors)
+	{
+		BoneIndexType BoneIndex = NeighborWeights.Key;
+		if (!LockedBones.Contains(BoneIndex))
+		{
+			BoneWeightType TotalWeightOnThisBone = (BoneWeightType) 0.f;
+			for (const BoneWeightType& Value : NeighborWeights.Value)
+			{
+				TotalWeightOnThisBone += Value;
+			}
+			OutFinalWeights.Add(BoneIndex, TotalWeightOnThisBone / (BoneWeightType)NeighborWeights.Value.Num());
+		}
+	}
+
+	// normalize the weights
+	NormalizeWeights(OutFinalWeights);
+	
+	
+	for (TTuple<BoneIndexType, BoneWeightType>& FinalWeight : OutFinalWeights)
+	{
+		BoneWeightType AvailableWeight = 1 - TotalLockedWeight;
+		
+		FinalWeight.Value = FinalWeight.Value * AvailableWeight;
+	}
+	
+	// lerp weights from previous values, to fully relaxed values by brush strength scaled by falloff
+	for (TTuple<BoneIndexType, BoneWeightType>& FinalWeight : OutFinalWeights)
+	{
+		const BoneIndexType BoneIndex = FinalWeight.Key;
+		BoneWeightType NewWeight = FinalWeight.Value;
+		BoneWeightType OldWeight = DataSource->GetWeightOfBoneOnVertex(VertexID, BoneIndex);
+		FinalWeight.Value = FMath::Lerp<BoneWeightType>(OldWeight, NewWeight, VertexFalloff);
+	}
+
+	// normalize again
+	NormalizeWeights(OutFinalWeights);
+
+	for (TTuple<BoneIndexType, BoneWeightType>& FinalWeight : OutFinalWeights)
+	{
+		BoneWeightType AvailableWeight = 1 - TotalLockedWeight;
+		
+		FinalWeight.Value = FinalWeight.Value * AvailableWeight;
+	}
+	
+	OutFinalWeights.Append(LockedWeights);
+
+	if (Cancelled()) 
+	{
+		return false;
+	}
+		
+	return true;	
+}
+
 
 //
 // FSmoothDynamicMeshVertexSkinWeights
diff --git a/Engine/Plugins/Runtime/GeometryProcessing/Source/DynamicMesh/Public/Operations/SmoothBoneWeights.h b/Engine/Plugins/Runtime/GeometryProcessing/Source/DynamicMesh/Public/Operations/SmoothBoneWeights.h
--- a/Engine/Plugins/Runtime/GeometryProcessing/Source/DynamicMesh/Public/Operations/SmoothBoneWeights.h
+++ b/Engine/Plugins/Runtime/GeometryProcessing/Source/DynamicMesh/Public/Operations/SmoothBoneWeights.h
@@ -131,7 +131,21 @@ public:
 	virtual bool SmoothWeightsAtVertex(const int32 VertexID,
 									   const BoneWeightType VertexFalloff, 
 									   TMap<BoneIndexType, BoneWeightType>& OutFinalWeights);
-
+	
+	/**
+	 * Compute smoothed weights at a vertex and return the result.
+	 * 
+	 * @param VertexID		  The vertex whose weights we are smoothing.
+	 * @param VertexFalloff   The lerp value between the old and the new weight.
+	 * @param LockedBones	  The bones to exclude from the smooth op
+	 * @param OutFinalWeights The smoothed weights.
+	 * 
+	 * @return true if the algorithm succeeds, false if it failed or was canceled by the user.
+	 */
+	virtual bool SmoothWeightsAtVertex(const int32 VertexID,
+									   const BoneWeightType VertexFalloff,
+									   const TSet<BoneIndexType>& LockedBones,
+									   TMap<BoneIndexType, BoneWeightType>& OutFinalWeights);
 protected:
 	
     /** @return if true, abort the computation. */
diff --git a/Engine/Plugins/Runtime/MeshModelingToolset/Source/MeshModelingToolsEditorOnly/MeshModelingToolsEditorOnly.Build.cs b/Engine/Plugins/Runtime/MeshModelingToolset/Source/MeshModelingToolsEditorOnly/MeshModelingToolsEditorOnly.Build.cs
--- a/Engine/Plugins/Runtime/MeshModelingToolset/Source/MeshModelingToolsEditorOnly/MeshModelingToolsEditorOnly.Build.cs
+++ b/Engine/Plugins/Runtime/MeshModelingToolset/Source/MeshModelingToolsEditorOnly/MeshModelingToolsEditorOnly.Build.cs
@@ -45,7 +45,9 @@ public class MeshModelingToolsEditorOnly : ModuleRules
 				"SceneOutliner",
 				"Slate",
 				"Json",
-				"ApplicationCore"
+				"ApplicationCore",
+				"SlateCore",
+				"ToolWidgets"
 				// ... add private dependencies that you statically link with here ...	
 			}
 		);
diff --git a/Engine/Plugins/Runtime/MeshModelingToolset/Source/MeshModelingToolsEditorOnly/Private/MeshModelingToolsEditorOnly.cpp b/Engine/Plugins/Runtime/MeshModelingToolset/Source/MeshModelingToolsEditorOnly/Private/MeshModelingToolsEditorOnly.cpp
--- a/Engine/Plugins/Runtime/MeshModelingToolset/Source/MeshModelingToolsEditorOnly/Private/MeshModelingToolsEditorOnly.cpp
+++ b/Engine/Plugins/Runtime/MeshModelingToolset/Source/MeshModelingToolsEditorOnly/Private/MeshModelingToolsEditorOnly.cpp
@@ -7,6 +7,8 @@
 
 #include "DetailsCustomizations/UVUnwrapToolCustomizations.h"
 #include "Properties/RecomputeUVsProperties.h"
+#include "SkeletalMesh/SkeletonEditingTool.h"
+#include "SkeletalMesh/SkinWeightsPaintTool.h"
 
 #define LOCTEXT_NAMESPACE "FMeshModelingToolsEditorOnlyModule"
 
@@ -17,7 +19,9 @@ void FMeshModelingToolsEditorOnlyModule::StartupModule()
 	// This code will execute after your module is loaded into memory; the exact timing is specified in the .uplugin file per-module
 
 	FCoreDelegates::OnPostEngineInit.AddRaw(this, &FMeshModelingToolsEditorOnlyModule::OnPostEngineInit);
-
+	
+	FSkeletonEditingToolCommands::Register();
+	FSkinWeightsPaintToolCommands::Register();
 }
 
 void FMeshModelingToolsEditorOnlyModule::ShutdownModule()
diff --git a/Engine/Plugins/Runtime/MeshModelingToolset/Source/MeshModelingToolsEditorOnly/Private/SkeletalMesh/SReferenceSkeletonTree.cpp b/Engine/Plugins/Runtime/MeshModelingToolset/Source/MeshModelingToolsEditorOnly/Private/SkeletalMesh/SReferenceSkeletonTree.cpp
new file mode 100644
--- /dev/null
+++ b/Engine/Plugins/Runtime/MeshModelingToolset/Source/MeshModelingToolsEditorOnly/Private/SkeletalMesh/SReferenceSkeletonTree.cpp
@@ -0,0 +1,1385 @@
+// Copyright Epic Games, Inc. All Rights Reserved.
+
+#include "SkeletalMesh/SReferenceSkeletonTree.h"
+
+#include "SkeletonModifier.h"
+#include "SkeletalMesh/SkeletonEditingTool.h"
+#include "SkeletonClipboard.h"
+
+#include "SPositiveActionButton.h"
+
+#include "Misc/TextFilterExpressionEvaluator.h"
+#include "Widgets/Input/SSearchBox.h"
+#include "Widgets/Text/SInlineEditableTextBlock.h"
+
+#include "Framework/MultiBox/MultiBoxBuilder.h"
+#include "Editor.h"
+
+#include "Misc/ITransaction.h"
+
+#define LOCTEXT_NAMESPACE "SReferenceSkeletonTree"
+
+namespace ReferenceSkeletonTreeLocals
+{
+
+FSkeletonModifierChange::FSkeletonModifierChange(const USkeletonModifier* InModifier)
+	: FCommandChange()
+	, PreChangeSkeleton(InModifier->GetReferenceSkeleton())
+	, PreBoneTracker(InModifier->GetBoneIndexTracker())
+	, PostChangeSkeleton(InModifier->GetReferenceSkeleton())
+	, PostBoneTracker(InModifier->GetBoneIndexTracker())
+{}
+
+void FSkeletonModifierChange::StoreSkeleton(const USkeletonModifier* InModifier)
+{
+	PostChangeSkeleton = InModifier->GetReferenceSkeleton();
+	PostBoneTracker = InModifier->GetBoneIndexTracker();
+}
+
+void FSkeletonModifierChange::Apply(UObject* Object)
+{ // redo
+	USkeletonModifier* Modifier = CastChecked<USkeletonModifier>(Object);
+	Modifier->ExternalUpdate(PostChangeSkeleton, PostBoneTracker);
+}
+
+void FSkeletonModifierChange::Revert(UObject* Object)
+{ // undo
+	USkeletonModifier* Modifier = CastChecked<USkeletonModifier>(Object);
+	Modifier->ExternalUpdate(PreChangeSkeleton, PreBoneTracker);
+}
+
+}
+
+FBoneElement::FBoneElement(const FName& InBoneName, TWeakObjectPtr<USkeletonModifier> InModifier)
+	: BoneName(InBoneName)
+	, WeakModifier(InModifier)
+{}
+
+void FBoneElement::RequestRename(const ReferenceSkeletonTreeLocals::EBoneRenameType InRenameType) const
+{
+	OnRenameRequested.ExecuteIfBound(InRenameType);
+}
+
+SBoneItem::~SBoneItem()
+{
+	if (TSharedPtr<FBoneElement> BoneElement = WeakTreeElement.IsValid() ? WeakTreeElement.Pin() : nullptr)
+	{
+		BoneElement->OnRenameRequested.Unbind();
+	}
+}
+
+void SBoneItem::Construct(const FArguments& InArgs)
+{
+	// cf. FSkeletonTreeBoneItem::GenerateWidgetForNameColumn
+	
+	WeakTreeElement = InArgs._TreeElement;
+	WeakModifier = InArgs._WeakModifier;
+
+	const TSharedPtr<FBoneElement> TreeElement = WeakTreeElement.Pin();
+
+	const FSlateBrush* Brush = FAppStyle::Get().GetBrush("SkeletonTree.Bone");
+
+	TSharedPtr<SHorizontalBox> RowBox;
+
+	ChildSlot
+	[
+		SAssignNew(RowBox, SHorizontalBox)
+		+ SHorizontalBox::Slot()
+		.AutoWidth()
+		.Padding(FMargin(0.0f, 2.0f))
+		.VAlign(VAlign_Center)
+		.HAlign(HAlign_Center)
+		[
+			SNew(SImage)
+			.Image(Brush)
+			.ColorAndOpacity_Static(&SBoneItem::GetTextColor, InArgs._IsSelected)
+		]
+	];
+
+	RowBox->AddSlot()
+    .AutoWidth()
+	.Padding(4, 0, 0, 0)
+    .VAlign(VAlign_Center)
+    [
+	    SAssignNew(InlineWidget, SInlineEditableTextBlock)
+		.Text(this, &SBoneItem::GetName)
+		// .Font(TextFont)
+	    .IsReadOnly(InArgs._IsReadOnly)
+		.ColorAndOpacity_Static(&SBoneItem::GetTextColor, InArgs._IsSelected)
+		.MaximumLength(NAME_SIZE-1)
+		.OnVerifyTextChanged(this, &SBoneItem::OnVerifyNameChanged)
+		.OnTextCommitted(InArgs._Delegates.OnBoneNameCommitted)
+		.MultiLine(false)
+    	.OnEnterEditingMode(this, &SBoneItem::OnEnterEditingMode)
+		.OnExitEditingMode(this, &SBoneItem::OnExitEditingMode)
+    	.ToolTipText_Lambda([this]()
+		{
+    		if (bEditing)
+    		{
+    			return LOCTEXT("BoneItem_RenameTooltip",
+    				"use $name to add a prefix/suffix or /find/replace/ to replace a substring\n"
+    				" - \"spine_$name\" to add \"spine_\" as prefix.\n"
+    				" - \"$name_left\" to add \"_left\" as suffix.\n"
+    				" - \"/neck_/tail_/\" to replace \"neck_\" by \"tail_\".\n");	
+    		}
+    		return FText::GetEmpty();
+		})
+    ];
+
+	if (InlineWidget)
+	{
+		TreeElement->OnRenameRequested.BindRaw(this, &SBoneItem::HandleOnRenameRequested);	
+	}
+}
+
+FSlateColor SBoneItem::GetTextColor(FIsSelected InIsSelected)
+{
+	const bool bIsSelected = InIsSelected.IsBound() ? InIsSelected.Execute() : false;
+	if (bIsSelected)
+	{
+		return FAppStyle::Get().GetSlateColor("Colors.ForegroundInverted");
+	}
+	
+	return FSlateColor::UseForeground();
+}
+
+void SBoneItem::HandleOnRenameRequested(const ReferenceSkeletonTreeLocals::EBoneRenameType InRenameType)
+{
+	OptRenameType = InRenameType;
+	if (InlineWidget)
+	{
+		InlineWidget->EnterEditingMode();
+	}
+}
+
+bool SBoneItem::OnVerifyNameChanged(const FText& InText, FText& OutErrorMessage) const
+{
+	if (!WeakModifier.IsValid())
+	{
+		return false;
+	}
+
+	const TSharedPtr<FBoneElement> BoneElementPtr = WeakTreeElement.Pin();
+	const FReferenceSkeleton& ReferenceSkeleton = WeakModifier->GetReferenceSkeleton();
+
+	const FName NewName = FName(InText.ToString());
+	int32 BoneIndex = ReferenceSkeleton.FindRawBoneIndex(NewName);
+	if (BoneIndex != INDEX_NONE)
+	{
+		// changing case is allowed 
+		if (!ReferenceSkeleton.GetBoneName(BoneIndex).IsEqual(NewName, ENameCase::CaseSensitive))
+		{
+			BoneIndex = INDEX_NONE;
+		}
+	}
+	
+	return BoneIndex == INDEX_NONE;
+}
+
+void SBoneItem::OnEnterEditingMode()
+{
+	bEditing = true;
+}
+
+void SBoneItem::OnExitEditingMode()
+{
+	bEditing = false;
+	OptRenameType.Reset();
+}
+
+FText SBoneItem::GetName() const
+{
+	TSharedPtr<FBoneElement> BoneElement = WeakTreeElement.IsValid() ? WeakTreeElement.Pin() : nullptr;
+	ensure(BoneElement);
+	
+	if (bEditing && OptRenameType)
+	{
+		switch (*OptRenameType)
+		{
+		case ReferenceSkeletonTreeLocals::EBoneRenameType::Rename:
+			return FText::FromName(BoneElement->BoneName);
+		case ReferenceSkeletonTreeLocals::EBoneRenameType::SearchAndReplace:
+			return FText::FromString(TEXT("/to_search/to_replace/"));
+		case ReferenceSkeletonTreeLocals::EBoneRenameType::AddPrefix:
+			return FText::FromString(TEXT("prefix_$name"));
+		case ReferenceSkeletonTreeLocals::EBoneRenameType::AddSuffix:
+			return FText::FromString(TEXT("$name_suffix"));
+		default:
+			return FText::FromName(BoneElement->BoneName);
+		}
+	}
+	
+	return FText::FromName(WeakTreeElement.Pin()->BoneName);
+}
+
+void SReferenceSkeletonRow::Construct(const FArguments& InArgs, const TSharedRef<STableViewBase>& InOwnerTable)
+{
+	WeakTreeElement = InArgs._TreeElement;
+	WeakModifier = InArgs._WeakModifier;
+	Delegates = InArgs._Delegates;
+	bIsReadOnly = InArgs._IsReadOnly;
+
+	const FSuperRowType::FArguments Args = FSuperRowType::FArguments()
+		.OnDragDetected(InArgs._Delegates.OnDragDetected)
+		.OnCanAcceptDrop(InArgs._Delegates.OnCanAcceptDrop)
+		.OnAcceptDrop(InArgs._Delegates.OnAcceptDrop)
+		.Style( FAppStyle::Get(), "TableView.AlternatingRow" );
+	
+	SMultiColumnTableRow::Construct(Args , InOwnerTable);
+}
+
+TSharedRef<SWidget> SReferenceSkeletonRow::GenerateWidgetForColumn(const FName& ColumnName)
+{
+	TSharedRef<SHorizontalBox> Widget =
+		SNew(SHorizontalBox)
+		+ SHorizontalBox::Slot()
+		.AutoWidth()
+		[
+			SNew( SExpanderArrow, SharedThis(this) )
+				.ShouldDrawWires(true)
+		]
+
+		+ SHorizontalBox::Slot()
+		.AutoWidth()
+		[
+			SNew(SBoneItem)
+				.WeakModifier(WeakModifier)
+				.TreeElement(WeakTreeElement)
+				.Delegates(Delegates)
+				.IsSelected(this, &SReferenceSkeletonRow::IsSelected)
+				.IsReadOnly(bIsReadOnly)
+		];
+
+	if (Delegates.OnGetRowWidget.IsBound())
+	{
+		Widget->InsertSlot(0)
+		.AutoWidth()
+		[
+			Delegates.OnGetRowWidget.Execute(WeakTreeElement.Pin()->BoneName)	
+		];
+	}
+
+	return Widget;
+}
+
+TSharedRef<FBoneItemDragDropOp> FBoneItemDragDropOp::New(TArrayView<const TSharedPtr<FBoneElement>> InElements)
+{
+	TSharedRef<FBoneItemDragDropOp> Operation = MakeShared<FBoneItemDragDropOp>();
+
+	for (const TSharedPtr<FBoneElement>& BoneElement: InElements)
+	{
+		const bool bContainsParent = InElements.Contains(BoneElement->Parent);
+		if (!bContainsParent)
+		{
+			Operation->Elements.Add(BoneElement);
+		}
+	}
+	
+	Operation->Construct();
+	return Operation;
+}
+
+TSharedPtr<SWidget> FBoneItemDragDropOp::GetDefaultDecorator() const
+{
+	return SNew(SBorder)
+        .Visibility(EVisibility::Visible)
+        .BorderImage(FAppStyle::GetBrush("Menu.Background"))
+        [
+            SNew(STextBlock)
+            .Text_Lambda([this]()
+            {
+	            return !Elements.IsEmpty() && Elements[0].IsValid() ? FText::FromName(Elements[0].Pin()->BoneName) : FText::GetEmpty();
+            })
+        ];
+}
+
+SReferenceSkeletonTree::SReferenceSkeletonTree()
+{}
+
+SReferenceSkeletonTree::~SReferenceSkeletonTree()
+{
+	// GEditor->UnregisterForUndo(this);
+}
+
+void SReferenceSkeletonTree::Construct(const FArguments& InArgs)
+{
+	Modifier = InArgs._Modifier;
+	bIsReadOnly = Modifier->IsReadOnly();
+
+	Delegates = InArgs._Delegates;
+	Delegates.OnCanAcceptDrop = !Delegates.OnCanAcceptDrop.IsBound() ? FOnRefSkeletonTreeCanAcceptDrop::CreateSP(this, &SReferenceSkeletonTree::OnCanAcceptDrop) : Delegates.OnCanAcceptDrop;
+	Delegates.OnAcceptDrop = !Delegates.OnAcceptDrop.IsBound() ? FOnRefSkeletonTreeAcceptDrop::CreateSP(this, &SReferenceSkeletonTree::OnAcceptDrop) : Delegates.OnAcceptDrop;
+	Delegates.OnDragDetected = !Delegates.OnDragDetected.IsBound() ? FOnDragDetected::CreateSP(this, &SReferenceSkeletonTree::OnDragDetected) : Delegates.OnDragDetected;
+	Delegates.OnBoneRenamed = !Delegates.OnBoneRenamed.IsBound() ? FOnBoneRenamed::CreateSP(this, &SReferenceSkeletonTree::OnBoneRenamed) : Delegates.OnBoneRenamed;
+	Delegates.OnBoneNameCommitted =!Delegates.OnBoneNameCommitted.IsBound() ? FOnTextCommitted::CreateSP(this, &SReferenceSkeletonTree::OnNewBoneNameCommitted) : Delegates.OnBoneNameCommitted;
+	
+	CommandList = MakeShared<FUICommandList>();
+	BindCommands();
+
+
+	TextFilter = MakeShareable(new FTextFilterExpressionEvaluator(ETextFilterExpressionEvaluatorMode::BasicString));
+
+	
+	ChildSlot
+	[
+		SNew( SOverlay )
+		+SOverlay::Slot()
+		[
+			SNew(SVerticalBox)
+			+SVerticalBox::Slot()
+			.AutoHeight()
+			.Padding(FMargin(0.f, 2.f))
+			[
+				SNew(SHorizontalBox)
+				+SHorizontalBox::Slot()
+				.AutoWidth()
+				.VAlign(VAlign_Center)
+				.Padding(FMargin(6.f, 0.0))
+				[
+					SNew(SPositiveActionButton)
+					.OnGetMenuContent(this, &SReferenceSkeletonTree::CreateAddNewMenu)
+					.Icon(FAppStyle::Get().GetBrush("Icons.Plus"))
+					.Visibility(this, &SReferenceSkeletonTree::GetAddNewButtonVisibility)
+				]
+
+				+SHorizontalBox::Slot()
+				.FillWidth(1.0f)
+				[
+					SNew(SSearchBox)
+					.SelectAllTextWhenFocused(true)
+					.OnTextChanged( this, &SReferenceSkeletonTree::OnFilterTextChanged )
+					.HintText(LOCTEXT( "SearchBoxHint", "Search Reference Skeleton Tree..."))
+				]
+			]
+			
+			+SVerticalBox::Slot()
+			.Padding(0.0f, 0.0f)
+			[
+				SNew(SBorder)
+				.Padding(2.0f)
+				.BorderImage(FAppStyle::GetBrush("SCSEditor.TreePanel"))
+				[
+					SAssignNew(TreeView, SRefSkeletonTreeView)
+					.TreeItemsSource(&RootElements)
+					.SelectionMode(ESelectionMode::Multi)
+					.OnGenerateRow_Lambda( [this](
+						TSharedPtr<FBoneElement> InItem,
+						const TSharedRef<STableViewBase>& OwnerTable) -> TSharedRef<ITableRow>
+					{
+						return SNew(SReferenceSkeletonRow, OwnerTable)
+							.WeakModifier(Modifier)
+							.TreeElement(InItem)
+							.Delegates(Delegates)
+							.IsReadOnly(bIsReadOnly);
+					})
+					.OnGetChildren(this, &SReferenceSkeletonTree::HandleGetChildrenForTree)
+					.OnSelectionChanged(this, &SReferenceSkeletonTree::OnSelectionChanged)
+					.OnContextMenuOpening(this, &SReferenceSkeletonTree::CreateContextMenu)
+					.OnMouseButtonDoubleClick(this, &SReferenceSkeletonTree::OnItemDoubleClicked)
+					.OnSetExpansionRecursive(this, &SReferenceSkeletonTree::OnSetExpansionRecursive)
+					.HighlightParentNodesForSelection(true)
+					.HeaderRow
+					(
+						SNew(SHeaderRow)
+						+ SHeaderRow::Column("Name")
+						.ShouldGenerateWidget(true)
+						.DefaultLabel(LOCTEXT("BoneNameLabel", "Name"))
+						.FillWidth(0.5f)
+					)
+				]
+			]
+		]
+    ];
+
+	static constexpr bool IsInitialSetup = true;
+	RefreshTreeView(IsInitialSetup);
+
+	GEditor->RegisterForUndo(this);
+}
+
+void SReferenceSkeletonTree::PostUndo(bool bSuccess)
+{
+	RefreshTreeView(true);
+}
+
+void SReferenceSkeletonTree::PostRedo(bool bSuccess)
+{
+	RefreshTreeView(true);
+}
+
+void SReferenceSkeletonTree::AddItemToSelection(const TSharedPtr<FBoneElement>& InItem)
+{
+	TreeView->SetItemSelection(InItem, true, ESelectInfo::Direct);
+	TreeView->RequestScrollIntoView(InItem);
+}
+
+void SReferenceSkeletonTree::RemoveItemFromSelection(const TSharedPtr<FBoneElement>& InItem)
+{
+	TreeView->SetItemSelection(InItem, false, ESelectInfo::Direct);
+}
+
+void SReferenceSkeletonTree::ReplaceItemInSelection(const FText& InOldName, const FText& InNewName)
+{
+	const FName OldName(*InOldName.ToString());
+	const FName NewName(*InNewName.ToString());
+	
+	for (const TSharedPtr<FBoneElement>& Item : AllElements)
+	{
+		// remove old selection
+		if (Item->BoneName.IsEqual(OldName))
+		{
+			TreeView->SetItemSelection(Item, false, ESelectInfo::Direct);
+		}
+		// add new selection
+		if (Item->BoneName.IsEqual(NewName))
+		{
+			TreeView->SetItemSelection(Item, true, ESelectInfo::Direct);
+		}
+	}
+}
+
+TArray<TSharedPtr<FBoneElement>> SReferenceSkeletonTree::GetSelectedItems() const
+{
+	return TreeView->GetSelectedItems();
+}
+
+bool SReferenceSkeletonTree::HasSelectedItems() const
+{
+	return TreeView->GetNumItemsSelected() > 0;
+}
+
+void SReferenceSkeletonTree::BindCommands()
+{
+	if (!bIsReadOnly)
+	{
+		const FSkeletonEditingToolCommands& Commands = FSkeletonEditingToolCommands::Get();
+ 
+		CommandList->MapAction(Commands.NewBone,
+			FExecuteAction::CreateSP(this, &SReferenceSkeletonTree::HandleNewBone),
+			FCanExecuteAction::CreateSP(this, &SReferenceSkeletonTree::CanAddNewBone));
+		
+		CommandList->MapAction(Commands.RemoveBone,
+			FExecuteAction::CreateSP(this, &SReferenceSkeletonTree::HandleDeleteBone),
+			FCanExecuteAction::CreateSP(this, &SReferenceSkeletonTree::CanDeleteBone));
+	 
+		CommandList->MapAction(Commands.RenameBone,
+			FExecuteAction::CreateSP(this, &SReferenceSkeletonTree::HandleRenameBone, ReferenceSkeletonTreeLocals::EBoneRenameType::Rename),
+			FCanExecuteAction::CreateSP(this, &SReferenceSkeletonTree::CanRenameBone));
+
+		CommandList->MapAction(Commands.SearchAndReplace,
+			FExecuteAction::CreateSP(this, &SReferenceSkeletonTree::HandleRenameBone, ReferenceSkeletonTreeLocals::EBoneRenameType::SearchAndReplace),
+			FCanExecuteAction::CreateSP(this, &SReferenceSkeletonTree::CanRenameBone));
+
+		CommandList->MapAction(Commands.AddPrefix,
+			FExecuteAction::CreateSP(this, &SReferenceSkeletonTree::HandleRenameBone, ReferenceSkeletonTreeLocals::EBoneRenameType::AddPrefix),
+			FCanExecuteAction::CreateSP(this, &SReferenceSkeletonTree::CanRenameBone));
+
+		CommandList->MapAction(Commands.AddSuffix,
+			FExecuteAction::CreateSP(this, &SReferenceSkeletonTree::HandleRenameBone, ReferenceSkeletonTreeLocals::EBoneRenameType::AddSuffix),
+			FCanExecuteAction::CreateSP(this, &SReferenceSkeletonTree::CanRenameBone));
+		
+		CommandList->MapAction(Commands.UnParentBone,
+			FExecuteAction::CreateSP(this, &SReferenceSkeletonTree::HandleUnParentBone),
+			FCanExecuteAction::CreateSP(this, &SReferenceSkeletonTree::CanUnParentBone));
+
+		CommandList->MapAction(Commands.CopyBones,
+			FExecuteAction::CreateSP(this, &SReferenceSkeletonTree::HandleCopyBones),
+			FCanExecuteAction::CreateSP(this, &SReferenceSkeletonTree::CanCopyBones));
+		
+		CommandList->MapAction(Commands.PasteBones,
+			FExecuteAction::CreateSP(this, &SReferenceSkeletonTree::HandlePasteBones),
+			FCanExecuteAction::CreateSP(this, &SReferenceSkeletonTree::CanPasteBones));
+		
+		CommandList->MapAction(Commands.DuplicateBones,
+			FExecuteAction::CreateSP(this, &SReferenceSkeletonTree::HandleDuplicateBones),
+			FCanExecuteAction::CreateSP(this, &SReferenceSkeletonTree::CanDuplicateBones));	
+	}
+
+
+	if (Delegates.OnExtendCommandList.IsBound())
+	{
+		Delegates.OnExtendCommandList.Execute(CommandList);
+	}
+}
+
+void SReferenceSkeletonTree::HandleNewBone()
+{
+	if (!Modifier.IsValid())
+	{
+		return;
+	}
+
+	BeginChange();
+	
+	TArray<FName> BoneNames;
+	GetSelectedBoneNames(BoneNames);
+	
+	const FName ParentName = BoneNames.IsEmpty() ? NAME_None : BoneNames[0];
+
+	static const FName DefaultName("joint");
+	const FName BoneName = Modifier->GetUniqueName(DefaultName);
+	const bool bAdded = Modifier->AddBone(BoneName, ParentName, FTransform::Identity);
+	if (bAdded)
+	{
+		RefreshTreeView(true);
+		if (Notifier.IsValid())
+		{
+			Notifier->Notify( {BoneName}, ESkeletalMeshNotifyType::BonesAdded);
+		}
+		EndChange();
+	}
+	else
+	{
+		CancelChange();
+	}
+}
+
+bool SReferenceSkeletonTree::CanAddNewBone() const
+{
+	if (bIsReadOnly)
+	{
+		return false;
+	}
+	
+	return true;
+}
+
+void SReferenceSkeletonTree::HandleDeleteBone()
+{
+	if (!Modifier.IsValid())
+	{
+		return;
+	}
+
+	BeginChange();
+	
+	TArray<FName> BoneNames;
+	GetSelectedBoneNames(BoneNames);
+	
+	const bool bRemoved = Modifier->RemoveBones(BoneNames, true);
+	if (bRemoved)
+	{
+		RefreshTreeView(true);
+		if (Notifier.IsValid())
+		{
+			Notifier->Notify(BoneNames, ESkeletalMeshNotifyType::BonesRemoved);
+		}
+		EndChange();
+	}
+	else
+	{
+		CancelChange();
+	}
+}
+
+bool SReferenceSkeletonTree::CanDeleteBone() const
+{
+	if (bIsReadOnly)
+	{
+		return false;
+	}
+	
+	return HasSelectedItems();
+}
+
+void SReferenceSkeletonTree::HandleUnParentBone()
+{
+	if (!Modifier.IsValid())
+	{
+		return;
+	}
+
+	BeginChange();
+	
+	TArray<FName> BoneNames;
+	GetSelectedBoneNames(BoneNames);
+
+	static const TArray<FName> Dummy;
+	const bool bUnParented = Modifier->ParentBones(BoneNames, Dummy);
+	if (bUnParented)
+	{
+		RefreshTreeView(true);
+		SelectItemFromNames(BoneNames);
+		if (Notifier.IsValid())
+		{
+			Notifier->Notify(BoneNames, ESkeletalMeshNotifyType::HierarchyChanged);
+		}
+		EndChange();
+	}
+	else
+	{
+		CancelChange();
+	}
+}
+
+bool SReferenceSkeletonTree::CanUnParentBone() const
+{
+	if (bIsReadOnly)
+	{
+		return false;
+	}
+	
+	return HasSelectedItems();
+}
+
+void SReferenceSkeletonTree::HandleCopyBones() const
+{
+	if (!ensure(Modifier.IsValid()))
+	{
+		return;
+	}
+
+	TArray<FName> BoneNames;
+	GetSelectedBoneNames(BoneNames);
+
+	if (BoneNames.IsEmpty())
+	{
+		return;
+	}
+
+	SkeletonClipboard::CopyToClipboard(*Modifier.Get(), BoneNames);
+}
+
+bool SReferenceSkeletonTree::CanCopyBones() const
+{
+	if (bIsReadOnly)
+	{
+		return false;
+	}
+
+	return Modifier.IsValid() && HasSelectedItems();
+}
+
+void SReferenceSkeletonTree::HandlePasteBones()
+{
+	TArray<FName> BoneNames;
+	GetSelectedBoneNames(BoneNames);
+
+	const FName DefaultParent = BoneNames.IsEmpty() ? NAME_None : BoneNames[0];  
+
+	BeginChange();
+	
+	const TArray<FName> NewBones = SkeletonClipboard::PasteFromClipboard(*Modifier.Get(), DefaultParent);
+	if (NewBones.IsEmpty())
+	{
+		CancelChange();
+		return;
+	}
+
+	static constexpr bool bRebuildAll = true;
+	RefreshTreeView(bRebuildAll);
+	static constexpr bool bFrameSelection = true;
+	SelectItemFromNames(NewBones, bFrameSelection);
+
+	if (Notifier.IsValid())
+	{
+		Notifier->Notify(NewBones, ESkeletalMeshNotifyType::HierarchyChanged);
+		Notifier->Notify(NewBones, ESkeletalMeshNotifyType::BonesSelected);
+	}
+
+	EndChange();
+}
+
+bool SReferenceSkeletonTree::CanPasteBones() const
+{
+	if (bIsReadOnly)
+	{
+		return false;
+	}
+	
+	return Modifier.IsValid() && SkeletonClipboard::IsClipboardValid();
+}
+
+void SReferenceSkeletonTree::HandleDuplicateBones()
+{
+	HandleCopyBones();
+
+	if (!SkeletonClipboard::IsClipboardValid())
+	{
+		return;
+	}  
+
+	BeginChange();
+	
+	const TArray<FName> NewBones = SkeletonClipboard::PasteFromClipboard(*Modifier.Get(), NAME_None);
+	if (NewBones.IsEmpty())
+	{
+		CancelChange();
+		return;
+	}
+
+	static constexpr bool bRebuildAll = true;
+	RefreshTreeView(bRebuildAll);
+	static constexpr bool bFrameSelection = true;
+	SelectItemFromNames(NewBones, bFrameSelection);
+
+	if (Notifier.IsValid())
+	{
+		Notifier->Notify(NewBones, ESkeletalMeshNotifyType::HierarchyChanged);
+		Notifier->Notify(NewBones, ESkeletalMeshNotifyType::BonesSelected);
+	}
+
+	EndChange();
+}
+
+bool SReferenceSkeletonTree::CanDuplicateBones() const
+{
+	return CanCopyBones();
+}
+
+void SReferenceSkeletonTree::GetSelectedBoneNames(TArray<FName>& OutSelectedBoneNames) const
+{
+	OutSelectedBoneNames.Reset();
+	const TArray<TSharedPtr<FBoneElement>> BoneElements = TreeView->GetSelectedItems();
+	OutSelectedBoneNames.Reserve(BoneElements.Num());
+
+	Algo::Transform(BoneElements, OutSelectedBoneNames, [](const TSharedPtr<FBoneElement>& BoneElement)
+	{
+		return BoneElement->BoneName;
+	});
+}
+
+void SReferenceSkeletonTree::SelectItemFromNames(const TArray<FName>& InBoneNames, bool bFrameSelection)
+{
+	for (const TSharedPtr<FBoneElement>& Item : AllElements)
+	{
+		const bool bSelect = InBoneNames.Contains(Item->BoneName);
+		TreeView->SetItemSelection(Item, bSelect, ESelectInfo::Direct);
+
+		if (bFrameSelection && bSelect)
+		{
+			TreeView->RequestScrollIntoView(Item);
+			bFrameSelection = false;
+		}
+	}
+}
+
+TSharedPtr<ISkeletalMeshNotifier> SReferenceSkeletonTree::GetNotifier()
+{
+	if (!Notifier)
+	{
+		Notifier = MakeShared<FReferenceSkeletonWidgetNotifier>(SharedThis(this));
+	}
+	
+	return Notifier;
+}
+
+void SReferenceSkeletonTree::HandleRenameBone(const ReferenceSkeletonTreeLocals::EBoneRenameType InRenameType) const
+{
+	TArray<TSharedPtr<FBoneElement>> BoneElements = TreeView->GetSelectedItems();
+	if (BoneElements.IsEmpty())
+	{
+		return;
+	}
+	BoneElements[0]->RequestRename(InRenameType);
+}
+
+bool SReferenceSkeletonTree::CanRenameBone() const
+{
+	if (bIsReadOnly)
+	{
+		return false;
+	}
+	
+	return HasSelectedItems();
+}
+
+void SReferenceSkeletonTree::OnBoneRenamed(const FName InOldName, const FName InNewName) const
+{
+	if (!Modifier.IsValid())
+	{
+		return;
+	}
+
+	if (Notifier.IsValid())
+	{
+		Notifier->Notify( {InNewName}, ESkeletalMeshNotifyType::BonesRenamed);
+	}
+}
+
+void SReferenceSkeletonTree::OnNewBoneNameCommitted(const FText& InText, ETextCommit::Type InCommitType)
+{
+	if (InCommitType == ETextCommit::OnCleared)
+	{
+		return;
+	}
+	
+	TArray<FName> BoneNames;
+	GetSelectedBoneNames(BoneNames);
+
+	static const FString NamePattern(TEXT("$name"));
+	static constexpr TCHAR Slash(TEXT('/'));
+	
+	const FString NameStr = InText.ToString();
+	const FName NewName = FName(NameStr);
+
+	// wrong input
+	if (BoneNames.IsEmpty() || NewName == NAME_None || NameStr == NamePattern)
+	{
+		return;
+	}
+	
+	// same name
+	if (BoneNames[0].IsEqual(NewName, ENameCase::CaseSensitive))
+	{
+		return;
+	}
+
+	// look for prefix / suffix pattern 
+	const bool bAddPrefixSuffix = NameStr.Contains(NamePattern);
+	
+	// look for find & replace pattern 
+	int32 SlashIndex[3] = { INDEX_NONE, INDEX_NONE, INDEX_NONE };
+	for (int32 CharIndex = 0, NumSlash = 0; NumSlash < 3 && CharIndex < NameStr.Len(); ++CharIndex)
+	{
+		if (NameStr[CharIndex] == Slash)
+		{
+			SlashIndex[NumSlash++] = CharIndex;
+		}
+	}
+
+	const bool bFindReplace =
+		SlashIndex[0] == 0 && SlashIndex[2] == NameStr.Len()-1 && // starts and ends with /
+		SlashIndex[0] != SlashIndex[1] && SlashIndex[0] != SlashIndex[2]; // all three indices are different
+
+	if (bFindReplace)
+	{
+		if (bAddPrefixSuffix || (SlashIndex[1] - SlashIndex[0]) < 2)
+		{
+			// error
+			return;
+		}
+	}
+	
+	TArray<int32> BoneIndices;
+	const FReferenceSkeleton& ReferenceSkeleton = Modifier->GetReferenceSkeleton();
+	Algo::Transform(BoneNames, BoneIndices, [&](const FName& BoneName)
+	{
+		return ReferenceSkeleton.FindRawBoneIndex(BoneName);
+	});
+
+	TArray<FName> NewNames;
+
+	if (bFindReplace)
+	{
+		const FString Find = FString::ConstructFromPtrSize(&NameStr[SlashIndex[0]+1], SlashIndex[1]-SlashIndex[0]-1);
+		const FString Replace = FString::ConstructFromPtrSize(&NameStr[SlashIndex[1]+1], SlashIndex[2]-SlashIndex[1]-1);
+		
+		NewNames.Reserve(BoneNames.Num());
+		for (int32 Index = 0; Index < BoneNames.Num(); Index++)
+		{
+			FString NewBoneNameStr = BoneNames[Index].ToString().Replace(*Find, *Replace, ESearchCase::CaseSensitive);
+			NewNames.Add(*NewBoneNameStr);
+		} 
+	}
+	else
+	{
+		FString Prefix, Suffix;
+		if (bAddPrefixSuffix && NameStr.Split(NamePattern, &Prefix, &Suffix))
+		{
+			NewNames.Reserve(BoneNames.Num());
+			for (int32 Index = 0; Index < BoneNames.Num(); Index++)
+			{
+				FString NewBoneNameStr = FString::Printf( TEXT("%s%s%s"), *Prefix, *BoneNames[Index].ToString(), *Suffix);
+				NewNames.Add(*NewBoneNameStr);
+			}
+		}
+		else
+		{
+			NewNames.Init(NewName, BoneNames.Num());
+		}
+	}
+	
+	BeginChange();
+
+	const bool bBoneRenamed = Modifier->RenameBones(BoneNames, NewNames);
+	if (bBoneRenamed)
+	{
+		const TArray<FMeshBoneInfo>& BoneInfos = ReferenceSkeleton.GetRawRefBoneInfo();
+		for (int32 Index = 0; Index < BoneIndices.Num(); Index++)
+		{
+			if (BoneIndices[Index] != INDEX_NONE)
+			{
+				const FName OldName = BoneNames[Index];
+				BoneNames[Index] = BoneInfos[BoneIndices[Index]].Name;
+				const int32 ItemIndex = AllElements.IndexOfByPredicate([OldName](const TSharedPtr<FBoneElement>& Item)
+				{
+					return Item->BoneName == OldName;
+				});
+
+				if (ItemIndex != INDEX_NONE)
+				{
+					AllElements[ItemIndex]->BoneName = BoneNames[Index];
+				}
+			}
+		}
+		
+		if (Notifier.IsValid())
+		{
+			Notifier->Notify(BoneNames, ESkeletalMeshNotifyType::BonesRenamed);
+		}
+		
+		EndChange();
+	}
+	else
+	{
+		CancelChange();	
+	}
+}
+
+void SReferenceSkeletonTree::OnFilterTextChanged(const FText& SearchText)
+{
+	TextFilter->SetFilterText(SearchText);
+	RefreshTreeView();
+}
+
+void SReferenceSkeletonTree::RefreshTreeView(bool IsInitialSetup /*=false*/)
+{
+	if (!Modifier.IsValid())
+	{
+		return;
+	}
+	
+	const FReferenceSkeleton& RefSkeleton = Modifier->GetReferenceSkeleton();
+	const TArray<FMeshBoneInfo>& BoneInfos = RefSkeleton.GetRawRefBoneInfo();
+	
+	// reset all tree items
+	RootElements.Reset();
+	AllElements.Reset();
+
+	// validate we have a skeleton to load
+	if (BoneInfos.IsEmpty())
+	{
+		TreeView->RequestTreeRefresh();
+		return;
+	}
+	
+	// record bone element indices
+	TMap<FName, int32> BoneTreeElementIndices;
+
+	auto FilterString = [this](const FString& StringToTest) ->bool
+	{
+		return TextFilter->TestTextFilter(FBasicStringFilterExpressionContext(StringToTest));
+	};
+
+	TMap<int32, int32> PostponedParent;
+	
+	// create all bone elements
+	for (int32 BoneIndex=0; BoneIndex < BoneInfos.Num(); ++BoneIndex)
+	{
+		const FName BoneName = BoneInfos[BoneIndex].Name;
+		TSharedPtr<FBoneElement> BoneElement = MakeShared<FBoneElement>(BoneName, Modifier);
+		const int32 BoneElementIndex = AllElements.Add(BoneElement);
+		BoneTreeElementIndices.Add(BoneName, BoneElementIndex);
+
+		// store pointer to parent (if there is one)
+		const int32 ParentIndex = BoneInfos[BoneIndex].ParentIndex;
+		if (ParentIndex != INDEX_NONE)
+		{
+			// get parent tree element
+			const FName ParentBoneName = BoneInfos[ParentIndex].Name;
+
+			const int32* FoundParent = BoneTreeElementIndices.Find(ParentBoneName);
+			if (FoundParent)
+			{ // set parent info directly
+				const TSharedPtr<FBoneElement> ParentBoneTreeElement = AllElements[*FoundParent];
+				BoneElement->UnFilteredParent = ParentBoneTreeElement;
+			}
+			else
+			{ // postpone as the parent might not have been added to AllElements / BoneTreeElementIndices yet
+				PostponedParent.Emplace(BoneElementIndex, ParentIndex);
+			}
+		}
+		
+		// apply text filter to bones
+		if (!(TextFilter->GetFilterText().IsEmpty() || FilterString(BoneName.ToString())))
+		{
+			BoneElement->bIsHidden = true;
+		}
+	}
+
+	// store pointer to parents that were postponed 
+	for (const auto & [BoneElementIndex, ParentIndex] : PostponedParent)
+	{
+		const FName ParentBoneName = BoneInfos[ParentIndex].Name;
+		const int32* FoundParent = BoneTreeElementIndices.Find(ParentBoneName);
+		if (ensure(FoundParent))
+		{
+			const TSharedPtr<FBoneElement> BoneElement = AllElements[BoneElementIndex];
+			const TSharedPtr<FBoneElement> ParentBoneTreeElement = AllElements[*FoundParent];
+			BoneElement->UnFilteredParent = ParentBoneTreeElement;
+		}
+	}
+	
+	// resolve parent/children pointers on all tree elements, taking into consideration the filter options
+	// (elements are parented to their nearest non-hidden/filtered parent element)
+	for (TSharedPtr<FBoneElement>& Element : AllElements)
+	{
+		if (Element->bIsHidden)
+		{
+			continue;
+		}
+		
+		// find first parent that is not filtered
+		TSharedPtr<FBoneElement> ParentElement = Element->UnFilteredParent;
+		while (true)
+		{
+			if (!ParentElement.IsValid())
+			{
+				break;
+			}
+
+			if (!ParentElement->bIsHidden)
+			{
+				break;
+			}
+
+			ParentElement = ParentElement->UnFilteredParent;
+		}
+
+		if (ParentElement.IsValid())
+		{
+			// store pointer to child on parent
+			ParentElement->Children.Add(Element);
+			// store pointer to parent on child
+			Element->Parent = ParentElement;
+		}
+		else
+		{
+			// has no parent, store a root element
+			RootElements.Add(Element);
+		}
+	}
+
+	// expand all elements upon the initial construction of the tree
+	for (const TSharedPtr<FBoneElement>& RootElement : RootElements)
+	{
+		SetExpansionRecursive(RootElement, false, true);
+	}
+	
+	TreeView->RequestTreeRefresh();
+}
+
+void SReferenceSkeletonTree::HandleGetChildrenForTree(
+	TSharedPtr<FBoneElement> InItem,
+	TArray<TSharedPtr<FBoneElement>>& OutChildren)
+{
+	OutChildren = InItem.Get()->Children;
+}
+
+void SReferenceSkeletonTree::OnSelectionChanged(TSharedPtr<FBoneElement> InItem, ESelectInfo::Type InSelectInfo)
+{
+	if (InSelectInfo == ESelectInfo::Direct)
+	{
+		return;
+	}
+	
+	if (!InItem)
+	{
+		if (Notifier.IsValid())
+		{
+			Notifier->Notify({}, ESkeletalMeshNotifyType::BonesSelected);
+		}
+		return;
+	}
+	
+	TArray<FName> BoneNames;
+	GetSelectedBoneNames(BoneNames);
+	
+	if (Notifier.IsValid())
+	{
+		Notifier->Notify(BoneNames, ESkeletalMeshNotifyType::BonesSelected);
+	}
+}
+
+TSharedRef<SWidget> SReferenceSkeletonTree::CreateAddNewMenu()
+{
+	const FSkeletonEditingToolCommands& Commands = FSkeletonEditingToolCommands::Get();
+
+	static constexpr bool bCloseAfterSelection = true;
+	FMenuBuilder MenuBuilder(bCloseAfterSelection, CommandList);
+	
+	MenuBuilder.BeginSection("NewBone", LOCTEXT("AddNewBoneOperations", "Bones"));
+	MenuBuilder.AddMenuEntry(Commands.NewBone);
+	MenuBuilder.EndSection();
+		
+	return MenuBuilder.MakeWidget();
+}
+
+EVisibility SReferenceSkeletonTree::GetAddNewButtonVisibility() const
+{
+	return bIsReadOnly ? EVisibility::Collapsed : EVisibility::Visible;
+}
+
+TSharedPtr<SWidget> SReferenceSkeletonTree::CreateContextMenu()
+{
+	static constexpr bool bCloseAfterSelection = true;
+	FMenuBuilder MenuBuilder(bCloseAfterSelection, CommandList);
+
+	
+	const TArray<TSharedPtr<FBoneElement>> SelectedItems = GetSelectedItems();
+	
+	bool bHasValidMenuOptions = false;
+	if (!bIsReadOnly)
+	{
+		const FSkeletonEditingToolCommands& Commands = FSkeletonEditingToolCommands::Get();
+		bHasValidMenuOptions = true;
+		if (SelectedItems.IsEmpty())
+		{
+			
+			MenuBuilder.BeginSection("NewBone", LOCTEXT("AddNewBoneOperations", "Bones"));
+	
+			MenuBuilder.AddMenuEntry(Commands.NewBone);
+
+			MenuBuilder.EndSection();
+
+			MenuBuilder.BeginSection("CopyPasteBones", LOCTEXT("CopyPasteBonesOperations", "Copy & Paste"));
+
+			MenuBuilder.AddMenuEntry(Commands.PasteBones);
+
+			MenuBuilder.EndSection();
+		}
+		else
+		{
+			MenuBuilder.BeginSection("EditBones", LOCTEXT("EditBonesOperations", "Bones"));
+			{
+				MenuBuilder.AddMenuEntry(Commands.NewBone);
+				MenuBuilder.AddMenuEntry(Commands.RemoveBone);
+				MenuBuilder.AddMenuEntry(Commands.UnParentBone);
+			}
+			MenuBuilder.EndSection();
+		
+			MenuBuilder.BeginSection("RenameBones", LOCTEXT("RenameBonesOperations", "Rename"));
+			{
+				MenuBuilder.AddMenuEntry(Commands.RenameBone);
+				MenuBuilder.AddMenuEntry(Commands.SearchAndReplace);
+				MenuBuilder.AddMenuEntry(Commands.AddPrefix);
+				MenuBuilder.AddMenuEntry(Commands.AddSuffix);
+			}
+			MenuBuilder.EndSection();
+
+			MenuBuilder.BeginSection("CopyPasteBones", LOCTEXT("CopyPasteBonesOperations", "Copy & Paste"));
+			{
+				MenuBuilder.AddMenuEntry(Commands.CopyBones);
+				MenuBuilder.AddMenuEntry(Commands.PasteBones);
+				MenuBuilder.AddMenuEntry(Commands.DuplicateBones);
+			}
+			MenuBuilder.EndSection();
+		}	
+	}
+
+	if (Delegates.OnExtendContextMenu.IsBound())
+	{
+		if (Delegates.OnExtendContextMenu.Execute( MenuBuilder))
+		{
+			bHasValidMenuOptions = true;
+		}
+	}
+	
+	if (bHasValidMenuOptions)
+	{
+		return MenuBuilder.MakeWidget();
+	}
+	
+	return nullptr;
+}
+
+void SReferenceSkeletonTree::OnItemDoubleClicked(TSharedPtr<FBoneElement> InItem)
+{
+	InItem->RequestRename(ReferenceSkeletonTreeLocals::EBoneRenameType::Rename);
+}
+
+void SReferenceSkeletonTree::OnSetExpansionRecursive(TSharedPtr<FBoneElement> InItem, bool bShouldBeExpanded)
+{
+	SetExpansionRecursive(InItem, false, bShouldBeExpanded);
+}
+
+void SReferenceSkeletonTree::SetExpansionRecursive(
+	TSharedPtr<FBoneElement> InElement,
+	bool bTowardsParent,
+    bool bShouldBeExpanded)
+{
+	TreeView->SetItemExpansion(InElement, bShouldBeExpanded);
+    
+    if (bTowardsParent)
+    {
+    	if (InElement->Parent.Get())
+    	{
+    		SetExpansionRecursive(InElement->Parent, bTowardsParent, bShouldBeExpanded);
+    	}
+    }
+    else
+    {
+    	for (int32 ChildIndex = 0; ChildIndex < InElement->Children.Num(); ++ChildIndex)
+    	{
+    		SetExpansionRecursive(InElement->Children[ChildIndex], bTowardsParent, bShouldBeExpanded);
+    	}
+    }
+}
+
+FReply SReferenceSkeletonTree::OnKeyDown(const FGeometry& MyGeometry, const FKeyEvent& InKeyEvent)
+{
+	if (!Modifier.IsValid())
+	{
+		return FReply::Handled();
+	}
+	
+	if (CommandList.IsValid() && CommandList->ProcessCommandBindings(InKeyEvent))
+	{
+		return FReply::Handled();
+	}
+	
+	return FReply::Unhandled();
+}
+
+FReply SReferenceSkeletonTree::OnDragDetected(const FGeometry& MyGeometry, const FPointerEvent& MouseEvent)
+{
+	const TArray<TSharedPtr<FBoneElement>> BoneElements = GetSelectedItems();
+	if (BoneElements.IsEmpty())
+	{
+		return FReply::Unhandled();
+	}
+	
+	if (MouseEvent.IsMouseButtonDown(EKeys::LeftMouseButton))
+	{
+		const TSharedRef<FBoneItemDragDropOp> DragDropOp = FBoneItemDragDropOp::New(BoneElements);
+		return FReply::Handled().BeginDragDrop(DragDropOp);
+	}
+	
+	return FReply::Unhandled();
+}
+
+TOptional<EItemDropZone> SReferenceSkeletonTree::OnCanAcceptDrop(const FDragDropEvent& DragDropEvent, EItemDropZone DropZone, TSharedPtr<FBoneElement> TargetItem)
+{
+	if (bIsReadOnly)
+	{
+		return TOptional<EItemDropZone>();
+		
+	}
+	TOptional<EItemDropZone> ReturnedDropZone;
+
+	const TSharedPtr<FBoneItemDragDropOp> DragDropOp = DragDropEvent.GetOperationAs<FBoneItemDragDropOp>();
+	if (DragDropOp.IsValid())
+	{
+		ReturnedDropZone = EItemDropZone::OntoItem;
+	}
+
+	return ReturnedDropZone;
+}
+
+FReply SReferenceSkeletonTree::OnAcceptDrop(const FDragDropEvent& DragDropEvent, EItemDropZone DropZone, TSharedPtr<FBoneElement> TargetItem)
+{
+	const TSharedPtr<FBoneItemDragDropOp> DragDropOp = DragDropEvent.GetOperationAs<FBoneItemDragDropOp>();
+	if (!DragDropOp.IsValid())
+	{
+		return FReply::Unhandled();
+	}
+
+	if (!Modifier.IsValid())
+	{
+		return FReply::Handled();
+	}
+	
+	TArray<FName> BoneNames;
+	BoneNames.Reserve(DragDropOp->Elements.Num());
+	for (const TWeakPtr<FBoneElement>& BoneElement: DragDropOp->Elements)
+	{
+		if (BoneElement.IsValid())
+		{
+			BoneNames.Add(BoneElement.Pin()->BoneName);
+		}
+	}
+	const FName ParentName = TargetItem->BoneName;
+
+	BeginChange();
+	
+	const bool bParented = Modifier->ParentBones(BoneNames, {ParentName});
+	if (bParented)
+	{
+		RefreshTreeView(true);
+		SelectItemFromNames(BoneNames);
+		if (Notifier.IsValid())
+		{
+			Notifier->Notify( BoneNames, ESkeletalMeshNotifyType::HierarchyChanged);
+			Notifier->Notify( BoneNames, ESkeletalMeshNotifyType::BonesSelected);
+		}
+		EndChange();
+	}
+	else
+	{
+		CancelChange();
+	}
+	
+	return FReply::Handled();
+}
+
+FReferenceSkeletonWidgetNotifier::FReferenceSkeletonWidgetNotifier(TSharedRef<SReferenceSkeletonTree> InWidget)
+	: ISkeletalMeshNotifier()
+	, Tree(InWidget)
+{}
+
+void FReferenceSkeletonWidgetNotifier::HandleNotification(const TArray<FName>& InBoneNames, const ESkeletalMeshNotifyType InNotifyType)
+{
+	if (Notifying() || !Tree.IsValid())
+	{
+		return;
+	}
+
+	TSharedPtr<SReferenceSkeletonTree> TreePtr = Tree.Pin();
+	switch (InNotifyType)
+	{
+		case ESkeletalMeshNotifyType::BonesAdded:
+			TreePtr->RefreshTreeView();
+			break;
+		case ESkeletalMeshNotifyType::BonesRemoved:
+			TreePtr->RefreshTreeView();
+			break;
+		case ESkeletalMeshNotifyType::BonesMoved:
+			break;
+		case ESkeletalMeshNotifyType::BonesSelected:
+			{
+				static constexpr bool bFrameSelection = true;
+				TreePtr->SelectItemFromNames(InBoneNames, bFrameSelection);
+				break;
+			}
+		case ESkeletalMeshNotifyType::BonesRenamed:
+			TreePtr->RefreshTreeView();
+			TreePtr->SelectItemFromNames(InBoneNames);
+			break;
+		case ESkeletalMeshNotifyType::HierarchyChanged:
+			{
+				TArray<FName> BoneNames;
+				TreePtr->GetSelectedBoneNames(BoneNames);
+				TreePtr->RefreshTreeView();
+				TreePtr->SelectItemFromNames(BoneNames);
+			}
+			break;
+	}
+}
+
+void SReferenceSkeletonTree::BeginChange()
+{
+	ensure( ActiveChange == nullptr );
+	ActiveChange = MakeUnique<ReferenceSkeletonTreeLocals::FSkeletonModifierChange>(Modifier.Get()); 
+}
+
+void SReferenceSkeletonTree::EndChange()
+{
+	if (!ActiveChange.IsValid())
+	{
+		return;
+	}
+
+	ActiveChange->StoreSkeleton(Modifier.Get());
+	
+	GEditor->BeginTransaction(LOCTEXT("ModifyRefSkeleton", "Modify Reference Skeleton"));
+	GUndo->StoreUndo(Modifier.Get(), MoveTemp(ActiveChange));
+	GEditor->EndTransaction();
+}
+
+void SReferenceSkeletonTree::CancelChange()
+{
+	ActiveChange.Reset();
+}
+
+#undef LOCTEXT_NAMESPACE
\ No newline at end of file
diff --git a/Engine/Plugins/Runtime/MeshModelingToolset/Source/MeshModelingToolsEditorOnly/Private/SkeletalMesh/SkeletalMeshEditingInterface.cpp b/Engine/Plugins/Runtime/MeshModelingToolset/Source/MeshModelingToolsEditorOnly/Private/SkeletalMesh/SkeletalMeshEditingInterface.cpp
--- a/Engine/Plugins/Runtime/MeshModelingToolset/Source/MeshModelingToolsEditorOnly/Private/SkeletalMesh/SkeletalMeshEditingInterface.cpp
+++ b/Engine/Plugins/Runtime/MeshModelingToolset/Source/MeshModelingToolsEditorOnly/Private/SkeletalMesh/SkeletalMeshEditingInterface.cpp
@@ -21,7 +21,12 @@ bool ISkeletalMeshEditingInterface::NeedsNotification() const
 	return Notifier && Notifier->Delegate().IsBound(); 
 }
 
-TWeakObjectPtr<USkeletonModifier> ISkeletalMeshEditingInterface::GetModifier() const
+bool ISkeletalMeshEditingInterface::CanModifySkeleton() const
+{
+	return false;
+}
+
+TSharedPtr<SWidget> ISkeletalMeshEditingInterface::GetCustomSkeletonTreeWidget(TSharedPtr<ISkeletalMeshEditorBinding> EditorBinding)
 {
 	return nullptr;
 }
diff --git a/Engine/Plugins/Runtime/MeshModelingToolset/Source/MeshModelingToolsEditorOnly/Private/SkeletalMesh/SkeletonEditingTool.cpp b/Engine/Plugins/Runtime/MeshModelingToolset/Source/MeshModelingToolsEditorOnly/Private/SkeletalMesh/SkeletonEditingTool.cpp
--- a/Engine/Plugins/Runtime/MeshModelingToolset/Source/MeshModelingToolsEditorOnly/Private/SkeletalMesh/SkeletonEditingTool.cpp
+++ b/Engine/Plugins/Runtime/MeshModelingToolset/Source/MeshModelingToolsEditorOnly/Private/SkeletalMesh/SkeletonEditingTool.cpp
@@ -22,6 +22,7 @@
 #include "Selection/PolygonSelectionMechanic.h"
 
 #include "SkeletalMesh/SkeletonTransformProxy.h"
+#include "SkeletalMesh/SReferenceSkeletonTree.h"
 #include "TargetInterfaces/PrimitiveComponentBackedTarget.h"
 #include "TargetInterfaces/SkeletonProvider.h"
 #include "TargetInterfaces/SkeletonCommitter.h"
@@ -66,6 +67,35 @@ void FRefSkeletonChange::Revert(UObject* Object)
 
 }
 
+
+FSkeletonEditingToolCommands::FSkeletonEditingToolCommands()
+	: TCommands<FSkeletonEditingToolCommands>
+	(
+		TEXT("SkeletonEditingTool"),
+		NSLOCTEXT("Contexts", "SkeletonEditingToolCommands", "Skeleton Editing Tool Commands"),
+		NAME_None,
+		FAppStyle::GetAppStyleSetName()
+	)
+{
+}
+
+void FSkeletonEditingToolCommands::RegisterCommands()
+{
+	UI_COMMAND(NewBone, "New Bone", "Add new Bone.", EUserInterfaceActionType::Button, FInputChord(EKeys::N, EModifierKey::Control));
+	UI_COMMAND(RemoveBone, "Remove Bone", "Remove selected bone(s).", EUserInterfaceActionType::Button, FInputChord(EKeys::Delete), FInputChord(EKeys::BackSpace));
+	UI_COMMAND(UnParentBone, "Unparent Bone", "Unparent selected bone(s).", EUserInterfaceActionType::Button, FInputChord(EKeys::P, EModifierKey::Shift));
+
+	UI_COMMAND(RenameBone, "Rename Bone", "Rename the selected bone.", EUserInterfaceActionType::Button, FInputChord(EKeys::F2));
+	UI_COMMAND(SearchAndReplace, "Search and Replace", "Search and Replace the selected bone(s) name(s).", EUserInterfaceActionType::Button, FInputChord());
+	UI_COMMAND(AddPrefix, "Add Prefix", "Add a prefix to the selected bone(s) name(s).", EUserInterfaceActionType::Button, FInputChord());
+	UI_COMMAND(AddSuffix, "Add Suffix", "Add a suffix to the selected bone(s) name(s).", EUserInterfaceActionType::Button, FInputChord());
+	
+	UI_COMMAND(CopyBones, "Copy Bone(s)", "Copy selected bone(s).", EUserInterfaceActionType::Button, FInputChord(EKeys::C, EModifierKey::Control));
+	UI_COMMAND(PasteBones, "Paste Bone(s)", "Paste selected bone(s).", EUserInterfaceActionType::Button, FInputChord(EKeys::V, EModifierKey::Control));
+	UI_COMMAND(DuplicateBones, "Duplicate Bone(s)", "Duplicate selected bone(s).", EUserInterfaceActionType::Button, FInputChord(EKeys::D, EModifierKey::Control));
+}
+
+
 /**
  * USkeletonEditingToolBuilder
  */
@@ -408,6 +438,8 @@ void USkeletonEditingTool::UpdateGizmo() const
 
 void USkeletonEditingTool::Shutdown(EToolShutdownType ShutdownType)
 {
+	SkeletonTreeNotifierBindScope.Reset();
+	
 	if (ShutdownType == EToolShutdownType::Accept)
 	{
 		GetToolManager()->BeginUndoTransaction(LOCTEXT("SkeletonEditingTool", "Commit Skeleton Editing"));
@@ -1433,6 +1465,28 @@ void USkeletonEditingTool::OnUpdateModifierState(int InModifierID, bool bIsOn)
 	}
 }
 
+bool USkeletonEditingTool::CanModifySkeleton() const
+{
+	return true;
+}
+
+TSharedPtr<SWidget> USkeletonEditingTool::GetCustomSkeletonTreeWidget(TSharedPtr<ISkeletalMeshEditorBinding> EditorBinding)
+{
+	TSharedPtr<SReferenceSkeletonTree> ToolSkeletonTree =
+		SNew(SReferenceSkeletonTree)
+		.Modifier(Modifier);
+
+	SkeletonTreeNotifierBindScope.Reset(
+		new FSkeletalMeshNotifierBindScope(
+			EditorBinding->GetNotifier(),
+			ToolSkeletonTree->GetNotifier()));	
+		
+	TArray<FName> SelectedBones = EditorBinding->GetSelectedBones();
+	ToolSkeletonTree->GetNotifier()->HandleNotification(SelectedBones, ESkeletalMeshNotifyType::BonesSelected);		
+	
+	return ToolSkeletonTree;
+}
+
 EEditingOperation USkeletonEditingTool::GetOperation() const
 {
 	return Operation;	
@@ -1679,11 +1733,6 @@ FInputRayHit USkeletonEditingTool::CanBeginClickDragSequence(const FInputDeviceR
 	return InvalidRayHit;
 }
 
-TWeakObjectPtr<USkeletonModifier> USkeletonEditingTool::GetModifier() const
-{
-	return Modifier;
-}
-
 void USkeletonEditingTool::HandleSkeletalMeshModified(const TArray<FName>& InBoneNames, const ESkeletalMeshNotifyType InNotifyType)
 {
 	if (GetNotifier()->Notifying())
diff --git a/Engine/Plugins/Runtime/MeshModelingToolset/Source/MeshModelingToolsEditorOnly/Private/SkeletalMesh/SkinWeightsPaintTool.cpp b/Engine/Plugins/Runtime/MeshModelingToolset/Source/MeshModelingToolsEditorOnly/Private/SkeletalMesh/SkinWeightsPaintTool.cpp
--- a/Engine/Plugins/Runtime/MeshModelingToolset/Source/MeshModelingToolsEditorOnly/Private/SkeletalMesh/SkinWeightsPaintTool.cpp
+++ b/Engine/Plugins/Runtime/MeshModelingToolset/Source/MeshModelingToolsEditorOnly/Private/SkeletalMesh/SkinWeightsPaintTool.cpp
@@ -41,6 +41,7 @@
 #include "TargetInterfaces/DynamicMeshCommitter.h"
 #include "TargetInterfaces/MeshDescriptionProvider.h"
 #include "SkeletalMeshAttributes.h"
+#include "SkeletonModifier.h"
 
 #include "Operations/SmoothBoneWeights.h"
 #include "Operations/TransferBoneWeights.h"
@@ -59,6 +60,13 @@
 #include "TargetInterfaces/SkeletonProvider.h"
 #include "ToolTargets/DynamicMeshComponentToolTarget.h"
 
+#include "Styling/AppStyle.h"
+#include "Widgets/Input/SCheckBox.h"
+#include "Widgets/Images/SImage.h"
+#include "Misc/MessageDialog.h"
+#include "SkeletalMesh/SReferenceSkeletonTree.h"
+
+
 #include UE_INLINE_GENERATED_CPP_BY_NAME(SkinWeightsPaintTool)
 
 #define LOCTEXT_NAMESPACE "USkinWeightsPaintTool"
@@ -787,6 +795,11 @@ void UWeightToolTransferManager::ApplyTranferredWeightsAsTransaction(
 	const TArray<int32>& InVertexSubset,
 	const FDynamicMesh3& InTargetMesh)
 {
+	if (!WeightTool->ShouldProceedWithAllBonesUnlocked())
+	{
+		return;
+	}
+	
 	check(InTransferredSkinWeights);
 	
 	// weight edits for transaction
@@ -1179,15 +1192,16 @@ void FSkinToolDeformer::SetVertexNeedsUpdated(int32 VertexIndex)
 }
 
 void FSkinToolWeights::InitializeSkinWeights(
-	USkinWeightsPaintTool* Tool,
+	USkinWeightsPaintTool* InTool,
 	const FDynamicMesh3& Mesh
 	)
 {
 	static constexpr int32 RootBoneIndex = 0;
 	static constexpr float FullWeight = 1.f;
 
+	Tool = InTool;
 	ISkeletonProvider* SkeletonProvider = CastChecked<ISkeletonProvider>(Tool->GetTarget());
-	
+
 	// initialize deformer data
 	Deformer.Initialize(SkeletonProvider->GetSkeleton(), Tool->GetComponentSpaceBoneTransforms(), Mesh);
 
@@ -1279,16 +1293,21 @@ void FSkinToolWeights::CreateWeightEditForVertex(
 	float NewWeightValue,
 	FMultiBoneWeightEdits& WeightEdits)
 {
+	const TSet<int32>& LockedBoneIndices = Tool->GetFinalLockedBoneIndices();
 	// this operation should never prune weights
 	constexpr bool bPruneInfluence = false;
-	
-	// clamp new weight
-	NewWeightValue = FMath::Clamp(NewWeightValue, 0.0f, 1.0f);
+
+	if (LockedBoneIndices.Contains(BoneToHoldIndex))
+	{
+		return;
+	}
+
 
 	// calculate the sum of all the weights on this vertex (not including the one we currently applied)
 	TArray<int32> RecordedBonesOnVertex;
 	TArray<float> ValuesToNormalize;
 	float Total = 0.0f;
+	float TotalLockedWeights = 0.0f;
 	const VertexWeights& VertexData = PreChangeWeights[VertexID];
 	for (const FVertexBoneWeight& VertexBoneData : VertexData)
 	{
@@ -1301,17 +1320,28 @@ void FSkinToolWeights::CreateWeightEditForVertex(
 		{
 			continue;
 		}
-		
-		RecordedBonesOnVertex.Add(VertexBoneData.BoneID);
-		ValuesToNormalize.Add(VertexBoneData.Weight);
-		Total += VertexBoneData.Weight;
+
+		if (!LockedBoneIndices.Contains(VertexBoneData.BoneID))
+		{
+			RecordedBonesOnVertex.Add(VertexBoneData.BoneID);
+			ValuesToNormalize.Add(VertexBoneData.Weight);
+			Total += VertexBoneData.Weight;
+		}
+		else
+		{
+			TotalLockedWeights += VertexBoneData.Weight;
+		}
 	}
 
+	float MaxNewWeight = 1.f - TotalLockedWeights;
+	// clamp new weight
+	NewWeightValue = FMath::Clamp(NewWeightValue, 0.0f, MaxNewWeight);
+	
 	// assigning full weight to this vertex?
-	if (FMath::IsNearlyEqual(NewWeightValue, 1.f))
+	if (FMath::IsNearlyEqual(NewWeightValue, MaxNewWeight))
 	{
 		// in this case normalization is trivial, just assign the full weight directly and zero all others
-		constexpr float FullWeight = 1.0f;
+		const float FullWeight = MaxNewWeight;
 		WeightEdits.MergeSingleEdit(BoneToHoldIndex, VertexID, FullWeight, bPruneInfluence, PreChangeWeights);
 
 		// zero all others
@@ -1346,7 +1376,7 @@ void FSkinToolWeights::CreateWeightEditForVertex(
 			// In this case, we evenly split the remaining weight among the recorded influences
 
 			// distribute remaining weight evenly over other recorded influences
-			const float WeightToDistribute = (1.0f - NewWeightValue) / RecordedBonesOnVertex.Num();
+			const float WeightToDistribute = (MaxNewWeight - NewWeightValue) / RecordedBonesOnVertex.Num();
 			for (int32 i=0; i<ValuesToNormalize.Num(); ++i)
 			{
 				const int32 BoneIndex = RecordedBonesOnVertex[i];
@@ -1366,19 +1396,19 @@ void FSkinToolWeights::CreateWeightEditForVertex(
 				
 			// assign remaining weight to the parent
 			const int32 ParentBoneIndex = GetParentBoneToWeightTo(BoneToHoldIndex);
-			if (ParentBoneIndex == BoneToHoldIndex)
+			if (ParentBoneIndex == BoneToHoldIndex || Tool->GetFinalLockedBoneIndices().Contains(ParentBoneIndex))
 			{
 				// was unable to find parent OR child bone!
 				// this could only happen if user is trying to remove weight from the ONLY bone in the whole skeleton
 				// in this case just assign the full weight to the bone (there's no other valid configuration)
 				// this is a "do nothing" operation, but at least it generates an undo transaction to let user know the input was received
-				constexpr float FullWeight = 1.0f;
+				const float FullWeight = MaxNewWeight;
 				WeightEdits.MergeSingleEdit(BoneToHoldIndex, VertexID, FullWeight, bPruneInfluence, PreChangeWeights);
 			}
 			else
 			{
 				// assign remaining weight to parent
-				const float NewParentWeight = 1.0f - NewWeightValue;
+				const float NewParentWeight = MaxNewWeight - NewWeightValue;
 				WeightEdits.MergeSingleEdit(ParentBoneIndex, VertexID, NewParentWeight, bPruneInfluence, PreChangeWeights);
 				// and assign user requested weight to the current bone
 				WeightEdits.MergeSingleEdit(BoneToHoldIndex, VertexID, NewWeightValue, bPruneInfluence, PreChangeWeights);
@@ -1392,7 +1422,7 @@ void FSkinToolWeights::CreateWeightEditForVertex(
 	// and split the remainder amongst the recorded influences
 	{
 		// calculate amount we have to spread across the other bones affecting this vertex
-		const float AvailableTotal = 1.0f - NewWeightValue;
+		const float AvailableTotal = MaxNewWeight - NewWeightValue;
 
 		// normalize weights into available space not set by current bone
 		for (int32 i=0; i<ValuesToNormalize.Num(); ++i)
@@ -1805,6 +1835,31 @@ void FMeshSkinWeightsChange::StoreMultipleWeightEdits(const FMultiBoneWeightEdit
 	}
 }
 
+FSkinWeightsPaintToolCommands::FSkinWeightsPaintToolCommands()
+	: TCommands<FSkinWeightsPaintToolCommands>
+	(
+		TEXT("SkinWeightsPaintTool"),
+		NSLOCTEXT("Contexts", "SkinWeightsPaintToolCommands", "Skin Weight Paint Tool Commands"),
+		NAME_None,
+		FAppStyle::GetAppStyleSetName()
+	)
+{
+}
+
+void FSkinWeightsPaintToolCommands::RegisterCommands()
+{
+	UI_COMMAND(IsolateSelectedBones, "Isolate Selected (Unlock)", "Isolate selected bones by locking all other bones, Skin weights for locked bones cannot be changed", EUserInterfaceActionType::Button, FInputChord());
+	
+	UI_COMMAND(LockSelectedBones, "Lock Selected", "Lock selected bones, Skin weights for locked bones cannot be changed", EUserInterfaceActionType::Button, FInputChord());	
+	UI_COMMAND(UnlockSelectedBones, "Unlock Selected", "Unlock selected bones, Skin weights for locked bones cannot be changed", EUserInterfaceActionType::Button, FInputChord());
+	
+	UI_COMMAND(LockAllBones, "Lock All", "Lock all bones, Skin weights for locked bones cannot be changed", EUserInterfaceActionType::Button, FInputChord());	
+	UI_COMMAND(UnlockAllBones, "Unlock All", "Unlock all bones, Skin weights for locked bones cannot be changed", EUserInterfaceActionType::Button, FInputChord());
+	
+	UI_COMMAND(BypassLocks, "Bypass Locks", "Disable skin weight locking, existing locks are preserved but disabled", EUserInterfaceActionType::ToggleButton, FInputChord());
+	
+}
+
 /*
  * ToolBuilder
  */
@@ -1891,6 +1946,8 @@ void USkinWeightsPaintTool::Setup()
 	auto OnSelectionChangedLambda = [this](){OnSelectionChanged.Broadcast();};
 	MeshSelector->InitialSetup(TargetWorld.Get(), this, OnSelectionChangedLambda);
 
+	InitializeSkeletonReader();
+	
 	// run all initialization for mesh/weights
 	UpdateCurrentlyEditedMesh(EditedMesh);
 
@@ -2843,7 +2900,7 @@ void USkinWeightsPaintTool::CreateWeightEditsToRelaxVertices(
 			const float UseFalloff = (VertexFalloffs.IsValidIndex(VertexIndex) ? VertexFalloffs[VertexIndex] * Strength : Strength) * PercentPerIteration;
 
 			TMap<int32, float> FinalWeights;
-			const bool bSmoothSuccess = SmoothWeightsOp->SmoothWeightsAtVertex(VertexID, UseFalloff, FinalWeights);
+			const bool bSmoothSuccess = SmoothWeightsOp->SmoothWeightsAtVertex(VertexID, UseFalloff, GetFinalLockedBoneIndices() ,FinalWeights);
 			if (!ensure(bSmoothSuccess))
 			{
 				continue;
@@ -2938,6 +2995,11 @@ BoneIndex USkinWeightsPaintTool::GetBoneIndexFromName(const FName BoneName) cons
 	return Found ? *Found : INDEX_NONE;
 }
 
+bool USkinWeightsPaintTool::HasSelectedBone() const
+{
+	return !SelectedBoneIndices.IsEmpty();
+}
+
 void USkinWeightsPaintTool::SetFocusInViewport() const
 {
 	if (PersonaModeManagerContext.IsValid())
@@ -2948,6 +3010,7 @@ void USkinWeightsPaintTool::SetFocusInViewport() const
 
 void USkinWeightsPaintTool::OnShutdown(EToolShutdownType ShutdownType)
 {
+	SkeletonTreeNotifierBindScope.Reset();
 	// shutdown must be performed on full mesh, so end isolated selection
 	if (SelectionIsolator)
 	{
@@ -3021,6 +3084,27 @@ USkinWeightsPaintToolProperties* USkinWeightsPaintTool::GetWeightToolProperties(
 }
 
 
+const TSet<int32>& USkinWeightsPaintTool::GetFinalLockedBoneIndices()
+{
+	if (IsBypassingAllLocks())
+	{
+		static TSet<int32> EmptySet;
+		return EmptySet;
+	}
+	
+	return GetRawLockedBoneIndices();	
+}
+
+const TSet<int32>& USkinWeightsPaintTool::GetRawLockedBoneIndices()
+{
+	return LockedBoneIndices;
+}
+
+bool USkinWeightsPaintTool::IsBypassingAllLocks() const
+{
+	return bBypassLocks;
+}
+
 void USkinWeightsPaintTool::BeginChange()
 {
 	const FName SkinProfile = WeightToolProperties->GetActiveSkinWeightProfile();
@@ -3230,6 +3314,11 @@ bool FSkinMirrorData::IsPointOnTargetMirrorSide(const FVector& InPoint) const
 
 void USkinWeightsPaintTool::MirrorWeights(EAxis::Type Axis, EMirrorDirection Direction)
 {
+	if (!ShouldProceedWithAllBonesUnlocked())
+	{
+		return;
+	}
+	
 	check(Axis != EAxis::None);
 	
 	// get all ref pose vertices
@@ -3403,6 +3492,11 @@ void USkinWeightsPaintTool::EditWeightsOnVertices(
 
 void USkinWeightsPaintTool::PruneWeights(float Threshold, const TArray<BoneIndex>& BonesToPrune)
 {
+	if (!ShouldProceedWithAllBonesUnlocked())
+	{
+		return;
+	}
+	
 	// set weights below the given threshold to zero AND remove them as a recorded influence on that vertex
 	FMultiBoneWeightEdits WeightEditsFromPrune;
 	const TArray<VertexIndex>& VerticesToPrune = MeshSelector->GetSelectedVertices();
@@ -3477,23 +3571,27 @@ void USkinWeightsPaintTool::AverageWeights(const float Strength)
 	
 	// get vertices to edit weights on
 	const TArray<VertexIndex>& VerticesToAverage = MeshSelector->GetSelectedVertices();
-	TMap<BoneIndex, float> AveragedWeights;
-	AccumulateWeights(Weights.PreChangeWeights, VerticesToAverage, AveragedWeights);
-	TruncateWeightMap(AveragedWeights);
-	NormalizeWeightMap(AveragedWeights);
-
+	TMap<BoneIndex, float> AccumulatedWeights;	
+	AccumulateWeights(Weights.PreChangeWeights, VerticesToAverage, GetFinalLockedBoneIndices(), AccumulatedWeights);
 	// store weight edits to apply averaging to selected vertices
-	FMultiBoneWeightEdits WeightEditsFromAveraging;
-	
-	// FULLY apply averaged weights to vertices if strength is 1.0
-	if (FMath::IsNearlyEqual(Strength, 1.0f))
+	FMultiBoneWeightEdits WeightEditsFromAveraging;	
+	for (const VertexIndex VertexID : VerticesToAverage)
 	{
-		for (const VertexIndex VertexID : VerticesToAverage)
+		float TotalLockedWeight = 0.0f;
+		int32 NumLockedBones = 0;
+		GetLockedWeight(Weights.PreChangeWeights, VertexID, GetFinalLockedBoneIndices(), TotalLockedWeight, NumLockedBones);
+			
+		TMap<BoneIndex, float> AveragedWeights = AccumulatedWeights;	
+		TruncateWeightMap(AveragedWeights, NumLockedBones);
+		NormalizeWeightMap(AveragedWeights, 1.f - TotalLockedWeight);	
+			
+		// FULLY apply averaged weights to vertices if strength is 1.0
+		if (FMath::IsNearlyEqual(Strength, 1.0f))
 		{
 			// remove influences not a part of the average results
 			for (const FVertexBoneWeight& BoneWeight : Weights.PreChangeWeights[VertexID])
 			{
-				if (!AveragedWeights.Contains(BoneWeight.BoneID))
+				if (!AveragedWeights.Contains(BoneWeight.BoneID) && !GetFinalLockedBoneIndices().Contains(BoneWeight.BoneID))
 				{
 					constexpr bool bPruneInfluence = false;
 					constexpr float NewWeight = 0.f;
@@ -3510,21 +3608,22 @@ void USkinWeightsPaintTool::AverageWeights(const float Strength)
 				WeightEditsFromAveraging.MergeSingleEdit(IndexOfBone, VertexID, NewWeight, bPruneInfluence, Weights.PreChangeWeights);
 			}
 		}
-	}
-	else
-	{
-		// blend averaged weight with the existing weight based on the strength value
-		const float OldWeightStrength = 1.0f - Strength;
-		const float NewWeightStrength = Strength;
-		for (const VertexIndex VertexID : VerticesToAverage)
+		else
 		{
+			// blend averaged weight with the existing weight based on the strength value
+			const float OldWeightStrength = 1.0f - Strength;
+			const float NewWeightStrength = Strength;
+
 			// storage for final blended weights on this vertex
 			TMap<BoneIndex, float> BlendedWeights;
 
 			// scale the existing weights by OldWeightStrength
 			for (const FVertexBoneWeight& BoneWeight : Weights.PreChangeWeights[VertexID])
 			{
-				BlendedWeights.Add(BoneWeight.BoneID, BoneWeight.Weight * OldWeightStrength);
+				if (!GetFinalLockedBoneIndices().Contains(BoneWeight.BoneID))
+				{
+					BlendedWeights.Add(BoneWeight.BoneID, BoneWeight.Weight * OldWeightStrength);
+				}
 			}
 			
 			// accumulate existing weights with the scaled averaged weights
@@ -3541,8 +3640,8 @@ void USkinWeightsPaintTool::AverageWeights(const float Strength)
 			}
 
 			// enforce max influences and normalize
-			TruncateWeightMap(BlendedWeights);
-			NormalizeWeightMap(BlendedWeights);
+			TruncateWeightMap(BlendedWeights, NumLockedBones);
+			NormalizeWeightMap(BlendedWeights, 1.0f - TotalLockedWeight);
 			
 			// apply blended weights to this vertex
 			for (const TTuple<BoneIndex, float>& BlendedWeight : BlendedWeights)
@@ -3597,6 +3696,11 @@ void USkinWeightsPaintTool::HammerWeights()
 		return;
 	}
 
+	if (!ShouldProceedWithAllBonesUnlocked())
+	{
+		return;
+	}
+	
 	// reset mesh to ref pose so that Dijkstra path lengths are not deformed
 	Weights.Deformer.SetToRefPose(this);
 	
@@ -3677,6 +3781,11 @@ void USkinWeightsPaintTool::ClampInfluences(const int32 MaxInfluences)
 		ShowEditorMessage(ELogVerbosity::Error, NotificationText);
 		return;
 	}
+
+	if (!ShouldProceedWithAllBonesUnlocked())
+	{
+		return;
+	}
 	
 	// create set of weight edits that clamp the smallest influences
 	FMultiBoneWeightEdits ClampedWeightEdits;
@@ -3750,7 +3859,7 @@ void USkinWeightsPaintTool::CopyWeights()
 
 	// get a map of the average weight of all the selected vertices (accumulating and normalizing produces an average)
 	TMap<BoneIndex, float> WeightsToCopy;
-	AccumulateWeights(Weights.PreChangeWeights, VerticesToCopy, WeightsToCopy);
+	AccumulateWeights(Weights.PreChangeWeights, VerticesToCopy, {}, WeightsToCopy);
 	TruncateWeightMap(WeightsToCopy);
 	NormalizeWeightMap(WeightsToCopy);
 
@@ -3804,6 +3913,11 @@ void USkinWeightsPaintTool::PasteWeights()
 		return;
 	}
 
+	if (!ShouldProceedWithAllBonesUnlocked())
+	{
+		return;
+	}
+
 	// deserialize the string into JSON and parse it as name/weight pairs
 	TMap<BoneIndex, float> LoadedWeights;
 	const FString JsonString = ClipboardContent.RightChop(CopyPasteWeightsIdentifier.Len());
@@ -3898,7 +4012,7 @@ void USkinWeightsPaintTool::PasteWeights()
 	ShowEditorMessage(ELogVerbosity::Log, NotificationText);
 }
 
-void USkinWeightsPaintTool::TruncateWeightMap(TMap<BoneIndex, float>& InOutWeights)
+void USkinWeightsPaintTool::TruncateWeightMap(TMap<BoneIndex, float>& InOutWeights, int32 NumLockedBones)
 {
 	// sort influences by total weight
 	InOutWeights.ValueSort([](const float& A, const float& B)
@@ -3910,7 +4024,7 @@ void USkinWeightsPaintTool::TruncateWeightMap(TMap<BoneIndex, float>& InOutWeigh
 	int32 Index = 0;
 	for (TMap<BoneIndex, float>::TIterator It(InOutWeights); It; ++It)
 	{
-		if (Index >= MAX_TOTAL_INFLUENCES)
+		if (Index >= MAX_TOTAL_INFLUENCES - NumLockedBones)
 		{
 			It.RemoveCurrent();
 		}
@@ -3921,7 +4035,7 @@ void USkinWeightsPaintTool::TruncateWeightMap(TMap<BoneIndex, float>& InOutWeigh
 	}
 }
 
-void USkinWeightsPaintTool::NormalizeWeightMap(TMap<BoneIndex, float>& InOutWeights)
+void USkinWeightsPaintTool::NormalizeWeightMap(TMap<BoneIndex, float>& InOutWeights, float ExpectedSum)
 {
 	// normalize remaining influences
 	float TotalWeight = 0.f;
@@ -3933,20 +4047,45 @@ void USkinWeightsPaintTool::NormalizeWeightMap(TMap<BoneIndex, float>& InOutWeig
 	for (TTuple<BoneIndex, float>& Weight : InOutWeights)
 	{
 		Weight.Value /= TotalWeight > SMALL_NUMBER ? TotalWeight : 1.f;
+		Weight.Value *= ExpectedSum;
 	}
 }
 
 void USkinWeightsPaintTool::AccumulateWeights(
 	const TArray<SkinPaintTool::VertexWeights>& AllWeights,
 	const TArray<VertexIndex>& VerticesToAccumulate,
-	TMap<BoneIndex, float>& OutWeights)
+	const TSet<BoneIndex>& ExcludedBones,
+	TMap<BoneIndex, float>& OutAccumulatedWeights)
 {
 	for (const VertexIndex VertexID : VerticesToAccumulate)
 	{
 		for (const FVertexBoneWeight& BoneWeight : AllWeights[VertexID])
 		{
-			float& AccumulatedWeight = OutWeights.FindOrAdd(BoneWeight.BoneID);
-			AccumulatedWeight += BoneWeight.Weight;
+			if (!ExcludedBones.Contains(BoneWeight.BoneID))
+			{
+				float& AccumulatedWeight = OutAccumulatedWeights.FindOrAdd(BoneWeight.BoneID);
+				AccumulatedWeight += BoneWeight.Weight;
+			}
+		}
+	}
+}
+
+void USkinWeightsPaintTool::GetLockedWeight(
+	const TArray<VertexWeights>& AllWeights,
+	VertexIndex InVertexID,
+	const TSet<BoneIndex>& LockedBones,
+	float& OutTotalLockedWeight,
+	int32& OutNumLockedBones)
+{
+	OutTotalLockedWeight = 0;
+	OutNumLockedBones = 0;
+	
+	for (const FVertexBoneWeight& BoneWeight : AllWeights[InVertexID])
+	{
+		if (LockedBones.Contains(BoneWeight.BoneID))
+		{
+			OutTotalLockedWeight += BoneWeight.Weight;
+			OutNumLockedBones++;
 		}
 	}
 }
@@ -3988,6 +4127,234 @@ void USkinWeightsPaintTool::HandleSkeletalMeshModified(const TArray<FName>& InBo
 	}
 }
 
+TSharedPtr<SWidget> USkinWeightsPaintTool::GetCustomSkeletonTreeWidget(TSharedPtr<ISkeletalMeshEditorBinding> EditorBinding)
+{
+	if (!SkeletonReader)
+	{
+		return nullptr;
+	}
+	
+	FRefSkeletonTreeDelegates Delegates;
+	Delegates.OnGetRowWidget.BindUObject(this,&USkinWeightsPaintTool::OnGetLockWidget);
+	Delegates.OnExtendCommandList.BindLambda([this]( TSharedPtr<FUICommandList> InOutCommandList)
+		{
+			FSkinWeightsPaintToolCommands Commands = FSkinWeightsPaintToolCommands::Get();
+			InOutCommandList->MapAction( Commands.IsolateSelectedBones,
+				FExecuteAction::CreateUObject(this, &USkinWeightsPaintTool::IsolateSelectedBones),
+				FCanExecuteAction::CreateUObject(this, &USkinWeightsPaintTool::HasSelectedBone));
+			
+			InOutCommandList->MapAction( Commands.LockSelectedBones,
+				FExecuteAction::CreateUObject(this, &USkinWeightsPaintTool::ToggleLockSelectedBones, true),
+				FCanExecuteAction::CreateUObject(this, &USkinWeightsPaintTool::HasSelectedBone));
+			InOutCommandList->MapAction( Commands.UnlockSelectedBones,
+				FExecuteAction::CreateUObject(this, &USkinWeightsPaintTool::ToggleLockSelectedBones, false),
+				FCanExecuteAction::CreateUObject(this, &USkinWeightsPaintTool::HasSelectedBone));
+			
+			InOutCommandList->MapAction( Commands.LockAllBones, FExecuteAction::CreateUObject(this, &USkinWeightsPaintTool::ToggleLockAllBones, true));
+			InOutCommandList->MapAction( Commands.UnlockAllBones, FExecuteAction::CreateUObject(this, &USkinWeightsPaintTool::ToggleLockAllBones, false));
+			
+			InOutCommandList->MapAction( Commands.BypassLocks,
+				FExecuteAction::CreateUObject(this, &USkinWeightsPaintTool::ToggleBypassLocks),
+				FCanExecuteAction(),
+				FIsActionChecked::CreateUObject(this, &USkinWeightsPaintTool::IsBypassingAllLocks));
+		});
+	Delegates.OnExtendContextMenu.BindLambda([](FMenuBuilder& MenuBuilder)
+		{
+			FSkinWeightsPaintToolCommands Commands = FSkinWeightsPaintToolCommands::Get();
+			MenuBuilder.BeginSection("Isolate", LOCTEXT("IsolateSection", "Isolate"));
+			{
+				MenuBuilder.AddMenuEntry(Commands.IsolateSelectedBones);
+			}
+			MenuBuilder.EndSection();
+			
+			MenuBuilder.BeginSection("Lock", LOCTEXT("LockSection", "Lock"));
+			{
+				MenuBuilder.AddMenuEntry(Commands.LockSelectedBones);
+				MenuBuilder.AddMenuEntry(Commands.LockAllBones);
+			}
+			MenuBuilder.EndSection();
+			
+			MenuBuilder.BeginSection("Unlock", LOCTEXT("UnlockSection", "Unlock"));
+			{
+				MenuBuilder.AddMenuEntry(Commands.UnlockSelectedBones);
+				MenuBuilder.AddMenuEntry(Commands.UnlockAllBones);
+			}
+			MenuBuilder.EndSection();
+			
+			MenuBuilder.BeginSection("Bypass", LOCTEXT("BypassSection", "Bypass"));
+			{
+				MenuBuilder.AddMenuEntry(Commands.BypassLocks);
+			}
+			MenuBuilder.EndSection();
+
+			return true;
+		});
+	
+	TSharedPtr<SReferenceSkeletonTree> ToolSkeletonTree =
+		SNew(SReferenceSkeletonTree)
+		.Modifier(SkeletonReader)
+		.Delegates(Delegates);
+
+	SkeletonTreeNotifierBindScope.Reset(
+		new FSkeletalMeshNotifierBindScope(
+			EditorBinding->GetNotifier(),
+			ToolSkeletonTree->GetNotifier()));	
+		
+	TArray<FName> SelectedBones = EditorBinding->GetSelectedBones();
+	ToolSkeletonTree->GetNotifier()->HandleNotification(SelectedBones, ESkeletalMeshNotifyType::BonesSelected);		
+	
+	return ToolSkeletonTree;
+}
+
+TSharedRef<SWidget> USkinWeightsPaintTool::OnGetLockWidget(FName BoneName)
+{
+	int32 BoneIndex = GetBoneIndexFromName(BoneName);
+
+	auto GetColorAndOpacity = [this, BoneIndex] ()
+		{
+			const bool bIsLockByUser = GetRawLockedBoneIndices().Contains(BoneIndex);
+			const bool bIsHovered = LockButtons[BoneIndex]->IsHovered();
+
+			FLinearColor OutColor = FLinearColor::White;
+			if (IsBypassingAllLocks())
+			{
+				OutColor = FLinearColor::Green;
+			}
+			
+			float Opacity = 0.0f;
+
+			if (bIsLockByUser)
+			{
+				// Directly active, full opacity
+				Opacity = 1.0f;
+			}
+			else if (bIsHovered)
+			{
+				// Mouse is over widget and it is not directly active.
+				Opacity = .8f;
+			}
+			else 
+			{
+				// Not active in any way and mouse is not over widget or item.
+				Opacity = 0.1f;
+			}
+			OutColor.A = Opacity;
+			return OutColor;	
+		};
+
+	TSharedPtr<SButton> LockButton = SNew(SButton)
+		.IsFocusable(false)
+		.ButtonStyle(FAppStyle::Get(), "NoBorder")
+		.ButtonColorAndOpacity_Lambda(GetColorAndOpacity)
+		.OnClicked_Lambda([this, BoneIndex]()
+		{
+			if (LockedBoneIndices.Contains(BoneIndex))
+			{
+				ToggleLockBone(BoneIndex, false);
+			}
+			else
+			{
+				ToggleLockBone(BoneIndex, true);
+			}
+			return FReply::Handled();
+		})
+		.ContentPadding(1.f)
+		.ToolTipText(LOCTEXT("BoneWeightLockToolTip", "Skin weights for locked bones cannot be changed"))
+		[
+			SNew(SImage)
+			.ColorAndOpacity_Lambda(GetColorAndOpacity)
+			.Image_Lambda([this]()
+				{
+					if (bBypassLocks)
+					{
+						return FAppStyle::GetBrush("Icons.Unlock");
+					}
+					return FAppStyle::GetBrush("Icons.Lock");
+				})
+		];
+
+	LockButtons.Emplace(BoneIndex, LockButton);
+	
+	return LockButton.ToSharedRef();
+}
+
+void USkinWeightsPaintTool::ToggleLockAllBones(bool bLock)
+{
+	ISkeletonProvider* SkeletonProvider = CastChecked<ISkeletonProvider>(Target);
+	const FReferenceSkeleton& RefSkeleton = SkeletonProvider->GetSkeleton();
+	int32 NumBones = RefSkeleton.GetRawBoneNum();
+
+	for (BoneIndex BoneIndex = 0; BoneIndex < NumBones; ++BoneIndex)
+	{
+		ToggleLockBone(BoneIndex, bLock);
+	}
+}
+
+void USkinWeightsPaintTool::ToggleLockSelectedBones(bool bLock)
+{
+	for (int32 BoneIndex : SelectedBoneIndices)
+	{
+		ToggleLockBone(BoneIndex, bLock);
+	}
+}
+
+void USkinWeightsPaintTool::IsolateSelectedBones()
+{
+	ToggleLockAllBones(true);
+	ToggleLockSelectedBones(false);
+}
+
+void USkinWeightsPaintTool::ToggleBypassLocks()
+{
+	bBypassLocks = !bBypassLocks;
+}
+
+void USkinWeightsPaintTool::ToggleLockBone(BoneIndex InBoneIndex, bool bLock)
+{
+	if (bLock)
+	{
+		LockedBoneIndices.Add(InBoneIndex);
+	}
+	else
+	{
+		LockedBoneIndices.Remove(InBoneIndex);
+	}	
+}
+
+void USkinWeightsPaintTool::InitializeSkeletonReader()
+{
+	if (ISkeletalMeshBackedTarget* SkeletalMeshBackedTarget = Cast<ISkeletalMeshBackedTarget>(Target))
+	{
+		SkeletonReader = NewObject<USkeletonModifier>();
+		SkeletonReader->SetSkeletalMesh(SkeletalMeshBackedTarget->GetSkeletalMesh());
+		SkeletonReader->SetReadOnly(true);
+	}
+	else if (IPersistentDynamicMeshSource* DynamicMeshSource = Cast<IPersistentDynamicMeshSource>(Target))
+	{
+		SkeletonReader = NewObject<USkeletonModifier>();
+		SkeletonReader->SetDynamicMesh(DynamicMeshSource->GetDynamicMeshContainer());
+		SkeletonReader->SetReadOnly(true);
+	}
+}
+
+bool USkinWeightsPaintTool::ShouldProceedWithAllBonesUnlocked()
+{
+	if (GetFinalLockedBoneIndices().IsEmpty())
+	{
+		return true;
+	}
+
+	FText Message = LOCTEXT("UnlockAllBonesDialog", "The current operation does not support skin weight locking, unlock weights for all bones and continue?");
+	EAppReturnType::Type Response = FMessageDialog::Open(EAppMsgType::OkCancel, Message);
+
+	if (Response == EAppReturnType::Ok)
+	{
+		ToggleLockAllBones(false);	
+	}
+	
+	return Response == EAppReturnType::Ok;
+}
+
 void USkinWeightsPaintTool::OnActiveSkinWeightProfileChanged()
 {
 
diff --git a/Engine/Plugins/Runtime/MeshModelingToolset/Source/MeshModelingToolsEditorOnly/Public/SkeletalMesh/SReferenceSkeletonTree.h b/Engine/Plugins/Runtime/MeshModelingToolset/Source/MeshModelingToolsEditorOnly/Public/SkeletalMesh/SReferenceSkeletonTree.h
new file mode 100644
--- /dev/null
+++ b/Engine/Plugins/Runtime/MeshModelingToolset/Source/MeshModelingToolsEditorOnly/Public/SkeletalMesh/SReferenceSkeletonTree.h
@@ -0,0 +1,305 @@
+// Copyright Epic Games, Inc. All Rights Reserved.
+
+#pragma once
+
+#include "EditorUndoClient.h"
+#include "SkeletalMeshNotifier.h"
+#include "DragAndDrop/DecoratedDragDropOp.h"
+#include "Widgets/SCompoundWidget.h"
+#include "Widgets/Views/SListView.h"
+#include "Widgets/Views/STableRow.h"
+#include "Widgets/Views/STreeView.h"
+#include "Misc/Change.h"
+#include "ReferenceSkeleton.h"
+#include "SkeletalMesh/SkeletalMeshEditingInterface.h"
+
+#define UE_API MESHMODELINGTOOLSEDITORONLY_API
+
+class FTextFilterExpressionEvaluator;
+class FUICommandList;
+class USkeletonModifier;
+class SInlineEditableTextBlock;
+
+namespace ReferenceSkeletonTreeLocals
+{
+
+class FSkeletonModifierChange : public FCommandChange
+{
+public:
+	FSkeletonModifierChange(const USkeletonModifier* InModifier);
+
+	virtual FString ToString() const override
+	{
+		return FString(TEXT("Edit Skeleton"));
+	}
+
+	virtual void Apply(UObject* Object) override;
+	virtual void Revert(UObject* Object) override;
+
+	void StoreSkeleton(const USkeletonModifier* InModifier);
+
+private:
+	FReferenceSkeleton PreChangeSkeleton;
+	TArray<int32> PreBoneTracker;
+	FReferenceSkeleton PostChangeSkeleton;
+	TArray<int32> PostBoneTracker;
+};
+
+enum class EBoneRenameType : uint32
+{
+	Rename,
+	SearchAndReplace,
+	AddPrefix,
+	AddSuffix
+};
+	
+}
+
+class FBoneElement : public TSharedFromThis<FBoneElement>
+{
+public:
+	FBoneElement(const FName& InBoneName, TWeakObjectPtr<USkeletonModifier> InModifier);
+
+	FName BoneName;
+
+	bool bIsHidden = false;
+	
+	TSharedPtr<FBoneElement> UnFilteredParent;
+	TSharedPtr<FBoneElement> Parent;
+	TArray<TSharedPtr<FBoneElement>> Children;
+
+	void RequestRename(const ReferenceSkeletonTreeLocals::EBoneRenameType InRenameType) const;
+	
+	DECLARE_DELEGATE_OneParam(FOnRenameRequested, const ReferenceSkeletonTreeLocals::EBoneRenameType);
+	FOnRenameRequested OnRenameRequested;
+
+	DECLARE_DELEGATE_TwoParams(FOnRenamed, const FName, const FName);
+	FOnRenamed OnRenamed;
+
+private:
+	TWeakObjectPtr<USkeletonModifier> WeakModifier;
+};
+
+using FOnRefSkeletonTreeCanAcceptDrop = STableRow<TSharedPtr<FBoneElement>>::FOnCanAcceptDrop;
+using FOnRefSkeletonTreeAcceptDrop = STableRow<TSharedPtr<FBoneElement>>::FOnAcceptDrop;
+using FOnBoneRenamed = FBoneElement::FOnRenamed;
+
+struct FRefSkeletonTreeDelegates
+{
+	FOnDragDetected					OnDragDetected;
+	FOnRefSkeletonTreeCanAcceptDrop	OnCanAcceptDrop;
+	FOnRefSkeletonTreeAcceptDrop	OnAcceptDrop;
+	FOnBoneRenamed					OnBoneRenamed;
+	FOnTextCommitted				OnBoneNameCommitted;
+
+	// Tools can use these to extend the skeleton tree with extra commands or widgets
+	DECLARE_DELEGATE_RetVal_OneParam(TSharedRef<SWidget>, FOnGetRowWidget, FName)
+	FOnGetRowWidget	OnGetRowWidget;
+	
+	DECLARE_DELEGATE_OneParam(FOnExtendCommandList, TSharedPtr<FUICommandList>)
+	FOnExtendCommandList OnExtendCommandList;
+	
+	DECLARE_DELEGATE_RetVal_OneParam(bool, FOnExtendContextMenu, FMenuBuilder&)
+	FOnExtendContextMenu OnExtendContextMenu;
+};
+
+class SBoneItem : public SCompoundWidget
+{
+	SLATE_BEGIN_ARGS(SBoneItem) {}
+		SLATE_ARGUMENT(TWeakPtr<FBoneElement>, TreeElement)
+		SLATE_ARGUMENT(TWeakObjectPtr<USkeletonModifier>, WeakModifier)
+		SLATE_ARGUMENT(FRefSkeletonTreeDelegates, Delegates)
+		SLATE_EVENT( FIsSelected, IsSelected )
+		SLATE_ATTRIBUTE(bool, IsReadOnly)
+	SLATE_END_ARGS()
+
+	virtual ~SBoneItem() override;
+	
+	// Slate construction function
+	void Construct(const FArguments& InArgs);
+
+	void HandleOnRenameRequested(const ReferenceSkeletonTreeLocals::EBoneRenameType InType);
+	
+private:
+
+	bool OnVerifyNameChanged(const FText& InText, FText& OutErrorMessage) const;
+
+	void OnEnterEditingMode();
+	void OnExitEditingMode();
+
+	static FSlateColor GetTextColor(FIsSelected InIsSelected);
+	
+	FText GetName() const;
+	
+	TWeakPtr<FBoneElement> WeakTreeElement;
+	TWeakObjectPtr<USkeletonModifier> WeakModifier;
+	TSharedPtr< SInlineEditableTextBlock > InlineWidget;
+	
+	bool bEditing = false;
+	TOptional<ReferenceSkeletonTreeLocals::EBoneRenameType> OptRenameType;
+};
+
+class FBoneItemDragDropOp : public FDecoratedDragDropOp
+{
+public:
+	DRAG_DROP_OPERATOR_TYPE(FBoneItemDragDropOp, FDecoratedDragDropOp)
+
+	static TSharedRef<FBoneItemDragDropOp> New(TArrayView<const TSharedPtr<FBoneElement>> InElements);
+
+	// FDecoratedDragDropOp interface
+	virtual TSharedPtr<SWidget> GetDefaultDecorator() const override;
+	// End FDecoratedDragDropOp interface
+	
+	TArray< TWeakPtr<FBoneElement> > Elements;
+};
+
+class SReferenceSkeletonRow : public SMultiColumnTableRow<TSharedPtr<FBoneElement>>
+{
+public:
+	SLATE_BEGIN_ARGS(SReferenceSkeletonRow) {}
+		SLATE_ARGUMENT(TWeakObjectPtr<USkeletonModifier>, WeakModifier)
+		SLATE_ARGUMENT(TSharedPtr<FBoneElement>, TreeElement)
+		SLATE_ARGUMENT(FRefSkeletonTreeDelegates, Delegates)
+		SLATE_ARGUMENT(bool, IsReadOnly)
+	SLATE_END_ARGS()
+
+	// Slate construction function
+	UE_API void Construct(const FArguments& InArgs, const TSharedRef<STableViewBase>& InOwnerTable);
+
+	// SMultiColumnTableRow interface
+	virtual TSharedRef<SWidget> GenerateWidgetForColumn(const FName& InColumnName) override;
+	// End SMultiColumnTableRow interface
+
+private:
+	TWeakPtr<FBoneElement> WeakTreeElement;
+	TWeakObjectPtr<USkeletonModifier> WeakModifier;
+	FRefSkeletonTreeDelegates Delegates;
+	bool bIsReadOnly = false;
+};
+
+typedef STreeView< TSharedPtr<FBoneElement> > SRefSkeletonTreeView;
+
+class SReferenceSkeletonTree;
+
+class FReferenceSkeletonWidgetNotifier: public ISkeletalMeshNotifier
+{
+public:
+	FReferenceSkeletonWidgetNotifier(TSharedRef<SReferenceSkeletonTree> InTree);
+	virtual void HandleNotification(const TArray<FName>& BoneNames, const ESkeletalMeshNotifyType InNotifyType) override;
+	
+private:
+	TWeakPtr<SReferenceSkeletonTree> Tree;
+};
+
+class SReferenceSkeletonTree : public SCompoundWidget, public FEditorUndoClient
+{
+public:
+	SLATE_BEGIN_ARGS(SReferenceSkeletonTree)
+		: _bIsReadOnly(false)
+		{}
+		SLATE_ARGUMENT(TWeakObjectPtr<USkeletonModifier>, Modifier)
+		SLATE_ARGUMENT(FRefSkeletonTreeDelegates, Delegates)
+		SLATE_ARGUMENT(bool, bIsReadOnly)
+	SLATE_END_ARGS()
+
+	UE_API SReferenceSkeletonTree();
+    UE_API virtual ~SReferenceSkeletonTree() override;
+
+	// Slate construction function
+	UE_API void Construct(const FArguments& InArgs);
+
+	virtual void PostUndo(bool bSuccess) override;
+	virtual void PostRedo(bool bSuccess) override;
+	
+	// Selection
+	UE_API void GetSelectedBoneNames(TArray<FName>& OutSelectedBoneNames) const;
+	TArray<TSharedPtr<FBoneElement>> GetSelectedItems() const;
+	bool HasSelectedItems() const;
+	void SelectItemFromNames(const TArray<FName>& InBoneNames, bool bFrameSelection = false);
+
+	UE_API TSharedPtr<ISkeletalMeshNotifier> GetNotifier();
+
+private:
+	// SWidget interface
+	virtual FReply OnKeyDown(const FGeometry& MyGeometry, const FKeyEvent& InKeyEvent) override;
+	// END SWidget interface
+	
+	void BindCommands();
+
+	void AddItemToSelection(const TSharedPtr<FBoneElement>& InItem);
+	void RemoveItemFromSelection(const TSharedPtr<FBoneElement>& InItem);
+	void ReplaceItemInSelection(const FText& OldName, const FText& NewName);
+
+	// Create, Rename, Delete, Parent bones
+	void HandleNewBone();
+	bool CanAddNewBone() const;
+	void HandleRenameBone(const ReferenceSkeletonTreeLocals::EBoneRenameType InRenameType) const;
+	bool CanRenameBone() const;
+	void HandleDeleteBone();
+	bool CanDeleteBone() const;
+	void HandleUnParentBone();
+	bool CanUnParentBone() const;
+
+	// Copy & Paste
+	void HandleCopyBones() const;
+	bool CanCopyBones() const;
+	void HandlePasteBones();
+	bool CanPasteBones() const;
+	void HandleDuplicateBones();
+	bool CanDuplicateBones() const;
+
+	// Filter
+	void OnFilterTextChanged(const FText& SearchText);
+	
+	// Callbacks
+	void RefreshTreeView(bool IsInitialSetup=false);
+	void HandleGetChildrenForTree(TSharedPtr<FBoneElement> InItem, TArray<TSharedPtr<FBoneElement>>& OutChildren);
+	void OnSelectionChanged(TSharedPtr<FBoneElement> InItem, ESelectInfo::Type InSelectInfo);
+	
+	TSharedRef< SWidget > CreateAddNewMenu();
+	EVisibility GetAddNewButtonVisibility() const;
+	TSharedPtr< SWidget > CreateContextMenu();
+
+	void OnItemDoubleClicked(TSharedPtr<FBoneElement> InItem);
+	void OnSetExpansionRecursive(TSharedPtr<FBoneElement> InItem, bool bShouldBeExpanded);
+	void SetExpansionRecursive(TSharedPtr<FBoneElement> InElement, bool bTowardsParent, bool bShouldBeExpanded);
+
+	// Delegates
+	virtual FReply OnDragDetected(const FGeometry& MyGeometry, const FPointerEvent& MouseEvent) override;
+	TOptional<EItemDropZone> OnCanAcceptDrop(const FDragDropEvent& DragDropEvent, EItemDropZone DropZone, TSharedPtr<FBoneElement> TargetItem);
+	FReply OnAcceptDrop(const FDragDropEvent& DragDropEvent, EItemDropZone DropZone, TSharedPtr<FBoneElement> TargetItem);
+	void OnBoneRenamed(const FName InOldName, const FName InNewName) const;
+	
+	void OnNewBoneNameCommitted(const FText& InText, ETextCommit::Type InCommitType);
+
+	// The modifier that actually does modifications on the reference skeleton
+	TWeakObjectPtr<USkeletonModifier> Modifier;
+	
+	TSharedPtr<FReferenceSkeletonWidgetNotifier> Notifier;
+
+	// undo
+	void BeginChange();
+	void EndChange();
+	void CancelChange();
+	TUniquePtr<ReferenceSkeletonTreeLocals::FSkeletonModifierChange> ActiveChange;
+	
+	// Commands
+	TSharedPtr<FUICommandList> CommandList;
+	
+	// Tree view
+	TSharedPtr<SRefSkeletonTreeView> TreeView;
+	TArray<TSharedPtr<FBoneElement>> RootElements;
+	TArray<TSharedPtr<FBoneElement>> AllElements;
+
+	// filter
+	TSharedPtr<FTextFilterExpressionEvaluator> TextFilter;
+
+
+	bool bIsReadOnly = false;
+
+	FRefSkeletonTreeDelegates Delegates;
+
+	friend class FReferenceSkeletonWidgetNotifier;
+};
+
+#undef UE_API
\ No newline at end of file
diff --git a/Engine/Plugins/Runtime/MeshModelingToolset/Source/MeshModelingToolsEditorOnly/Public/SkeletalMesh/SkeletalMeshEditingInterface.h b/Engine/Plugins/Runtime/MeshModelingToolset/Source/MeshModelingToolsEditorOnly/Public/SkeletalMesh/SkeletalMeshEditingInterface.h
--- a/Engine/Plugins/Runtime/MeshModelingToolset/Source/MeshModelingToolsEditorOnly/Public/SkeletalMesh/SkeletalMeshEditingInterface.h
+++ b/Engine/Plugins/Runtime/MeshModelingToolset/Source/MeshModelingToolsEditorOnly/Public/SkeletalMesh/SkeletalMeshEditingInterface.h
@@ -22,6 +22,7 @@ class IGizmoStateTarget;
 class USkeletonModifier;
 struct FInputDeviceRay;
 class UToolTarget;
+class SWidget;
 
 /**
  * USkeletalMeshEditingInterface
@@ -45,7 +46,9 @@ public:
 	UE_API TSharedPtr<ISkeletalMeshNotifier> GetNotifier();
 	UE_API bool NeedsNotification() const;
 
-	UE_API virtual TWeakObjectPtr<USkeletonModifier> GetModifier() const;
+	UE_API virtual bool CanModifySkeleton() const;
+	UE_API virtual TSharedPtr<SWidget> GetCustomSkeletonTreeWidget(TSharedPtr<ISkeletalMeshEditorBinding> EditorBinding);
+	
 
 protected:
 	virtual void HandleSkeletalMeshModified(const TArray<FName>& BoneNames, const ESkeletalMeshNotifyType InNotifyType) = 0;
@@ -53,7 +56,7 @@ protected:
 private:
 	TSharedPtr<FSkeletalMeshToolNotifier> Notifier;
 	
-	friend FSkeletalMeshToolNotifier;
+	friend class FSkeletalMeshToolNotifier;
 };
 
 /**
diff --git a/Engine/Plugins/Runtime/MeshModelingToolset/Source/MeshModelingToolsEditorOnly/Public/SkeletalMesh/SkeletonEditingTool.h b/Engine/Plugins/Runtime/MeshModelingToolset/Source/MeshModelingToolsEditorOnly/Public/SkeletalMesh/SkeletonEditingTool.h
--- a/Engine/Plugins/Runtime/MeshModelingToolset/Source/MeshModelingToolsEditorOnly/Public/SkeletalMesh/SkeletonEditingTool.h
+++ b/Engine/Plugins/Runtime/MeshModelingToolset/Source/MeshModelingToolsEditorOnly/Public/SkeletalMesh/SkeletonEditingTool.h
@@ -14,6 +14,7 @@
 #include "Mechanics/RectangleMarqueeMechanic.h"
 #include "Misc/EnumClassFlags.h"
 #include "Selection/MeshTopologySelectionMechanic.h"
+#include "Framework/Commands/Commands.h"
 
 #include "SkeletonEditingTool.generated.h"
 
@@ -59,6 +60,28 @@ private:
 	
 }
 
+class FSkeletonEditingToolCommands : public TCommands<FSkeletonEditingToolCommands>
+{
+public:
+	FSkeletonEditingToolCommands();
+
+	virtual void RegisterCommands() override;
+
+	// Skeleton Editing commands
+	TSharedPtr<FUICommandInfo> NewBone;
+	TSharedPtr<FUICommandInfo> RemoveBone;
+	TSharedPtr<FUICommandInfo> UnParentBone;
+	TSharedPtr<FUICommandInfo> RenameBone;
+	TSharedPtr<FUICommandInfo> SearchAndReplace;
+	TSharedPtr<FUICommandInfo> AddPrefix;
+	TSharedPtr<FUICommandInfo> AddSuffix;
+
+	// Copy & Paste commands
+	TSharedPtr<FUICommandInfo> CopyBones;
+	TSharedPtr<FUICommandInfo> PasteBones;
+	TSharedPtr<FUICommandInfo> DuplicateBones;
+};	
+
 /**
  * USkeletonEditingToolBuilder
  */
@@ -155,7 +178,9 @@ public:
 	UE_API virtual void OnUpdateModifierState(int ModifierID, bool bIsOn) override;
 
 	// ISkeletalMeshEditingInterface overrides
-	UE_API virtual TWeakObjectPtr<USkeletonModifier> GetModifier() const override;
+	UE_API virtual bool CanModifySkeleton() const override;
+	UE_API virtual TSharedPtr<SWidget> GetCustomSkeletonTreeWidget(TSharedPtr<ISkeletalMeshEditorBinding> EditorBinding) override;
+	
 
 	UE_API EEditingOperation GetOperation() const;
 	UE_API void SetOperation(const EEditingOperation InOperation, const bool bUpdateGizmo = true);
@@ -296,6 +321,8 @@ private:
 	bool bDeferUntilFocused = false;
 
 	TWeakObjectPtr<USkeletalMesh> WeakMesh = nullptr;
+
+	TUniquePtr<FSkeletalMeshNotifierBindScope> SkeletonTreeNotifierBindScope;
 };
 
 ENUM_CLASS_FLAGS(USkeletonEditingTool::EBoneSelectionMode);
diff --git a/Engine/Plugins/Runtime/MeshModelingToolset/Source/MeshModelingToolsEditorOnly/Public/SkeletalMesh/SkinWeightsPaintTool.h b/Engine/Plugins/Runtime/MeshModelingToolset/Source/MeshModelingToolsEditorOnly/Public/SkeletalMesh/SkinWeightsPaintTool.h
--- a/Engine/Plugins/Runtime/MeshModelingToolset/Source/MeshModelingToolsEditorOnly/Public/SkeletalMesh/SkinWeightsPaintTool.h
+++ b/Engine/Plugins/Runtime/MeshModelingToolset/Source/MeshModelingToolsEditorOnly/Public/SkeletalMesh/SkinWeightsPaintTool.h
@@ -18,6 +18,8 @@
 #include "Selections/GeometrySelection.h"
 #include "TargetInterfaces/MeshTargetInterfaceTypes.h"
 #include "ToolMeshSelector.h"
+#include "Framework/Commands/Commands.h"
+
 
 #include "SkinWeightsPaintTool.generated.h"
 
@@ -31,6 +33,8 @@ class UPolygonSelectionMechanic;
 class UPersonaEditorModeManagerContext;
 class FEditorViewportClient;
 class UVolumetricBrushStampIndicator;
+class USkeletonModifier;
+class SButton;
 
 namespace UE::Geometry 
 {
@@ -207,7 +211,7 @@ namespace SkinPaintTool
 	{
 		// copy the initial weight values from the skeletal mesh
 		void InitializeSkinWeights(
-			USkinWeightsPaintTool* Tool,
+			USkinWeightsPaintTool* InTool,
 			const FDynamicMesh3& Mesh);
 
 		// applies an edit to a single vertex weight on a single bone, then normalizes the remaining weights while
@@ -281,6 +285,8 @@ namespace SkinPaintTool
 
 		// which skin profile is currently edited
 		FName Profile = FSkeletalMeshAttributesShared::DefaultSkinWeightProfileName;
+
+		TWeakObjectPtr<USkinWeightsPaintTool> Tool;
 	};
 
 	struct FSkinMirrorData
@@ -315,6 +321,25 @@ namespace SkinPaintTool
 	
 }
 
+class FSkinWeightsPaintToolCommands : public TCommands<FSkinWeightsPaintToolCommands>
+{
+public:
+	FSkinWeightsPaintToolCommands();
+
+	virtual void RegisterCommands() override;
+	
+	TSharedPtr<FUICommandInfo> IsolateSelectedBones;
+	
+	TSharedPtr<FUICommandInfo> LockSelectedBones;
+	TSharedPtr<FUICommandInfo> UnlockSelectedBones;
+
+	TSharedPtr<FUICommandInfo> LockAllBones;
+	TSharedPtr<FUICommandInfo> UnlockAllBones;
+
+	TSharedPtr<FUICommandInfo> BypassLocks;
+};	
+
+
 UCLASS(MinimalAPI)
 class USkinWeightsPaintToolBuilder : public UMeshSurfacePointMeshEditingToolBuilder
 {
@@ -665,16 +690,25 @@ public:
 
 	// HELPER functions for modifying weights
 	//
-	// given a map of BoneIndex > Weight values for a single vertex, modify the map by removing the smallest weights to fit in Max Influences
-	static UE_API void TruncateWeightMap(TMap<BoneIndex, float>& InOutWeights);
-	// given a map of BoneIndex > Weight values for a single vertex, modify the weights to sum to 1
-	static UE_API void NormalizeWeightMap(TMap<BoneIndex, float>& InOutWeights);
-	// sum all the weights on all bones for a given list of vertices (results we not be normalized!)
+	// given a map of BoneIndex > Weight values for a single vertex, modify the map by removing the smallest weights to fit in Max Influences - NumLockedBones
+	static UE_API void TruncateWeightMap(TMap<BoneIndex, float>& InOutWeights, int32 NumLockedBones = 0);
+	// given a map of BoneIndex > Weight values for a single vertex, modify the weights to sum to ExpectedSum
+	static UE_API void NormalizeWeightMap(TMap<BoneIndex, float>& InOutWeights, float ExpectedSum = 1.0f);
+	// sum all the weights on all bones for a given list of vertices (results are not normalized!)
 	static UE_API void AccumulateWeights(
 		const TArray<SkinPaintTool::VertexWeights>& AllWeights,
 		const TArray<VertexIndex>& VerticesToAccumulate,
-		TMap<BoneIndex, float>& OutWeights);
+		const TSet<BoneIndex>& ExcludedBones,
+		TMap<BoneIndex, float>& OutAccumulatedWeights);
 
+	// query information about how weights are locked on a specific vertex
+	static UE_API void GetLockedWeight(
+		const TArray<SkinPaintTool::VertexWeights>& AllWeights,
+		VertexIndex InVertexID,
+		const TSet<BoneIndex>& LockedBones,
+		float& OutTotalLockedWeight,
+		int32& OutNumLockedBones);
+	
 	// copy paste
 	UE_API void CopyWeights();
 	UE_API void PasteWeights();
@@ -742,6 +776,18 @@ public:
 	// get access to the mesh being edited
 	UE_API FDynamicMesh3& GetMesh() {return EditedMesh;};
 
+	// get the set of bones whose weights are actually locked, taking into account bBypassAllLocks
+	UE_API const TSet<int32>& GetFinalLockedBoneIndices();
+
+	// get the set of bones whose weights are locked by user, does not consider bBypassAllLocks
+	UE_API const TSet<int32>& GetRawLockedBoneIndices();
+
+	// whether weight edits should respect locks or not
+	UE_API bool IsBypassingAllLocks() const;;
+	
+	// Certain operations do not support bone locking yet, show a dialog to check if user still want to proceed
+	UE_API bool ShouldProceedWithAllBonesUnlocked();
+	
 	// HOW TO EDIT WEIGHTS WITH UNDO/REDO:
 	//
 	// "Interactive" Edits:
@@ -853,6 +899,7 @@ protected:
 	// which bone are we currently painting?
 	UE_API void UpdateCurrentBone(const FName &BoneName);
 	UE_API BoneIndex GetBoneIndexFromName(const FName BoneName) const;
+	UE_API bool HasSelectedBone() const;
 	FName CurrentBone = NAME_None;
 	TOptional<FName> PendingCurrentBone;
 	TArray<FName> SelectedBoneNames;
@@ -860,7 +907,27 @@ protected:
 
 	// ISkeletalMeshEditingInterface
 	UE_API virtual void HandleSkeletalMeshModified(const TArray<FName>& InBoneNames, const ESkeletalMeshNotifyType InNotifyType) override;
+	UE_API virtual TSharedPtr<SWidget> GetCustomSkeletonTreeWidget(TSharedPtr<ISkeletalMeshEditorBinding> EditorBinding) override;
+	
+	// used to spawn a tool skeleton tree widget
+	UPROPERTY()
+	TObjectPtr<USkeletonModifier> SkeletonReader = nullptr;
+
+	// skin weight locking
+	UE_API virtual TSharedRef<SWidget> OnGetLockWidget(FName BoneName);
+	UE_API void ToggleLockAllBones(bool bLock);
+	UE_API void ToggleLockSelectedBones(bool bLock);
+	UE_API void IsolateSelectedBones();
+	UE_API void ToggleBypassLocks();
+	UE_API void ToggleLockBone(BoneIndex InBoneIndex, bool bLock);
+	UE_API void InitializeSkeletonReader();
+
+	TMap<BoneIndex, TSharedPtr<SButton>> LockButtons;
+	TSet<BoneIndex> LockedBoneIndices;
+	TSet<BoneIndex> UnlockedBoneIndices;
+	bool bBypassLocks = false;
 
+	
 	// the selection system for the main mesh
 	UPROPERTY()
 	TObjectPtr<UToolMeshSelector> MeshSelector;
@@ -891,6 +958,8 @@ protected:
 
 	TArray<FTransform> DefaultComponentSpaceBoneTransforms;
 
+	TUniquePtr<FSkeletalMeshNotifierBindScope> SkeletonTreeNotifierBindScope;
+	
 	friend SkinPaintTool::FSkinToolDeformer;
 };
 
diff --git a/Engine/Plugins/Runtime/MeshModelingToolset/Source/SkeletalMeshModifiers/Private/SkeletonModifier.cpp b/Engine/Plugins/Runtime/MeshModelingToolset/Source/SkeletalMeshModifiers/Private/SkeletonModifier.cpp
--- a/Engine/Plugins/Runtime/MeshModelingToolset/Source/SkeletalMeshModifiers/Private/SkeletonModifier.cpp
+++ b/Engine/Plugins/Runtime/MeshModelingToolset/Source/SkeletalMeshModifiers/Private/SkeletonModifier.cpp
@@ -494,6 +494,16 @@ bool USkeletonModifier::SetDynamicMesh(UDynamicMesh* InDynamicMesh)
 
 }
 
+void USkeletonModifier::SetReadOnly(bool bInReadOnly)
+{
+	bReadOnly = bInReadOnly;
+}
+
+bool USkeletonModifier::IsReadOnly() const
+{
+	return bReadOnly;
+}
+
 bool USkeletonModifier::CommitSkeletonToDynamicMesh()
 {
 	if (!WeakDynamicMesh.IsValid())
@@ -2357,4 +2367,4 @@ const FTransform& USkeletonModifier::GetTransform(const int32 InBoneIndex, const
 }
 
 
-#undef LOCTEXT_NAMESPACE
\ No newline at end of file
+#undef LOCTEXT_NAMESPACE
diff --git a/Engine/Plugins/Runtime/MeshModelingToolset/Source/SkeletalMeshModifiers/Public/SkeletonModifier.h b/Engine/Plugins/Runtime/MeshModelingToolset/Source/SkeletalMeshModifiers/Public/SkeletonModifier.h
--- a/Engine/Plugins/Runtime/MeshModelingToolset/Source/SkeletalMeshModifiers/Public/SkeletonModifier.h
+++ b/Engine/Plugins/Runtime/MeshModelingToolset/Source/SkeletalMeshModifiers/Public/SkeletonModifier.h
@@ -143,6 +143,9 @@ public:
 
 	UE_API bool SetDynamicMesh(UDynamicMesh* InDynamicMesh);
 	
+	UE_API void SetReadOnly(bool bInReadOnly);
+	UE_API bool IsReadOnly() const;
+	
 	/**
 	 * Applies the skeleton modifications to the dynamic mesh.
 	 * @return true if commit succeeded.
@@ -306,7 +309,7 @@ private:
 	TUniquePtr<FReferenceSkeleton> ReferenceSkeleton;
 	TUniquePtr<FTransformComposer> TransformComposer;
 	TArray<int32> BoneIndexTracker;
-	bool bRestricted = false;
+	bool bReadOnly = false;
 
 public:
 	
diff --git a/Engine/Source/Developer/SkeletalMeshUtilitiesCommon/Private/SkeletalMeshNotifier.cpp b/Engine/Source/Developer/SkeletalMeshUtilitiesCommon/Private/SkeletalMeshNotifier.cpp
--- a/Engine/Source/Developer/SkeletalMeshUtilitiesCommon/Private/SkeletalMeshNotifier.cpp
+++ b/Engine/Source/Developer/SkeletalMeshUtilitiesCommon/Private/SkeletalMeshNotifier.cpp
@@ -21,3 +21,41 @@ void ISkeletalMeshNotifier::Notify(const TArray<FName>& BoneNames, const ESkelet
 	}
 }
 
+FSkeletalMeshNotifierBindScope::FSkeletalMeshNotifierBindScope(
+	TWeakPtr<ISkeletalMeshNotifier> InNotifierA,
+	TWeakPtr<ISkeletalMeshNotifier> InNotifierB
+	) : NotifierA(InNotifierA)
+	, NotifierB(InNotifierB)
+{
+
+	DelegateToRemoveFromNotifierA = InNotifierA.Pin()->Delegate().AddLambda(
+	[InNotifierB](const TArray<FName>& BoneNames, const ESkeletalMeshNotifyType InNotifyType)
+	{
+		if (InNotifierB.IsValid())
+		{
+			InNotifierB.Pin()->HandleNotification(BoneNames, InNotifyType);
+		}
+	});
+	
+	DelegateToRemoveFromNotifierB = InNotifierB.Pin()->Delegate().AddLambda(
+	[InNotifierA](const TArray<FName>& BoneNames, const ESkeletalMeshNotifyType InNotifyType)
+	{
+		if (InNotifierA.IsValid())
+		{
+			InNotifierA.Pin()->HandleNotification(BoneNames, InNotifyType);
+		}
+	});
+}
+
+FSkeletalMeshNotifierBindScope::~FSkeletalMeshNotifierBindScope()
+{
+	if (DelegateToRemoveFromNotifierA.IsValid() && NotifierA.IsValid())
+	{
+		check(NotifierA.Pin()->Delegate().Remove(DelegateToRemoveFromNotifierA));
+	}
+
+	if (DelegateToRemoveFromNotifierB.IsValid() && NotifierB.IsValid())
+	{
+		check(NotifierB.Pin()->Delegate().Remove(DelegateToRemoveFromNotifierB));
+	}
+}
diff --git a/Engine/Source/Developer/SkeletalMeshUtilitiesCommon/Public/SkeletalMeshNotifier.h b/Engine/Source/Developer/SkeletalMeshUtilitiesCommon/Public/SkeletalMeshNotifier.h
--- a/Engine/Source/Developer/SkeletalMeshUtilitiesCommon/Public/SkeletalMeshNotifier.h
+++ b/Engine/Source/Developer/SkeletalMeshUtilitiesCommon/Public/SkeletalMeshNotifier.h
@@ -43,6 +43,21 @@ private:
 	mutable bool bNotifying = false;
 };
 
+class FSkeletalMeshNotifierBindScope 
+{
+	
+public:
+	UE_API FSkeletalMeshNotifierBindScope(TWeakPtr<ISkeletalMeshNotifier> InNotifierA,  TWeakPtr<ISkeletalMeshNotifier> InNotifierB);
+	UE_API ~FSkeletalMeshNotifierBindScope();
+
+private:
+	TWeakPtr<ISkeletalMeshNotifier> NotifierA;
+	TWeakPtr<ISkeletalMeshNotifier> NotifierB;
+
+	FDelegateHandle	DelegateToRemoveFromNotifierA;	
+	FDelegateHandle	DelegateToRemoveFromNotifierB;	
+};
+
 class ISkeletalMeshEditorBinding
 {
 public:
