Refactor weight map density settings properties so they can be re-used by the simplify tool

diff --git a/Engine/Plugins/Runtime/MeshModelingToolset/Source/MeshModelingTools/Private/Properties/MeshWeightChannelProperties.cpp b/Engine/Plugins/Runtime/MeshModelingToolset/Source/MeshModelingTools/Private/Properties/MeshWeightChannelProperties.cpp
new file mode 100644
--- /dev/null
+++ b/Engine/Plugins/Runtime/MeshModelingToolset/Source/MeshModelingTools/Private/Properties/MeshWeightChannelProperties.cpp
@@ -0,0 +1,54 @@
+// Copyright Epic Games, Inc. All Rights Reserved.
+
+#include "Properties/MeshWeightChannelProperties.h"
+
+#include "DynamicMesh/DynamicMesh3.h"
+#include "DynamicMesh/DynamicMeshAttributeSet.h"
+
+#include UE_INLINE_GENERATED_CPP_BY_NAME(MeshWeightChannelProperties)
+
+void UMeshWeightChannelDensityProperties::InitializeAttributes(const UE::Geometry::FDynamicMesh3* const Mesh)
+{
+	using namespace UE::Geometry;
+
+	AttributeNames.Reset();
+	if (const FDynamicMeshAttributeSet* Attributes = Mesh->Attributes())
+	{
+		const int32 NumLayers = Attributes->NumWeightLayers();
+		AttributeNames.SetNum(NumLayers);
+		for (int32 LayerIdx = 0; LayerIdx < NumLayers; ++LayerIdx)
+		{
+			AttributeNames[LayerIdx] = Attributes->GetWeightLayer(LayerIdx)->GetName();
+		}
+	}
+}
+
+void UMeshWeightChannelDensityProperties::ValidateSelectedAttribute()
+{
+	const int32 FoundIndex = AttributeNames.IndexOfByKey(SelectedAttribute);
+	if (FoundIndex == INDEX_NONE)
+	{
+		SelectedAttribute = (AttributeNames.Num() > 0) ? AttributeNames[0] : TEXT("");
+	}
+}
+
+int32 UMeshWeightChannelDensityProperties::GetSelectedWeightAttributeIndex(const UE::Geometry::FDynamicMesh3* const Mesh) const
+{
+	using namespace UE::Geometry;
+
+	int32 WeightMapAttributeLayerIndex = -1;
+	if (const FDynamicMeshAttributeSet* const Attribs = Mesh->Attributes())
+	{
+		for (int32 AttributeIndex = 0; AttributeIndex < Attribs->NumWeightLayers(); ++AttributeIndex)
+		{
+			const FDynamicMeshWeightAttribute* const AttributeLayer = Attribs->GetWeightLayer(AttributeIndex);
+
+			if (AttributeLayer->GetName() == SelectedAttribute)
+			{
+				WeightMapAttributeLayerIndex = AttributeIndex;
+				break;
+			}
+		}
+	}
+	return WeightMapAttributeLayerIndex;
+}
\ No newline at end of file
diff --git a/Engine/Plugins/Runtime/MeshModelingToolset/Source/MeshModelingTools/Private/RemeshMeshTool.cpp b/Engine/Plugins/Runtime/MeshModelingToolset/Source/MeshModelingTools/Private/RemeshMeshTool.cpp
--- a/Engine/Plugins/Runtime/MeshModelingToolset/Source/MeshModelingTools/Private/RemeshMeshTool.cpp
+++ b/Engine/Plugins/Runtime/MeshModelingToolset/Source/MeshModelingTools/Private/RemeshMeshTool.cpp
@@ -32,27 +32,7 @@ using namespace UE::Geometry;
 
 #define LOCTEXT_NAMESPACE "URemeshMeshTool"
 
-/*
- * Tool properties
- */
-
-void URemeshMeshToolProperties::InitializeAttributes(const TArray<FName>& InAttributeNames)
-{
-	AttributeNames.Init(FString(), InAttributeNames.Num());
-	for (int32 AttributeIndex = 0; AttributeIndex < InAttributeNames.Num(); ++AttributeIndex)
-	{
-		AttributeNames[AttributeIndex] = InAttributeNames[AttributeIndex].ToString();
-	}
-}
 
-void URemeshMeshToolProperties::ValidateSelectedAttribute()
-{
-	const int32 FoundIndex = AttributeNames.IndexOfByKey(SelectedAttribute);
-	if (FoundIndex == INDEX_NONE)
-	{
-		SelectedAttribute = (AttributeNames.Num() > 0) ? AttributeNames[0] : TEXT("");
-	}
-}
 
 /*
  * ToolBuilder
@@ -157,28 +137,21 @@ void URemeshMeshTool::Setup()
 	BasicProperties->TargetTriangleCount = FMath::Clamp(OriginalMesh->TriangleCount(), MinInitialTriangles, MaxInitialTriangles);
 	BasicProperties->TargetEdgeLength = FRemeshMeshOp::CalculateTargetEdgeLength(OriginalMesh.Get(), BasicProperties->TargetTriangleCount, InitialMeshArea);
 
+	WeightMapProperties = NewObject<UMeshWeightChannelDensityProperties>(this);
+	WeightMapProperties->RestoreProperties(this);
+	WeightMapProperties->InitializeAttributes(OriginalMesh.Get());
+	WeightMapProperties->ValidateSelectedAttribute();
+
 	// Attribute properties
-	BasicProperties->WatchProperty(BasicProperties->SelectedAttribute,
-		[this](const FString& NewName)
+	WeightMapProperties->WatchProperty(WeightMapProperties->SelectedAttribute,
+		[this](const FName& NewName)
 		{
-			BasicProperties->ValidateSelectedAttribute();
+			WeightMapProperties->ValidateSelectedAttribute();
 		});
 
-	TArray<FName> AttributeNames;
-	if (FDynamicMeshAttributeSet* Attributes = OriginalMesh->Attributes())
-	{
-		const int32 NumLayers = Attributes->NumWeightLayers();
-		AttributeNames.SetNum(NumLayers);
-		for (int32 LayerIdx = 0; LayerIdx < NumLayers; ++LayerIdx)
-		{
-			AttributeNames[LayerIdx] = Attributes->GetWeightLayer(LayerIdx)->GetName();
-		}
-	}
-	BasicProperties->InitializeAttributes(AttributeNames);
-	BasicProperties->ValidateSelectedAttribute();
-
 	// add properties to GUI
 	AddToolPropertySource(BasicProperties);
+	AddToolPropertySource(WeightMapProperties);
 	AddToolPropertySource(MeshStatisticsProperties);
 
 
@@ -215,6 +188,7 @@ void URemeshMeshTool::Setup()
 void URemeshMeshTool::OnShutdown(EToolShutdownType ShutdownType)
 {
 	BasicProperties->SaveProperties(this);
+	WeightMapProperties->SaveProperties(this);
 
 	if (ensure(MeshElementsDisplay->Settings))
 	{
@@ -290,26 +264,13 @@ TUniquePtr<FDynamicMeshOperator> URemeshMeshTool::MakeNewOperator()
 	Op->ProjectionTargetSpatial = nullptr;
 
 
-	if (BasicProperties->bUseWeightMap && OriginalMesh->HasAttributes())
+	if (WeightMapProperties->bUseWeightMap)
 	{
-		int32 WeightMapAttributeLayerIndex = -1;
-		if (const FDynamicMeshAttributeSet* const Attribs = OriginalMesh->Attributes())
-		{
-			for (int32 AttributeIndex = 0; AttributeIndex < Attribs->NumWeightLayers(); ++AttributeIndex)
-			{
-				const FDynamicMeshWeightAttribute* const AttributeLayer = Attribs->GetWeightLayer(AttributeIndex);
-
-				if (AttributeLayer->GetName() == BasicProperties->SelectedAttribute)
-				{
-					WeightMapAttributeLayerIndex = AttributeIndex;
-					break;
-				}
-			}
-		}
+		int32 WeightMapAttributeLayerIndex = WeightMapProperties->GetSelectedWeightAttributeIndex(OriginalMesh.Get());
 
 		if (WeightMapAttributeLayerIndex >= 0)
 		{
-			Op->CustomEdgeLengthScaleF = [WeightMapAttributeLayerIndex, RelativeDensity = BasicProperties->RelativeDensity](const FDynamicMesh3& Mesh, int VertexA, int VertexB) -> double
+			Op->CustomEdgeLengthScaleF = [WeightMapAttributeLayerIndex, RelativeDensity = WeightMapProperties->RelativeDensity](const FDynamicMesh3& Mesh, int VertexA, int VertexB) -> double
 			{
 				if (Mesh.HasAttributes() && Mesh.Attributes()->GetWeightLayer(WeightMapAttributeLayerIndex))
 				{
diff --git a/Engine/Plugins/Runtime/MeshModelingToolset/Source/MeshModelingTools/Public/Properties/MeshWeightChannelProperties.h b/Engine/Plugins/Runtime/MeshModelingToolset/Source/MeshModelingTools/Public/Properties/MeshWeightChannelProperties.h
new file mode 100644
--- /dev/null
+++ b/Engine/Plugins/Runtime/MeshModelingToolset/Source/MeshModelingTools/Public/Properties/MeshWeightChannelProperties.h
@@ -0,0 +1,53 @@
+// Copyright Epic Games, Inc. All Rights Reserved.
+
+#pragma once
+
+#include "CoreMinimal.h"
+#include "InteractiveTool.h"
+#include "GeometryBase.h"
+#include "MeshWeightChannelProperties.generated.h"
+
+#define UE_API MESHMODELINGTOOLS_API
+
+PREDECLARE_USE_GEOMETRY_CLASS(FDynamicMesh3);
+
+
+UCLASS(MinimalAPI)
+class UMeshWeightChannelDensityProperties : public UInteractiveToolPropertySet
+{
+	GENERATED_BODY()
+public:
+
+	/** Whether to use a vertex weight map layer to control relative mesh density */
+	UPROPERTY(EditAnywhere, Category = WeightMap)
+	bool bUseWeightMap;
+
+	/** Name of the vertex weight map layer that will be used to adjust relative mesh density */
+	UPROPERTY(EditAnywhere, Category = WeightMap, meta = (DisplayName = "Attribute", GetOptions = GetAttributeNames, EditCondition = bUseWeightMap))
+	FName SelectedAttribute;
+
+	UFUNCTION()
+	const TArray<FName>& GetAttributeNames() { return AttributeNames; };
+
+	TArray<FName> AttributeNames;
+
+	/** Amount to adjust density in regions where the vertex weight map is greater than zero. A positive value here results in a denser mesh, a negative value results in a coarser one */
+	UPROPERTY(EditAnywhere, Category = WeightMap, meta = (EditCondition = bUseWeightMap, UIMin = "-2.0", UIMax = "2.0"))
+	float RelativeDensity = 1.f;
+
+	/**
+	 * Set up the list of available attributes
+	 */
+	UE_API void InitializeAttributes(const UE::Geometry::FDynamicMesh3* const Mesh);
+
+	/**
+	 * If the attribute selection is not valid, use the attribute at index 0 or empty if there are no attributes
+	 */
+	UE_API void ValidateSelectedAttribute();
+
+	/** @return the selected weight attribute on the given dynamic mesh or INDEX_NONE if not found */
+	UE_API int32 GetSelectedWeightAttributeIndex(const UE::Geometry::FDynamicMesh3* const Mesh) const;
+};
+
+
+#undef UE_API
diff --git a/Engine/Plugins/Runtime/MeshModelingToolset/Source/MeshModelingTools/Public/RemeshMeshTool.h b/Engine/Plugins/Runtime/MeshModelingToolset/Source/MeshModelingTools/Public/RemeshMeshTool.h
--- a/Engine/Plugins/Runtime/MeshModelingToolset/Source/MeshModelingTools/Public/RemeshMeshTool.h
+++ b/Engine/Plugins/Runtime/MeshModelingToolset/Source/MeshModelingTools/Public/RemeshMeshTool.h
@@ -8,6 +8,7 @@
 #include "MeshOpPreviewHelpers.h"
 #include "CleaningOps/RemeshMeshOp.h"
 #include "Properties/RemeshProperties.h"
+#include "Properties/MeshWeightChannelProperties.h"
 #include "RemeshMeshTool.generated.h"
 
 #define UE_API MESHMODELINGTOOLS_API
@@ -100,33 +101,6 @@ public:
 	UPROPERTY(EditAnywhere, Category = Remeshing, AdvancedDisplay, meta = (EditCondition = "bReprojectConstraints", UIMin = "0", UIMax = "180", ClampMin = "0", ClampMax = "180"))
 	float BoundaryCornerAngleThreshold;
 
-
-	/** Whether to use a vertex weight map layer to control relative mesh density */
-	UPROPERTY(EditAnywhere, Category = WeightMap)
-	bool bUseWeightMap;
-
-	/** Name of the vertex weight map layer that will be used to adjust relative mesh density */
-	UPROPERTY(EditAnywhere, Category = WeightMap, meta = (DisplayName = "Attribute", GetOptions = GetAttributeNames, EditCondition = bUseWeightMap))
-	FString SelectedAttribute;
-
-	UFUNCTION()
-	const TArray<FString>& GetAttributeNames() { return AttributeNames; };
-	TArray<FString> AttributeNames;
-
-	/** Amount to adjust density in regions where the vertex weight map is greater than zero. A positive value here results in a denser mesh, a negative value results in a coarser one */
-	UPROPERTY(EditAnywhere, Category = WeightMap, meta = (EditCondition = bUseWeightMap, UIMin = "-2.0", UIMax = "2.0"))
-	float RelativeDensity = 0.f;
-
-	/**
-	 * Set up the list of available attributes
-	 */
-	void InitializeAttributes(const TArray<FName>& InAttributeNames);
-
-	/**
-	 * If the attribute selection is not valid, use the attribute at index 0 or empty if there are no attributes
-	 */
-	void ValidateSelectedAttribute();
-
 };
 
 
@@ -164,6 +138,9 @@ public:
 	UPROPERTY()
 	TObjectPtr<URemeshMeshToolProperties> BasicProperties;
 
+	UPROPERTY()
+	TObjectPtr<UMeshWeightChannelDensityProperties> WeightMapProperties;
+
 	UPROPERTY()
 	TObjectPtr<UMeshStatisticsProperties> MeshStatisticsProperties;
 
