add option for the mesh simplify tool to drive simplification by a weight map
+ expose the quadric error tolerance as simplification target

diff --git a/Engine/Plugins/Runtime/MeshModelingToolset/Source/MeshModelingToolsEditorOnly/Private/SimplifyMeshTool.cpp b/Engine/Plugins/Runtime/MeshModelingToolset/Source/MeshModelingToolsEditorOnly/Private/SimplifyMeshTool.cpp
--- a/Engine/Plugins/Runtime/MeshModelingToolset/Source/MeshModelingToolsEditorOnly/Private/SimplifyMeshTool.cpp
+++ b/Engine/Plugins/Runtime/MeshModelingToolset/Source/MeshModelingToolsEditorOnly/Private/SimplifyMeshTool.cpp
@@ -118,6 +118,21 @@ void USimplifyMeshTool::Setup()
 	SimplifyProperties->WatchProperty(SimplifyProperties->bShowGroupColors,
 									  [this](bool bNewValue) { UpdateVisualization(); });
 
+	WeightMapProperties = NewObject<UMeshWeightChannelDensityProperties>(this);
+	WeightMapProperties->RestoreProperties(this);
+	AddToolPropertySource(WeightMapProperties);
+
+	WeightMapProperties->InitializeAttributes(OriginalMesh.Get());
+	WeightMapProperties->ValidateSelectedAttribute();
+
+	// Attribute properties
+	WeightMapProperties->WatchProperty(WeightMapProperties->SelectedAttribute,
+		[this](const FName& NewName)
+		{
+			WeightMapProperties->ValidateSelectedAttribute();
+		});
+	
+
 	MeshStatisticsProperties = NewObject<UMeshStatisticsProperties>(this);
 	AddToolPropertySource(MeshStatisticsProperties);
 
@@ -162,6 +177,7 @@ bool USimplifyMeshTool::CanAccept() const
 void USimplifyMeshTool::OnShutdown(EToolShutdownType ShutdownType)
 {
 	SimplifyProperties->SaveProperties(this);
+	WeightMapProperties->SaveProperties(this);
 	
 	if (ensure(MeshElementsDisplay->Settings))
 	{
@@ -221,6 +237,7 @@ TUniquePtr<FDynamicMeshOperator> USimplifyMeshTool::MakeNewOperator()
 	Op->TargetEdgeLength = SimplifyProperties->TargetEdgeLength;
 	Op->TargetMode = SimplifyProperties->TargetMode;
 	Op->TargetPercentage = SimplifyProperties->TargetPercentage;
+	Op->TargetQuadricError = SimplifyProperties->TargetQuadricErrorRoot * SimplifyProperties->TargetQuadricErrorRoot;
 	Op->MeshBoundaryConstraint = (EEdgeRefineFlags)SimplifyProperties->MeshBoundaryConstraint;
 	Op->GroupBoundaryConstraint = (EEdgeRefineFlags)SimplifyProperties->GroupBoundaryConstraint;
 	Op->MaterialBoundaryConstraint = (EEdgeRefineFlags)SimplifyProperties->MaterialBoundaryConstraint;
@@ -238,6 +255,63 @@ TUniquePtr<FDynamicMeshOperator> USimplifyMeshTool::MakeNewOperator()
 	IMeshReductionManagerModule& MeshReductionModule = FModuleManager::Get().LoadModuleChecked<IMeshReductionManagerModule>("MeshReductionInterface");
 	Op->MeshReduction = MeshReductionModule.GetStaticMeshReductionInterface();
 
+	if (WeightMapProperties->bUseWeightMap)
+	{
+		int32 WeightMapAttributeLayerIndex = WeightMapProperties->GetSelectedWeightAttributeIndex(OriginalMesh.Get());
+
+		if (WeightMapAttributeLayerIndex >= 0)
+		{
+			TFunction<double(const FDynamicMesh3& Mesh, int VertexA, int VertexB)> EdgeMetricScale =
+			[WeightMapAttributeLayerIndex, RelativeDensity = WeightMapProperties->RelativeDensity](const FDynamicMesh3& Mesh, int VertexA, int VertexB) -> double
+			{
+				if (Mesh.HasAttributes() && Mesh.Attributes()->GetWeightLayer(WeightMapAttributeLayerIndex))
+				{
+					const FDynamicMeshWeightAttribute* const AttributeLayer = Mesh.Attributes()->GetWeightLayer(WeightMapAttributeLayerIndex);
+
+					float WeightValueA;
+					AttributeLayer->GetValue(VertexA, &WeightValueA);
+					WeightValueA = FMath::Clamp(WeightValueA, 0.0f, 1.0f);
+
+					float WeightValueB;
+					AttributeLayer->GetValue(VertexB, &WeightValueB);
+					WeightValueB = FMath::Clamp(WeightValueB, 0.0f, 1.0f);
+
+					// choose the weight that will lead to a larger measure, so we choose the more conservative of the two weights
+					const float UseWeight = RelativeDensity < 0 ? FMath::Max(WeightValueA, WeightValueB) : FMath::Min(WeightValueA, WeightValueB);
+
+					return FMath::Pow(0.5, -RelativeDensity * (double)UseWeight);
+				}
+
+				return 1.0;
+			};
+
+			// geometric tolerance can be used in multiple modes, so we pass it through always
+			Op->CustomGeometricErrorScaleF = EdgeMetricScale;
+			// only scale quadric error or edge lengths when they are the active target, so weights scale the most relevant metric
+			if (SimplifyProperties->TargetMode == ESimplifyTargetType::QuadricError)
+			{
+				Op->CustomQuadricErrorScaleF = EdgeMetricScale;
+			}
+			if (SimplifyProperties->TargetMode == ESimplifyTargetType::EdgeLength)
+			{
+				Op->CustomEdgeLengthScaleF = EdgeMetricScale;
+			}
+
+			TFunction<double(int Vertex)> ClusterEdgeTargetLengthScale =
+			[AttributeLayer = OriginalMesh->Attributes()->GetWeightLayer(WeightMapAttributeLayerIndex),
+				RelativeDensity = WeightMapProperties->RelativeDensity](int Vertex) -> double
+			{
+				float WeightValue;
+				AttributeLayer->GetValue(Vertex, &WeightValue);
+				WeightValue = FMath::Clamp(WeightValue, 0.0f, 1.0f);
+
+				// Note: We use positive RelativeDensity here instead of negative, b/c cluster simplifier scales target length rather measured length
+				return FMath::Pow(0.5, RelativeDensity * WeightValue);
+			};
+			Op->ClusterTargetEdgeLengthScaleF = ClusterEdgeTargetLengthScale;
+		}
+	}
+
 	return Op;
 }
 
diff --git a/Engine/Plugins/Runtime/MeshModelingToolset/Source/MeshModelingToolsEditorOnly/Public/SimplifyMeshTool.h b/Engine/Plugins/Runtime/MeshModelingToolset/Source/MeshModelingToolsEditorOnly/Public/SimplifyMeshTool.h
--- a/Engine/Plugins/Runtime/MeshModelingToolset/Source/MeshModelingToolsEditorOnly/Public/SimplifyMeshTool.h
+++ b/Engine/Plugins/Runtime/MeshModelingToolset/Source/MeshModelingToolsEditorOnly/Public/SimplifyMeshTool.h
@@ -7,6 +7,7 @@
 #include "BaseTools/SingleSelectionMeshEditingTool.h"
 #include "MeshOpPreviewHelpers.h"
 #include "Properties/RemeshProperties.h"
+#include "Properties/MeshWeightChannelProperties.h"
 #include "CleaningOps/SimplifyMeshOp.h"		// required in header for enum types
 #include "SimplifyMeshTool.generated.h"
 
@@ -59,6 +60,11 @@ public:
 		EditCondition = "SimplifierType == ESimplifyType::ClusterBased || (TargetMode == ESimplifyTargetType::EdgeLength && SimplifierType != ESimplifyType::UEStandard && SimplifierType != ESimplifyType::MinimalPlanar && SimplifierType != ESimplifyType::MinimalPolygroup)"))
 	float TargetEdgeLength;
 
+	/** Square root of target quadric error */
+	UPROPERTY(EditAnywhere, Category = Options, meta = (UIMin = "0.0", UIMax = "100.0", ClampMin = "0.", ClampMax = "10000.0",
+		EditCondition = "TargetMode == ESimplifyTargetType::QuadricError && (SimplifierType != ESimplifyType::UEStandard && SimplifierType != ESimplifyType::MinimalPlanar && SimplifierType != ESimplifyType::MinimalPolygroup && SimplifierType != ESimplifyType::ClusterBased)"))
+	float TargetQuadricErrorRoot = 0;
+
 	/** Target triangle count */
 	UPROPERTY(EditAnywhere, Category = Options, meta = (UIMin = "4", UIMax = "10000", ClampMin = "1", ClampMax = "9999999999",
 				EditCondition = "TargetMode == ESimplifyTargetType::TriangleCount && SimplifierType != ESimplifyType::MinimalPlanar && SimplifierType != ESimplifyType::MinimalPolygroup && SimplifierType != ESimplifyType::ClusterBased"))
@@ -153,6 +159,9 @@ private:
 	UPROPERTY()
 	TObjectPtr<USimplifyMeshToolProperties> SimplifyProperties;
 
+	UPROPERTY()
+	TObjectPtr<UMeshWeightChannelDensityProperties> WeightMapProperties;
+
 	UPROPERTY()
 	TObjectPtr<UMeshStatisticsProperties> MeshStatisticsProperties;
 
diff --git a/Engine/Plugins/Runtime/MeshModelingToolset/Source/ModelingOperatorsEditorOnly/Private/CleaningOps/SimplifyMeshOp.cpp b/Engine/Plugins/Runtime/MeshModelingToolset/Source/ModelingOperatorsEditorOnly/Private/CleaningOps/SimplifyMeshOp.cpp
--- a/Engine/Plugins/Runtime/MeshModelingToolset/Source/ModelingOperatorsEditorOnly/Private/CleaningOps/SimplifyMeshOp.cpp
+++ b/Engine/Plugins/Runtime/MeshModelingToolset/Source/ModelingOperatorsEditorOnly/Private/CleaningOps/SimplifyMeshOp.cpp
@@ -36,11 +36,14 @@ void ComputeSimplify(FDynamicMesh3* TargetMesh, const bool bReproject,
 					 EEdgeRefineFlags MaterialBoundaryConstraint,
 					 bool bPreserveSharpEdges, bool bAllowSeamCollapse, bool bPreventNormalFlips, bool bPreventTinyTriangles,
 					 const ESimplifyTargetType TargetMode,
-					 const float TargetPercentage, const int TargetCount, const float TargetEdgeLength,
+					 const float TargetPercentage, const int TargetCount, const float TargetEdgeLength, const float TargetQuadricError,
 					 const float AngleThreshold,
 	                 typename SimplificationType::ESimplificationCollapseModes CollapseMode,
 					 bool bUseQuadricMemory,
-					 float GeometricTolerance )
+					 float GeometricTolerance,
+	TFunction<double(const FDynamicMesh3&, int VertexA, int VertexB)> CustomGeometricErrorScaleF,
+	TFunction<double(const FDynamicMesh3&, int VertexA, int VertexB)> CustomQuadricErrorScaleF,
+	TFunction<double(const FDynamicMesh3&, int VertexA, int VertexB)> CustomEdgeLengthScaleF)
 {
 	SimplificationType Reducer(TargetMesh);
 
@@ -53,6 +56,10 @@ void ComputeSimplify(FDynamicMesh3* TargetMesh, const bool bReproject,
 	Reducer.bAllowSeamCollapse = bAllowSeamCollapse;
 	Reducer.bRetainQuadricMemory = bUseQuadricMemory;
 
+	Reducer.CustomGeometricErrorScaleF = CustomGeometricErrorScaleF;
+	Reducer.CustomQuadricErrorScaleF = CustomQuadricErrorScaleF;
+	Reducer.CustomEdgeLengthScaleF = CustomEdgeLengthScaleF;
+
 	if (bAllowSeamCollapse)
 	{
 		Reducer.SetEdgeFlipTolerance(1.e-5);
@@ -127,6 +134,10 @@ void ComputeSimplify(FDynamicMesh3* TargetMesh, const bool bReproject,
 	{
 		Reducer.SimplifyToMinimalPlanar(AngleThreshold);
 	}
+	else if (TargetMode == ESimplifyTargetType::QuadricError)
+	{
+		Reducer.SimplifyToMaxError(TargetQuadricError);
+	}
 }
 
 bool FSimplifyMeshOp::ComputeStandardSimplifier(IMeshReduction* MeshReduction, const FMeshDescription& SrcMeshDescription, FDynamicMesh3& OutResult, 
@@ -237,9 +248,12 @@ void FSimplifyMeshOp::CalculateResult(FProgressCancel* Progress)
 											GroupBoundaryConstraint,
 											MaterialBoundaryConstraint,
 											bPreserveSharpEdges, bAllowSeamCollapse, bPreventNormalFlips, bPreventTinyTriangles,
-											TargetMode, TargetPercentage, TargetCount, TargetEdgeLength, MinimalPlanarAngleThresh,
+											TargetMode, TargetPercentage, TargetCount, TargetEdgeLength, TargetQuadricError, MinimalPlanarAngleThresh,
 											FQEMSimplification::ESimplificationCollapseModes::MinimalQuadricPositionError, bUseQuadricMemory, 
-											UseGeometricTolerance);
+											UseGeometricTolerance,
+											CustomGeometricErrorScaleF,
+											CustomQuadricErrorScaleF,
+											CustomEdgeLengthScaleF);
 	}
 	else if (SimplifierType == ESimplifyType::Attribute)
 	{
@@ -254,9 +268,12 @@ void FSimplifyMeshOp::CalculateResult(FProgressCancel* Progress)
 													GroupBoundaryConstraint,
 													MaterialBoundaryConstraint,
 													bPreserveSharpEdges, bAllowSeamCollapse, bPreventNormalFlips, bPreventTinyTriangles,
-													TargetMode, TargetPercentage, TargetCount, TargetEdgeLength, MinimalPlanarAngleThresh,
+													TargetMode, TargetPercentage, TargetCount, TargetEdgeLength, TargetQuadricError, MinimalPlanarAngleThresh,
 													FAttrMeshSimplification::ESimplificationCollapseModes::MinimalQuadricPositionError, bUseQuadricMemory,
-													UseGeometricTolerance);
+													UseGeometricTolerance,
+													CustomGeometricErrorScaleF,
+													CustomQuadricErrorScaleF,
+													CustomEdgeLengthScaleF);
 	}
 	else if (SimplifierType == ESimplifyType::MinimalPlanar)
 	{
@@ -271,9 +288,12 @@ void FSimplifyMeshOp::CalculateResult(FProgressCancel* Progress)
 			GroupBoundaryConstraint,
 			MaterialBoundaryConstraint,
 			bPreserveSharpEdges, bAllowSeamCollapse, bPreventNormalFlips, bPreventTinyTriangles,
-			ESimplifyTargetType::MinimalPlanar, TargetPercentage, TargetCount, TargetEdgeLength, MinimalPlanarAngleThresh,
+			ESimplifyTargetType::MinimalPlanar, TargetPercentage, TargetCount, TargetEdgeLength, TargetQuadricError, MinimalPlanarAngleThresh,
 			FQEMSimplification::ESimplificationCollapseModes::MinimalQuadricPositionError, bUseQuadricMemory,
-			UseGeometricTolerance);
+			UseGeometricTolerance,
+			CustomGeometricErrorScaleF,
+			CustomQuadricErrorScaleF,
+			CustomEdgeLengthScaleF);
 	}
 	else if (SimplifierType == ESimplifyType::MinimalExistingVertex)
 	{
@@ -284,9 +304,12 @@ void FSimplifyMeshOp::CalculateResult(FProgressCancel* Progress)
 			GroupBoundaryConstraint,
 			MaterialBoundaryConstraint,
 			bPreserveSharpEdges, bAllowSeamCollapse, bPreventNormalFlips, bPreventTinyTriangles,
-			TargetMode, TargetPercentage, TargetCount, TargetEdgeLength, MinimalPlanarAngleThresh, 
+			TargetMode, TargetPercentage, TargetCount, TargetEdgeLength, TargetQuadricError, MinimalPlanarAngleThresh,
 			FQEMSimplification::ESimplificationCollapseModes::MinimalExistingVertexError, bUseQuadricMemory,
-			UseGeometricTolerance);
+			UseGeometricTolerance,
+			CustomGeometricErrorScaleF,
+			CustomQuadricErrorScaleF,
+			CustomEdgeLengthScaleF);
 	}
 	else if (SimplifierType == ESimplifyType::MinimalPolygroup)
 	{
@@ -306,6 +329,7 @@ void FSimplifyMeshOp::CalculateResult(FProgressCancel* Progress)
 		MeshClusterSimplify::FSimplifyOptions SimplifyOptions;
 		SimplifyOptions.TargetEdgeLength = TargetEdgeLength;
 		SimplifyOptions.FixBoundaryAngleTolerance = BoundaryEdgeAngleTolerance;
+		SimplifyOptions.OptionalTargetEdgeLengthScale = ClusterTargetEdgeLengthScaleF;
 
 		auto RefineFlagsToConstraintLevel = [](EEdgeRefineFlags Flags)
 		{
diff --git a/Engine/Plugins/Runtime/MeshModelingToolset/Source/ModelingOperatorsEditorOnly/Public/CleaningOps/SimplifyMeshOp.h b/Engine/Plugins/Runtime/MeshModelingToolset/Source/ModelingOperatorsEditorOnly/Public/CleaningOps/SimplifyMeshOp.h
--- a/Engine/Plugins/Runtime/MeshModelingToolset/Source/ModelingOperatorsEditorOnly/Public/CleaningOps/SimplifyMeshOp.h
+++ b/Engine/Plugins/Runtime/MeshModelingToolset/Source/ModelingOperatorsEditorOnly/Public/CleaningOps/SimplifyMeshOp.h
@@ -31,6 +31,9 @@ enum class ESimplifyTargetType : uint8
 	/** Target edge length */
 	EdgeLength = 3 UMETA(DisplayName = "Edge Length"),
 
+	/** Quadric error threshold */
+	QuadricError = 5 UMETA(DisplayName = "Quadric Error"),
+
 	/** Apply all allowable edge collapses that do not change the shape */
 	MinimalPlanar = 4 UMETA(Hidden)
 };
@@ -79,6 +82,7 @@ public:
 	ESimplifyType SimplifierType;
 	int TargetPercentage, TargetCount;
 	float TargetEdgeLength;
+	float TargetQuadricError;
 	bool bDiscardAttributes, bReproject, bPreventNormalFlips, bPreserveSharpEdges, bAllowSeamCollapse, bPreventTinyTriangles;
 	// When true, result will have attributes object regardless of whether attributes 
 	// were discarded or present initially.
@@ -96,6 +100,14 @@ public:
 	bool bGeometricDeviationConstraint = false;
 	float GeometricTolerance = 0.0f;
 
+	// QEM simplifier custom scales to apply to *measured* errors and edge lengths. Note this is in effect the inverse of scaling the tolerances / target lengths.
+	TFunction<double(const FDynamicMesh3&, int VertexA, int VertexB)> CustomGeometricErrorScaleF;
+	TFunction<double(const FDynamicMesh3&, int VertexA, int VertexB)> CustomQuadricErrorScaleF;
+	TFunction<double(const FDynamicMesh3&, int VertexA, int VertexB)> CustomEdgeLengthScaleF;
+
+	// Cluster simplifier custom scales apply to the *target* edge lengths, and so is conceptually the inverse of the above scales.
+	TFunction<double(int32)> ClusterTargetEdgeLengthScaleF = nullptr;
+
 	// stored for the UEStandard path
 	TSharedPtr<FMeshDescription, ESPMode::ThreadSafe> OriginalMeshDescription;
 	// stored for the GeometryProcessing custom simplifier paths (currently precomputed once in tool setup)
