[SkinWeightsPaintTool] Added hotkeys for toggle lock and bypass lock

diff --git a/Engine/Plugins/Animation/SkeletalMeshModelingTools/Source/SkeletalMeshModelingTools/Private/SkeletalMeshModelingToolsCommands.cpp b/Engine/Plugins/Animation/SkeletalMeshModelingTools/Source/SkeletalMeshModelingTools/Private/SkeletalMeshModelingToolsCommands.cpp
--- a/Engine/Plugins/Animation/SkeletalMeshModelingTools/Source/SkeletalMeshModelingTools/Private/SkeletalMeshModelingToolsCommands.cpp
+++ b/Engine/Plugins/Animation/SkeletalMeshModelingTools/Source/SkeletalMeshModelingTools/Private/SkeletalMeshModelingToolsCommands.cpp
@@ -3,6 +3,7 @@
 #include "SkeletalMeshModelingToolsCommands.h"
 
 #include "SkeletalMesh/SkeletonEditingTool.h"
+#include "SkeletalMesh/SkinWeightsPaintTool.h"
 
 #define LOCTEXT_NAMESPACE "SkeletalMeshModelingToolsCommands"
 
@@ -35,11 +36,13 @@ void FSkeletalMeshModelingToolsActionCommands::GetToolDefaultObjectList(TArray<U
 void FSkeletalMeshModelingToolsActionCommands::RegisterAllToolActions()
 {
 	FSkeletonEditingToolActionCommands::Register();
+	FSkinWeightsPaintToolActionCommands::Register();
 }
 
 void FSkeletalMeshModelingToolsActionCommands::UnregisterAllToolActions()
 {
 	FSkeletonEditingToolActionCommands::Unregister();
+	FSkinWeightsPaintToolActionCommands::Unregister();
 }
 
 void FSkeletalMeshModelingToolsActionCommands::UpdateToolCommandBinding(UInteractiveTool* Tool, TSharedPtr<FUICommandList> UICommandList, bool bUnbind)
@@ -50,6 +53,11 @@ void FSkeletalMeshModelingToolsActionCommands::UpdateToolCommandBinding(UInterac
 	{
 		UPDATE_BINDING(FSkeletonEditingToolActionCommands);
 	}
+
+	if (ExactCast<USkinWeightsPaintTool>(Tool) != nullptr)
+	{
+		UPDATE_BINDING(FSkinWeightsPaintToolActionCommands);
+	}
 }
 
 #define DEFINE_TOOL_ACTION_COMMANDS(CommandsClassName, ContextNameString, SettingsDialogString, ToolClassName ) \
@@ -61,5 +69,6 @@ ToolCDOs.Add(GetMutableDefault<ToolClassName>()); \
 }
 
 DEFINE_TOOL_ACTION_COMMANDS(FSkeletonEditingToolActionCommands, "SkeletalMeshModelingToolsSkeletonEditing", "Skeletal Mesh Modeling Tools - Skeleton Editing Tool", USkeletonEditingTool);
+DEFINE_TOOL_ACTION_COMMANDS(FSkinWeightsPaintToolActionCommands, "SkeletalMeshModelingToolsSkinWeightsPaintTool", "Skeletal Mesh Modeling Tools - Skin Weights Paint Tool", USkinWeightsPaintTool);
 
 #undef LOCTEXT_NAMESPACE
diff --git a/Engine/Plugins/Animation/SkeletalMeshModelingTools/Source/SkeletalMeshModelingTools/Private/SkeletalMeshModelingToolsCommands.h b/Engine/Plugins/Animation/SkeletalMeshModelingTools/Source/SkeletalMeshModelingTools/Private/SkeletalMeshModelingToolsCommands.h
--- a/Engine/Plugins/Animation/SkeletalMeshModelingTools/Source/SkeletalMeshModelingTools/Private/SkeletalMeshModelingToolsCommands.h
+++ b/Engine/Plugins/Animation/SkeletalMeshModelingTools/Source/SkeletalMeshModelingTools/Private/SkeletalMeshModelingToolsCommands.h
@@ -51,4 +51,5 @@ CommandsClassName();\
 virtual void GetToolDefaultObjectList(TArray<UInteractiveTool*>& ToolCDOs) override;\
 };\
 
-DECLARE_TOOL_ACTION_COMMANDS(FSkeletonEditingToolActionCommands);
\ No newline at end of file
+DECLARE_TOOL_ACTION_COMMANDS(FSkeletonEditingToolActionCommands);
+DECLARE_TOOL_ACTION_COMMANDS(FSkinWeightsPaintToolActionCommands);
diff --git a/Engine/Plugins/Runtime/MeshModelingToolset/Source/MeshModelingToolsEditorOnly/Private/SkeletalMesh/SkinWeightsPaintTool.cpp b/Engine/Plugins/Runtime/MeshModelingToolset/Source/MeshModelingToolsEditorOnly/Private/SkeletalMesh/SkinWeightsPaintTool.cpp
--- a/Engine/Plugins/Runtime/MeshModelingToolset/Source/MeshModelingToolsEditorOnly/Private/SkeletalMesh/SkinWeightsPaintTool.cpp
+++ b/Engine/Plugins/Runtime/MeshModelingToolset/Source/MeshModelingToolsEditorOnly/Private/SkeletalMesh/SkinWeightsPaintTool.cpp
@@ -1870,6 +1870,7 @@ FSkinWeightsPaintToolCommands::FSkinWeightsPaintToolCommands()
 void FSkinWeightsPaintToolCommands::RegisterCommands()
 {
 	UI_COMMAND(IsolateSelectedBones, "Isolate Selected (Unlock)", "Isolate selected bones by locking all other bones, Skin weights for locked bones cannot be changed", EUserInterfaceActionType::Button, FInputChord());
+	UI_COMMAND(ToggleSelectedBonesLockState, "Toggle Selected Lock State", "Toggle the lock state for selected bones, Skin weights for locked bones cannot be changed", EUserInterfaceActionType::Button, FInputChord(EKeys::L));
 	
 	UI_COMMAND(LockSelectedBones, "Lock Selected", "Lock selected bones, Skin weights for locked bones cannot be changed", EUserInterfaceActionType::Button, FInputChord());	
 	UI_COMMAND(UnlockSelectedBones, "Unlock Selected", "Unlock selected bones, Skin weights for locked bones cannot be changed", EUserInterfaceActionType::Button, FInputChord());
@@ -1877,7 +1878,7 @@ void FSkinWeightsPaintToolCommands::RegisterCommands()
 	UI_COMMAND(LockAllBones, "Lock All", "Lock all bones, Skin weights for locked bones cannot be changed", EUserInterfaceActionType::Button, FInputChord());	
 	UI_COMMAND(UnlockAllBones, "Unlock All", "Unlock all bones, Skin weights for locked bones cannot be changed", EUserInterfaceActionType::Button, FInputChord());
 	
-	UI_COMMAND(BypassLocks, "Bypass Locks", "Disable skin weight locking, existing locks are preserved but disabled", EUserInterfaceActionType::ToggleButton, FInputChord());
+	UI_COMMAND(BypassLocks, "Bypass Locks", "Disable skin weight locking, existing locks are preserved but disabled", EUserInterfaceActionType::ToggleButton, FInputChord(EModifierKey::Alt, EKeys::L));
 	
 }
 
@@ -2753,6 +2754,38 @@ float USkinWeightsPaintTool::CalculateBrushFalloff(float Distance) const
 	return w;
 }
 
+void USkinWeightsPaintTool::RegisterActions(FInteractiveToolActionSet& ActionSet)
+{
+	Super::RegisterActions(ActionSet);
+	
+	
+	int32 ActionId = static_cast<int32>(EStandardToolActions::BaseClientDefinedActionID) + 400;
+	auto GetActionId = [&ActionId]
+		{
+			return ActionId++;
+		};
+
+	ActionSet.RegisterAction(this, GetActionId(), TEXT("Toggle Selected Lock State"),
+		LOCTEXT("ToggleSelectedLockState", "Toggle Selected Lock State"),
+		LOCTEXT("ToggleSelectedLockStateDesc", "Toggle the lock state for selected bone(s)"),
+		EModifierKey::None, EKeys::L,
+		[this]()
+		{
+			ToggleSelectedBonesLockState();
+			GetToolManager()->DisplayMessage(LOCTEXT("ToggleLockMessage", "Toggled the lock state for selected bone(s) "), EToolMessageLevel::UserNotification);
+		});
+
+	ActionSet.RegisterAction(this, GetActionId(), TEXT("Toggle Bypass Locks"),
+		LOCTEXT("ToggleBypassLock", "Toggle Bypass Locks"),
+		LOCTEXT("ToggleBypassLockDesc", "Toggle Bypass Locks"),
+		EModifierKey::Alt, EKeys::L,
+		[this]()
+		{
+			ToggleBypassLocks();
+			GetToolManager()->DisplayMessage(LOCTEXT("ToggleBypassLockMessage", "Toggled Lock Bypass Mode"), EToolMessageLevel::UserNotification);
+		});
+}
+
 void USkinWeightsPaintTool::ApplyStamp(const FBrushStampData& Stamp)
 {
 	TRACE_CPUPROFILER_EVENT_SCOPE(SkinTool::ApplyStamp);
@@ -4166,15 +4199,19 @@ TSharedPtr<SWidget> USkinWeightsPaintTool::GetCustomSkeletonTreeWidget(TSharedPt
 				FExecuteAction::CreateUObject(this, &USkinWeightsPaintTool::IsolateSelectedBones),
 				FCanExecuteAction::CreateUObject(this, &USkinWeightsPaintTool::HasSelectedBone));
 			
+			InOutCommandList->MapAction( Commands.ToggleSelectedBonesLockState,
+				FExecuteAction::CreateUObject(this, &USkinWeightsPaintTool::ToggleSelectedBonesLockState),
+				FCanExecuteAction::CreateUObject(this, &USkinWeightsPaintTool::HasSelectedBone));
+			
 			InOutCommandList->MapAction( Commands.LockSelectedBones,
-				FExecuteAction::CreateUObject(this, &USkinWeightsPaintTool::ToggleLockSelectedBones, true),
+				FExecuteAction::CreateUObject(this, &USkinWeightsPaintTool::SetLockSelectedBones, true),
 				FCanExecuteAction::CreateUObject(this, &USkinWeightsPaintTool::HasSelectedBone));
 			InOutCommandList->MapAction( Commands.UnlockSelectedBones,
-				FExecuteAction::CreateUObject(this, &USkinWeightsPaintTool::ToggleLockSelectedBones, false),
+				FExecuteAction::CreateUObject(this, &USkinWeightsPaintTool::SetLockSelectedBones, false),
 				FCanExecuteAction::CreateUObject(this, &USkinWeightsPaintTool::HasSelectedBone));
 			
-			InOutCommandList->MapAction( Commands.LockAllBones, FExecuteAction::CreateUObject(this, &USkinWeightsPaintTool::ToggleLockAllBones, true));
-			InOutCommandList->MapAction( Commands.UnlockAllBones, FExecuteAction::CreateUObject(this, &USkinWeightsPaintTool::ToggleLockAllBones, false));
+			InOutCommandList->MapAction( Commands.LockAllBones, FExecuteAction::CreateUObject(this, &USkinWeightsPaintTool::SetLockAllBones, true));
+			InOutCommandList->MapAction( Commands.UnlockAllBones, FExecuteAction::CreateUObject(this, &USkinWeightsPaintTool::SetLockAllBones, false));
 			
 			InOutCommandList->MapAction( Commands.BypassLocks,
 				FExecuteAction::CreateUObject(this, &USkinWeightsPaintTool::ToggleBypassLocks),
@@ -4184,9 +4221,13 @@ TSharedPtr<SWidget> USkinWeightsPaintTool::GetCustomSkeletonTreeWidget(TSharedPt
 	Delegates.OnExtendContextMenu.BindLambda([](FMenuBuilder& MenuBuilder)
 		{
 			FSkinWeightsPaintToolCommands Commands = FSkinWeightsPaintToolCommands::Get();
-			MenuBuilder.BeginSection("Isolate", LOCTEXT("IsolateSection", "Isolate"));
+
+			
+			
+			MenuBuilder.BeginSection("Common", LOCTEXT("CommonSection", "Common"));
 			{
 				MenuBuilder.AddMenuEntry(Commands.IsolateSelectedBones);
+				MenuBuilder.AddMenuEntry(Commands.ToggleSelectedBonesLockState);
 			}
 			MenuBuilder.EndSection();
 			
@@ -4273,11 +4314,11 @@ TSharedRef<SWidget> USkinWeightsPaintTool::OnGetLockWidget(FName BoneName)
 		{
 			if (LockedBoneIndices.Contains(BoneIndex))
 			{
-				ToggleLockBone(BoneIndex, false);
+				SetLockBone(BoneIndex, false);
 			}
 			else
 			{
-				ToggleLockBone(BoneIndex, true);
+				SetLockBone(BoneIndex, true);
 			}
 			return FReply::Handled();
 		})
@@ -4301,7 +4342,7 @@ TSharedRef<SWidget> USkinWeightsPaintTool::OnGetLockWidget(FName BoneName)
 	return LockButton.ToSharedRef();
 }
 
-void USkinWeightsPaintTool::ToggleLockAllBones(bool bLock)
+void USkinWeightsPaintTool::SetLockAllBones(bool bLock)
 {
 	ISkeletonProvider* SkeletonProvider = CastChecked<ISkeletonProvider>(Target);
 	const FReferenceSkeleton& RefSkeleton = SkeletonProvider->GetSkeleton();
@@ -4309,22 +4350,37 @@ void USkinWeightsPaintTool::ToggleLockAllBones(bool bLock)
 
 	for (BoneIndex BoneIndex = 0; BoneIndex < NumBones; ++BoneIndex)
 	{
-		ToggleLockBone(BoneIndex, bLock);
+		SetLockBone(BoneIndex, bLock);
 	}
 }
 
-void USkinWeightsPaintTool::ToggleLockSelectedBones(bool bLock)
+void USkinWeightsPaintTool::SetLockSelectedBones(bool bLock)
 {
 	for (int32 BoneIndex : SelectedBoneIndices)
 	{
-		ToggleLockBone(BoneIndex, bLock);
+		SetLockBone(BoneIndex, bLock);
 	}
 }
 
 void USkinWeightsPaintTool::IsolateSelectedBones()
 {
-	ToggleLockAllBones(true);
-	ToggleLockSelectedBones(false);
+	SetLockAllBones(true);
+	SetLockSelectedBones(false);
+}
+
+void USkinWeightsPaintTool::ToggleSelectedBonesLockState()
+{
+	for (int32 BoneIndex : SelectedBoneIndices)
+	{
+		if (GetRawLockedBoneIndices().Contains(BoneIndex))
+		{
+			SetLockBone(BoneIndex,false);
+		}
+		else
+		{
+			SetLockBone(BoneIndex, true);
+		}
+	}
 }
 
 void USkinWeightsPaintTool::ToggleBypassLocks()
@@ -4332,7 +4388,7 @@ void USkinWeightsPaintTool::ToggleBypassLocks()
 	bBypassLocks = !bBypassLocks;
 }
 
-void USkinWeightsPaintTool::ToggleLockBone(BoneIndex InBoneIndex, bool bLock)
+void USkinWeightsPaintTool::SetLockBone(BoneIndex InBoneIndex, bool bLock)
 {
 	if (bLock)
 	{
@@ -4356,7 +4412,7 @@ bool USkinWeightsPaintTool::ShouldProceedWithAllBonesUnlocked()
 
 	if (Response == EAppReturnType::Ok)
 	{
-		ToggleLockAllBones(false);	
+		SetLockAllBones(false);	
 	}
 	
 	return Response == EAppReturnType::Ok;
diff --git a/Engine/Plugins/Runtime/MeshModelingToolset/Source/MeshModelingToolsEditorOnly/Public/SkeletalMesh/SkinWeightsPaintTool.h b/Engine/Plugins/Runtime/MeshModelingToolset/Source/MeshModelingToolsEditorOnly/Public/SkeletalMesh/SkinWeightsPaintTool.h
--- a/Engine/Plugins/Runtime/MeshModelingToolset/Source/MeshModelingToolsEditorOnly/Public/SkeletalMesh/SkinWeightsPaintTool.h
+++ b/Engine/Plugins/Runtime/MeshModelingToolset/Source/MeshModelingToolsEditorOnly/Public/SkeletalMesh/SkinWeightsPaintTool.h
@@ -330,6 +330,7 @@ public:
 	virtual void RegisterCommands() override;
 	
 	TSharedPtr<FUICommandInfo> IsolateSelectedBones;
+	TSharedPtr<FUICommandInfo> ToggleSelectedBonesLockState;
 	
 	TSharedPtr<FUICommandInfo> LockSelectedBones;
 	TSharedPtr<FUICommandInfo> UnlockSelectedBones;
@@ -822,6 +823,8 @@ public:
 	FOnWeightsChanged OnWeightsChanged;
 
 protected:
+	// UInteractiveTool overrides
+	UE_API virtual void RegisterActions(FInteractiveToolActionSet& ActionSet) override;
 
 	UE_API virtual void ApplyStamp(const FBrushStampData& Stamp);
 	UE_API void OnShutdown(EToolShutdownType ShutdownType) override;
@@ -916,11 +919,12 @@ protected:
 
 	// skin weight locking
 	UE_API virtual TSharedRef<SWidget> OnGetLockWidget(FName BoneName);
-	UE_API void ToggleLockAllBones(bool bLock);
-	UE_API void ToggleLockSelectedBones(bool bLock);
+	UE_API void SetLockAllBones(bool bLock);
+	UE_API void SetLockSelectedBones(bool bLock);
 	UE_API void IsolateSelectedBones();
+	UE_API void ToggleSelectedBonesLockState();
 	UE_API void ToggleBypassLocks();
-	UE_API void ToggleLockBone(BoneIndex InBoneIndex, bool bLock);
+	UE_API void SetLockBone(BoneIndex InBoneIndex, bool bLock);
 
 	TMap<BoneIndex, TSharedPtr<SButton>> LockButtons;
 	TSet<BoneIndex> LockedBoneIndices;
