Match Move Brush Behavior to other DCCs

diff --git a/Engine/Plugins/Runtime/MeshModelingToolset/Source/MeshModelingTools/Private/MeshVertexSculptTool.cpp b/Engine/Plugins/Runtime/MeshModelingToolset/Source/MeshModelingTools/Private/MeshVertexSculptTool.cpp
--- a/Engine/Plugins/Runtime/MeshModelingToolset/Source/MeshModelingTools/Private/MeshVertexSculptTool.cpp
+++ b/Engine/Plugins/Runtime/MeshModelingToolset/Source/MeshModelingTools/Private/MeshVertexSculptTool.cpp
@@ -591,11 +591,13 @@ void UMeshVertexSculptTool::OnBeginStroke(const FRay& WorldRay)
 
 	UpdateBrushPosition(WorldRay);
 
+	bool bUpdatedROI = false;
 	TUniquePtr<FMeshSculptBrushOp>& UseBrushOp = GetActiveBrushOp();
 	FMeshSculptBrushOp::EReferencePlaneType ReferencePlaneType = UseBrushOp->GetReferencePlaneType();
 	if (ReferencePlaneType == FMeshSculptBrushOp::EReferencePlaneType::InitialROI ||
 		ReferencePlaneType == FMeshSculptBrushOp::EReferencePlaneType::InitialROI_ViewAligned)
 	{
+		bUpdatedROI = true;
 		UpdateROI(GetBrushFrameLocal());
 		UpdateStrokeReferencePlaneForROI(GetBrushFrameLocal(), TriangleROIArray,
 			ReferencePlaneType == FMeshSculptBrushOp::EReferencePlaneType::InitialROI_ViewAligned);
@@ -605,6 +607,15 @@ void UMeshVertexSculptTool::OnBeginStroke(const FRay& WorldRay)
 		UpdateStrokeReferencePlaneFromWorkPlane();
 	}
 
+	if (UseBrushOp->GetDragBehavior() == FMeshSculptBrushOp::EDragBehavior::UpdateSingleStamp)
+	{
+		if (!bUpdatedROI)
+		{
+			bUpdatedROI = true;
+			UpdateROI(GetBrushFrameLocal());
+		}
+	}
+
 	// initialize first "Last Stamp", so that we can assume all stamps in stroke have a valid previous stamp
 	LastStamp.WorldFrame = GetBrushFrameWorld();
 	LastStamp.LocalFrame = GetBrushFrameLocal();
@@ -625,6 +636,9 @@ void UMeshVertexSculptTool::OnBeginStroke(const FRay& WorldRay)
 		LastStamp.WorldFrame.Transform(CurTargetTransform);
 	}
 
+	LastStamp.PrevLocalFrame = LastStamp.LocalFrame;
+	LastStamp.PrevWorldFrame = LastStamp.WorldFrame;
+
 	InitialStrokeTriangleID = -1;
 	InitialStrokeTriangleID = GetBrushTriangleID();
 
@@ -1062,8 +1076,16 @@ bool UMeshVertexSculptTool::UpdateStampPosition(const FRay& WorldRay)
 		CurrentStamp.WorldFrame.Transform(CurTargetTransform);
 	}
 
-	CurrentStamp.PrevLocalFrame = LastStamp.LocalFrame;
-	CurrentStamp.PrevWorldFrame = LastStamp.WorldFrame;
+	if (UseBrushOp->GetDragBehavior() == FMeshSculptBrushOp::EDragBehavior::CreateNewStamps)
+	{
+		CurrentStamp.PrevLocalFrame = LastStamp.LocalFrame;
+		CurrentStamp.PrevWorldFrame = LastStamp.WorldFrame;
+	}
+	else if (UseBrushOp->GetDragBehavior() == FMeshSculptBrushOp::EDragBehavior::UpdateSingleStamp)
+	{
+		CurrentStamp.PrevLocalFrame = LastStamp.PrevLocalFrame;
+		CurrentStamp.PrevWorldFrame = LastStamp.PrevWorldFrame;
+	}
 
 	bMouseMoved = (PreviousRayDirection - WorldRay.Direction).SquaredLength() > FMathd::ZeroTolerance;
 	if (bMouseMoved)
@@ -1148,7 +1170,14 @@ TFuture<void> UMeshVertexSculptTool::ApplyStamp()
 	FDynamicMesh3* Mesh = GetSculptMesh();
 	{
 		TRACE_CPUPROFILER_EVENT_SCOPE(VtxSculptTool_ApplyStamp_Apply);
-		UseBrushOp->ApplyStamp(Mesh, *StampToUse, VertexROI, ROIPositionBuffer);
+
+		FDynamicMesh3* MeshForBrush = Mesh;
+		if (UseBrushOp->GetDragBehavior() == FMeshSculptBrushOp::EDragBehavior::UpdateSingleStamp)
+		{
+			MeshForBrush = GetBaseMesh();
+		}
+		
+		UseBrushOp->ApplyStamp(MeshForBrush, *StampToUse, VertexROI, ROIPositionBuffer);
 	}
 
 	// can discard alpha now
@@ -1481,9 +1510,12 @@ void UMeshVertexSculptTool::OnTick(float DeltaTime)
 	{
 		TRACE_CPUPROFILER_EVENT_SCOPE(VtxSculptTool_Tick_StrokeUpdate);
 		{
-			TRACE_CPUPROFILER_EVENT_SCOPE(VtxSculptTool_Tick_UpdateROI);
-			// update sculpt ROI
-			UpdateROI(CurrentStamp.LocalFrame);
+			if (GetActiveBrushOp()->GetDragBehavior() == FMeshSculptBrushOp::EDragBehavior::CreateNewStamps)
+			{
+				TRACE_CPUPROFILER_EVENT_SCOPE(VtxSculptTool_Tick_UpdateROI);
+				// update sculpt ROI
+				UpdateROI(CurrentStamp.LocalFrame);
+			}
 		}
 
 		// Instead of figuring out which triangles are unique here, we're going to simply accumulate the stamp ROIs in place
diff --git a/Engine/Plugins/Runtime/MeshModelingToolset/Source/MeshModelingTools/Public/Sculpting/MeshBrushOpBase.h b/Engine/Plugins/Runtime/MeshModelingToolset/Source/MeshModelingTools/Public/Sculpting/MeshBrushOpBase.h
--- a/Engine/Plugins/Runtime/MeshModelingToolset/Source/MeshModelingTools/Public/Sculpting/MeshBrushOpBase.h
+++ b/Engine/Plugins/Runtime/MeshModelingToolset/Source/MeshModelingTools/Public/Sculpting/MeshBrushOpBase.h
@@ -166,6 +166,12 @@ public:
 		WorkPlane,
 	};
 
+	enum class EDragBehavior
+	{
+		CreateNewStamps,
+		UpdateSingleStamp
+	};
+
 	virtual ~FMeshSculptBrushOp() {}
 
 	TWeakObjectPtr<UMeshSculptBrushOpProps> PropertySet;
@@ -289,7 +295,12 @@ public:
 		default:
 			return false;
 		}
-	} 
+	}
+
+	virtual EDragBehavior GetDragBehavior() const
+	{
+		return EDragBehavior::CreateNewStamps;
+	}
 };
 
 
diff --git a/Engine/Plugins/Runtime/MeshModelingToolset/Source/MeshModelingTools/Public/Sculpting/MeshMoveBrushOps.h b/Engine/Plugins/Runtime/MeshModelingToolset/Source/MeshModelingTools/Public/Sculpting/MeshMoveBrushOps.h
--- a/Engine/Plugins/Runtime/MeshModelingToolset/Source/MeshModelingTools/Public/Sculpting/MeshMoveBrushOps.h
+++ b/Engine/Plugins/Runtime/MeshModelingToolset/Source/MeshModelingTools/Public/Sculpting/MeshMoveBrushOps.h
@@ -45,6 +45,13 @@ class FMoveBrushOp : public FMeshSculptBrushOp
 {
 public:
 
+	FSculptBrushStamp FirstStamp;
+	
+	virtual void BeginStroke(const FDynamicMesh3* Mesh, const FSculptBrushStamp& Stamp, const TArray<int32>& InitialVertices) override
+	{
+		FirstStamp = Stamp;
+	};
+
 	virtual void ApplyStamp(const FDynamicMesh3* Mesh, const FSculptBrushStamp& Stamp, const TArray<int32>& Vertices, TArray<FVector3d>& NewPositionsOut) override
 	{
 		const FVector3d& StampPos = Stamp.LocalFrame.Origin;
@@ -71,7 +78,7 @@ public:
 			int32 VertIdx = Vertices[k];
 			FVector3d OrigPos = Mesh->GetVertex(VertIdx);
 
-			double Falloff = GetFalloff().Evaluate(Stamp, OrigPos);
+			double Falloff = GetFalloff().Evaluate(FirstStamp, OrigPos);
 
 			FVector3d NewPos = OrigPos + Falloff * MoveVec;
 			NewPositionsOut[k] = NewPos;
@@ -88,4 +95,9 @@ public:
 	{
 		return true;
 	}
+
+	virtual EDragBehavior GetDragBehavior() const override
+	{
+		return EDragBehavior::UpdateSingleStamp;
+	}
 };
