[Sculpt Brush] Added virtual funcs to brush ops that define how ROIs are updated and what input mesh to use when applying a stamp

diff --git a/Engine/Plugins/Runtime/MeshModelingToolset/Source/MeshModelingTools/Private/MeshVertexSculptTool.cpp b/Engine/Plugins/Runtime/MeshModelingToolset/Source/MeshModelingTools/Private/MeshVertexSculptTool.cpp
--- a/Engine/Plugins/Runtime/MeshModelingToolset/Source/MeshModelingTools/Private/MeshVertexSculptTool.cpp
+++ b/Engine/Plugins/Runtime/MeshModelingToolset/Source/MeshModelingTools/Private/MeshVertexSculptTool.cpp
@@ -607,7 +607,7 @@ void UMeshVertexSculptTool::OnBeginStroke(const FRay& WorldRay)
 		UpdateStrokeReferencePlaneFromWorkPlane();
 	}
 
-	if (UseBrushOp->GetDragBehavior() == FMeshSculptBrushOp::EDragBehavior::UpdateSingleStamp)
+	if (UseBrushOp->GetROIUpdateType() == FMeshSculptBrushOp::EROIUpdateType::UseInitial)
 	{
 		if (!bUpdatedROI)
 		{
@@ -1076,16 +1076,8 @@ bool UMeshVertexSculptTool::UpdateStampPosition(const FRay& WorldRay)
 		CurrentStamp.WorldFrame.Transform(CurTargetTransform);
 	}
 
-	if (UseBrushOp->GetDragBehavior() == FMeshSculptBrushOp::EDragBehavior::CreateNewStamps)
-	{
-		CurrentStamp.PrevLocalFrame = LastStamp.LocalFrame;
-		CurrentStamp.PrevWorldFrame = LastStamp.WorldFrame;
-	}
-	else if (UseBrushOp->GetDragBehavior() == FMeshSculptBrushOp::EDragBehavior::UpdateSingleStamp)
-	{
-		CurrentStamp.PrevLocalFrame = LastStamp.PrevLocalFrame;
-		CurrentStamp.PrevWorldFrame = LastStamp.PrevWorldFrame;
-	}
+	CurrentStamp.PrevLocalFrame = LastStamp.LocalFrame;
+	CurrentStamp.PrevWorldFrame = LastStamp.WorldFrame;
 
 	bMouseMoved = (PreviousRayDirection - WorldRay.Direction).SquaredLength() > FMathd::ZeroTolerance;
 	if (bMouseMoved)
@@ -1172,7 +1164,7 @@ TFuture<void> UMeshVertexSculptTool::ApplyStamp()
 		TRACE_CPUPROFILER_EVENT_SCOPE(VtxSculptTool_ApplyStamp_Apply);
 
 		FDynamicMesh3* MeshForBrush = Mesh;
-		if (UseBrushOp->GetDragBehavior() == FMeshSculptBrushOp::EDragBehavior::UpdateSingleStamp)
+		if (UseBrushOp->GetStampInputMeshType() == FMeshSculptBrushOp::EStampInputMeshType::BaseMesh)
 		{
 			MeshForBrush = GetBaseMesh();
 		}
@@ -1510,7 +1502,7 @@ void UMeshVertexSculptTool::OnTick(float DeltaTime)
 	{
 		TRACE_CPUPROFILER_EVENT_SCOPE(VtxSculptTool_Tick_StrokeUpdate);
 		{
-			if (GetActiveBrushOp()->GetDragBehavior() == FMeshSculptBrushOp::EDragBehavior::CreateNewStamps)
+			if (GetActiveBrushOp()->GetROIUpdateType() == FMeshSculptBrushOp::EROIUpdateType::UpdatePerStamp)
 			{
 				TRACE_CPUPROFILER_EVENT_SCOPE(VtxSculptTool_Tick_UpdateROI);
 				// update sculpt ROI
diff --git a/Engine/Plugins/Runtime/MeshModelingToolset/Source/MeshModelingTools/Public/Sculpting/MeshBrushOpBase.h b/Engine/Plugins/Runtime/MeshModelingToolset/Source/MeshModelingTools/Public/Sculpting/MeshBrushOpBase.h
--- a/Engine/Plugins/Runtime/MeshModelingToolset/Source/MeshModelingTools/Public/Sculpting/MeshBrushOpBase.h
+++ b/Engine/Plugins/Runtime/MeshModelingToolset/Source/MeshModelingTools/Public/Sculpting/MeshBrushOpBase.h
@@ -166,10 +166,22 @@ public:
 		WorkPlane,
 	};
 
-	enum class EDragBehavior
+	// Determines which vertices should be moved by the stroke
+	enum class EROIUpdateType
 	{
-		CreateNewStamps,
-		UpdateSingleStamp
+		// ROI is initialized from the first stamp and remains constant
+		UseInitial,
+		// ROI grows as new stamps overlap with more vertices
+		UpdatePerStamp
+	};
+
+	// Determines which mesh should be used by ApplyStamp
+	enum class EStampInputMeshType
+	{
+		// The "live" mesh that is updated by each stamp
+		SculptMesh,
+		// The "base" mesh that is typically updated after every stroke
+		BaseMesh
 	};
 
 	virtual ~FMeshSculptBrushOp() {}
@@ -297,9 +309,14 @@ public:
 		}
 	}
 
-	virtual EDragBehavior GetDragBehavior() const
+	virtual EROIUpdateType GetROIUpdateType() const
+	{
+		return EROIUpdateType::UpdatePerStamp;
+	}
+
+	virtual EStampInputMeshType GetStampInputMeshType() const
 	{
-		return EDragBehavior::CreateNewStamps;
+		return EStampInputMeshType::SculptMesh;
 	}
 };
 
diff --git a/Engine/Plugins/Runtime/MeshModelingToolset/Source/MeshModelingTools/Public/Sculpting/MeshMoveBrushOps.h b/Engine/Plugins/Runtime/MeshModelingToolset/Source/MeshModelingTools/Public/Sculpting/MeshMoveBrushOps.h
--- a/Engine/Plugins/Runtime/MeshModelingToolset/Source/MeshModelingTools/Public/Sculpting/MeshMoveBrushOps.h
+++ b/Engine/Plugins/Runtime/MeshModelingToolset/Source/MeshModelingTools/Public/Sculpting/MeshMoveBrushOps.h
@@ -43,10 +43,11 @@ public:
 
 class FMoveBrushOp : public FMeshSculptBrushOp
 {
-public:
-
+private:	
 	FSculptBrushStamp FirstStamp;
 	
+public:
+	
 	virtual void BeginStroke(const FDynamicMesh3* Mesh, const FSculptBrushStamp& Stamp, const TArray<int32>& InitialVertices) override
 	{
 		FirstStamp = Stamp;
@@ -54,10 +55,7 @@ public:
 
 	virtual void ApplyStamp(const FDynamicMesh3* Mesh, const FSculptBrushStamp& Stamp, const TArray<int32>& Vertices, TArray<FVector3d>& NewPositionsOut) override
 	{
-		const FVector3d& StampPos = Stamp.LocalFrame.Origin;
-
-		double UsePower = Stamp.Power;
-		FVector3d MoveVec = Stamp.LocalFrame.Origin - Stamp.PrevLocalFrame.Origin;
+		FVector3d MoveVec = Stamp.LocalFrame.Origin - FirstStamp.LocalFrame.Origin;
 		
 		UMoveBrushOpProps* Props = GetPropertySetAs<UMoveBrushOpProps>();
 
@@ -65,12 +63,12 @@ public:
 		// that world vector to local space
 		if (Props->AxisFilters.AnyAxisFiltered())
 		{
-			FVector3d WorldMoveVec = Stamp.WorldFrame.Origin - Stamp.PrevWorldFrame.Origin;
+			FVector3d WorldMoveVec = Stamp.WorldFrame.Origin - FirstStamp.WorldFrame.Origin;
 			WorldMoveVec = MoveVec.Length() * UE::Geometry::Normalized(WorldMoveVec);		// apply local space scaling
 			WorldMoveVec.X = (Props->AxisFilters.bAxisX) ? WorldMoveVec.X : 0;
 			WorldMoveVec.Y = (Props->AxisFilters.bAxisY) ? WorldMoveVec.Y : 0;
 			WorldMoveVec.Z = (Props->AxisFilters.bAxisZ) ? WorldMoveVec.Z : 0;
-			MoveVec = Stamp.LocalFrame.FromFrameVector(Stamp.WorldFrame.ToFrameVector(WorldMoveVec));
+			MoveVec = FirstStamp.LocalFrame.FromFrameVector(FirstStamp.WorldFrame.ToFrameVector(WorldMoveVec));
 		}
 
 		ParallelFor(Vertices.Num(), [&](int32 k)
@@ -96,8 +94,13 @@ public:
 		return true;
 	}
 
-	virtual EDragBehavior GetDragBehavior() const override
+	virtual EROIUpdateType GetROIUpdateType() const override
+	{
+		return EROIUpdateType::UseInitial;
+	}
+
+	virtual EStampInputMeshType GetStampInputMeshType() const override
 	{
-		return EDragBehavior::UpdateSingleStamp;
+		return EStampInputMeshType::BaseMesh;
 	}
 };
