add optional regularization to the geometrycore qem simplifier

diff --git a/Engine/Plugins/Runtime/GeometryProcessing/Source/DynamicMesh/Private/MeshSimplification.cpp b/Engine/Plugins/Runtime/GeometryProcessing/Source/DynamicMesh/Private/MeshSimplification.cpp
--- a/Engine/Plugins/Runtime/GeometryProcessing/Source/DynamicMesh/Private/MeshSimplification.cpp
+++ b/Engine/Plugins/Runtime/GeometryProcessing/Source/DynamicMesh/Private/MeshSimplification.cpp
@@ -66,14 +66,18 @@ void TMeshSimplification<QuadricErrorType>::InitializeTriQuadrics()
 	const int NT = Mesh->MaxTriangleID();
 	triQuadrics.SetNum(NT);
 	triAreas.SetNum(NT);
+	if (RegularizeUsesNormals())
+	{
+		triNormals.SetNum(NT);
+	}
 
 	// tested with ParallelFor - no measurable benefit
 	//@todo parallel version
 	//gParallel.BlockStartEnd(0, Mesh->MaxTriangleID - 1, (start_tid, end_tid) = > {
-	FVector3d n, c;
+	FVector3d n(0), c;
 	for (int tid : Mesh->TriangleIndicesItr())
 	{
-		triQuadrics[tid] = ComputeFaceQuadric(tid, n, c, triAreas[tid]);
+		triQuadrics[tid] = ComputeFaceQuadric(tid, triNormals.IsEmpty() ? n : triNormals[tid], c, triAreas[tid]);
 	}
 
 }
@@ -157,12 +161,52 @@ void TMeshSimplification<QuadricErrorType>::InitializeVertexQuadrics()
 	//gParallel.BlockStartEnd(0, Mesh->MaxVertexID - 1, (start_vid, end_vid) = > {
 	for (int vid : Mesh->VertexIndicesItr())
 	{
-		vertQuadrics[vid] = FQuadricErrorType::Zero();
+		FQuadricErrorType Q = FQuadricErrorType::Zero();
+		
+		
+		double AreaSum = 0;
+		FVector3d Normal(0);
+		bool bNeedsNormals = RegularizeUsesNormals();
 		for (int tid : Mesh->VtxTrianglesItr(vid))
 		{
-			vertQuadrics[vid].Add(triAreas[tid], triQuadrics[tid]);
+			AreaSum += triAreas[tid];
+			if (bNeedsNormals)
+			{
+				Normal += triAreas[tid] * triNormals[tid];
+			}
+			Q.Add(triAreas[tid], triQuadrics[tid]);
+		}
+
+		ApplyRegularizationToVertexQuadric(Q, vid, AreaSum, bNeedsNormals, Normal);
+
+		vertQuadrics[vid] = Q;
+	}
+}
+
+template <typename QuadricErrorType>
+void TMeshSimplification<QuadricErrorType>::ApplyRegularizationToVertexQuadric(FQuadricErrorType& Q, int32 VID, double AreaSum, bool bNeedsNormals, FVector3d Normal)
+{
+	if (RegularizeWeight > 0.)
+	{
+		double UseRegularizeWeight = RegularizeWeight;
+		if (CustomRegularizationScaleF)
+		{
+			UseRegularizeWeight = RegularizeWeight * CustomRegularizationScaleF(*Mesh, VID);
+		}
+		FVector3d VertPos = Mesh->GetVertex(VID);
+		if (bNeedsNormals && Normal.SquaredLength() > 0)
+		{
+			Normal.Normalize();
+			// Note: Very large weights line-regularizer weights introduce instability in the solve
+			// so we clamp the map weight in this case
+			constexpr double MaxLineRegularizerWeight = 100.;
+			UseRegularizeWeight = FMath::Min(MaxLineRegularizerWeight, UseRegularizeWeight);
+			Q.AddLineRegularizer(VertPos, Normal, UseRegularizeWeight * AreaSum);
+		}
+		else
+		{
+			Q.AddPositionRegularizer(VertPos, UseRegularizeWeight * AreaSum);
 		}
-		//check(TMathUtil.EpsilonEqual(0, vertQuadrics[i].Evaluate(Mesh->GetVertex(i)), TMathUtil.Epsilon * 10));
 	}
 }
 
@@ -493,6 +537,10 @@ void TMeshSimplification<QuadricErrorType>::UpdateNeighborhood(const FDynamicMes
 			// update the arrays that hold the current face area & quadric
 			triAreas[tid] = NewtriArea;
 			triQuadrics[tid] = NewtriQuadric;
+			if (!triNormals.IsEmpty())
+			{
+				triNormals[tid] = n;
+			}
 
 			FIndex3i tri_vids = Mesh->GetTriangle(tid);
 
@@ -529,10 +577,19 @@ void TMeshSimplification<QuadricErrorType>::UpdateNeighborhood(const FDynamicMes
 		// NB: in the version with memory this quadric took the value of the edge quadric that collapsed.
 		{
 			FQuadricErrorType vertQuadric = FQuadricErrorType::Zero();
+			double AreaSum = 0;
+			FVector3d Normal(0);
+			bool bNeedsNormals = RegularizeUsesNormals();
 			for (int tid : Mesh->VtxTrianglesItr(kvid))
 			{
+				AreaSum += triAreas[tid];
+				if (bNeedsNormals)
+				{
+					Normal += triAreas[tid] * triNormals[tid];
+				}
 				vertQuadric.Add(triAreas[tid], triQuadrics[tid]);
 			}
+			ApplyRegularizationToVertexQuadric(vertQuadric, kvid, AreaSum, bNeedsNormals, Normal);
 			vertQuadrics[kvid] = vertQuadric;
 		}
 	}
diff --git a/Engine/Plugins/Runtime/GeometryProcessing/Source/DynamicMesh/Public/MeshSimplification.h b/Engine/Plugins/Runtime/GeometryProcessing/Source/DynamicMesh/Public/MeshSimplification.h
--- a/Engine/Plugins/Runtime/GeometryProcessing/Source/DynamicMesh/Public/MeshSimplification.h
+++ b/Engine/Plugins/Runtime/GeometryProcessing/Source/DynamicMesh/Public/MeshSimplification.h
@@ -31,6 +31,16 @@ enum class ESimplificationResult
 	Ignored_CreatesTinyTriangle = 10
 };
 
+enum class ERegularizeSimplificationMethod
+{
+	// Regularize with a distance-from-original-vertex-position term
+	// Simplest / fastest method, stable even for very large weights
+	Position,
+	// Regularize with a distance-from-line-through-vertex-along-normal term.
+	// Tends to give slightly better results for small weights; can be less stable for very large weights
+	NormalLine
+};
+
 /**
  * Implementation of Garland & Heckbert Quadric Error Metric (QEM) Triangle Mesh Simplification
  */
@@ -82,6 +92,15 @@ public:
 	 */
 	bool bPreventTinyTriangles = false;
 
+	/** 
+	 * If greater than 0, we will add a regularization term with this weight.
+	 * This can help give a better triangulation for example in flat regions where the error metric is under-determined
+	 */
+	double RegularizeWeight = 0.;
+
+	/** Method to use to regularize, if RegularizeWeight > 0 */
+	ERegularizeSimplificationMethod RegularizeMethod = ERegularizeSimplificationMethod::NormalLine;
+
 	/** When using the constraint system, these options will apply to the appropriate boundaries. */
 	EEdgeRefineFlags MeshBoundaryConstraint, GroupBoundaryConstraint, MaterialBoundaryConstraint;
 
@@ -105,6 +124,7 @@ public:
 	TFunction<double(const FDynamicMesh3&, int VertexA, int VertexB)> CustomGeometricErrorScaleF;
 	TFunction<double(const FDynamicMesh3&, int VertexA, int VertexB)> CustomQuadricErrorScaleF;
 	TFunction<double(const FDynamicMesh3&, int VertexA, int VertexB)> CustomEdgeLengthScaleF;
+	TFunction<double(const FDynamicMesh3&, int Vertex)> CustomRegularizationScaleF;
 
 
 	TMeshSimplification(FDynamicMesh3* m) : FMeshRefinerBase(m)
@@ -225,9 +245,15 @@ protected:
 	DYNAMICMESH_API virtual void InitializeSeamQuadrics();
 
 	TArray<double> triAreas;
+	TArray<FVector3d> triNormals; // this array is only populate if needed; i.e. if RegularizeUsesNormals() returns true
 	TArray<FQuadricErrorType> triQuadrics;
 	DYNAMICMESH_API virtual void InitializeTriQuadrics();
 
+	bool RegularizeUsesNormals() const
+	{
+		return RegularizeWeight > 0 && RegularizeMethod == ERegularizeSimplificationMethod::NormalLine;
+	}
+
 	FDynamicMeshNormalOverlay* NormalOverlay;
 
 	FQuadricErrorType ComputeFaceQuadric(const int tid, FVector3d& nface, FVector3d& c, double& Area) const;
@@ -433,6 +459,9 @@ protected:
 		//if (ENABLE_PROFILING) CollapseW.Stop();
 	}
 
+private:
+
+	void ApplyRegularizationToVertexQuadric(FQuadricErrorType& Q, int32 VID, double AreaSum, bool bNeedsNormals, FVector3d Normal);
 
 };
 
diff --git a/Engine/Plugins/Runtime/MeshModelingToolset/Source/MeshModelingToolsEditorOnly/Private/SimplifyMeshTool.cpp b/Engine/Plugins/Runtime/MeshModelingToolset/Source/MeshModelingToolsEditorOnly/Private/SimplifyMeshTool.cpp
--- a/Engine/Plugins/Runtime/MeshModelingToolset/Source/MeshModelingToolsEditorOnly/Private/SimplifyMeshTool.cpp
+++ b/Engine/Plugins/Runtime/MeshModelingToolset/Source/MeshModelingToolsEditorOnly/Private/SimplifyMeshTool.cpp
@@ -243,6 +243,7 @@ TUniquePtr<FDynamicMeshOperator> USimplifyMeshTool::MakeNewOperator()
 	Op->MaterialBoundaryConstraint = (EEdgeRefineFlags)SimplifyProperties->MaterialBoundaryConstraint;
 	Op->bGeometricDeviationConstraint = SimplifyProperties->bGeometricConstraint;
 	Op->GeometricTolerance = SimplifyProperties->GeometricTolerance;
+	Op->RegularizeWeight = SimplifyProperties->RegularizeWeight;
 	Op->PolyEdgeAngleTolerance = SimplifyProperties->PolyEdgeAngleTolerance;
 	Op->BoundaryEdgeAngleTolerance = SimplifyProperties->BoundaryEdgeAngleTolerance;
 	FTransform LocalToWorld = (FTransform)UE::ToolTarget::GetLocalToWorldTransform(Target);
@@ -285,6 +286,23 @@ TUniquePtr<FDynamicMeshOperator> USimplifyMeshTool::MakeNewOperator()
 				return 1.0;
 			};
 
+			TFunction<double(const FDynamicMesh3& Mesh, int Vertex)> VertexMetricScale =
+			[WeightMapAttributeLayerIndex, RelativeDensity = WeightMapProperties->RelativeDensity](const FDynamicMesh3& Mesh, int Vertex) -> double
+			{
+				if (Mesh.HasAttributes() && Mesh.Attributes()->GetWeightLayer(WeightMapAttributeLayerIndex))
+				{
+					const FDynamicMeshWeightAttribute* const AttributeLayer = Mesh.Attributes()->GetWeightLayer(WeightMapAttributeLayerIndex);
+
+					float WeightValue;
+					AttributeLayer->GetValue(Vertex, &WeightValue);
+					WeightValue = FMath::Clamp(WeightValue, 0.0f, 1.0f);
+
+					return FMath::Pow(0.5, -RelativeDensity * (double)WeightValue);
+				}
+
+				return 1.0;
+			};
+
 			// geometric tolerance can be used in multiple modes, so we pass it through always
 			Op->CustomGeometricErrorScaleF = EdgeMetricScale;
 			// only scale quadric error or edge lengths when they are the active target, so weights scale the most relevant metric
@@ -296,6 +314,10 @@ TUniquePtr<FDynamicMeshOperator> USimplifyMeshTool::MakeNewOperator()
 			{
 				Op->CustomEdgeLengthScaleF = EdgeMetricScale;
 			}
+			if (Op->RegularizeWeight > 0 && SimplifyProperties->bScaleRegularizationWithWeightMap)
+			{
+				Op->CustomRegularizationScaleF = VertexMetricScale;
+			}
 
 			TFunction<double(int Vertex)> ClusterEdgeTargetLengthScale =
 			[AttributeLayer = OriginalMesh->Attributes()->GetWeightLayer(WeightMapAttributeLayerIndex),
diff --git a/Engine/Plugins/Runtime/MeshModelingToolset/Source/MeshModelingToolsEditorOnly/Public/SimplifyMeshTool.h b/Engine/Plugins/Runtime/MeshModelingToolset/Source/MeshModelingToolsEditorOnly/Public/SimplifyMeshTool.h
--- a/Engine/Plugins/Runtime/MeshModelingToolset/Source/MeshModelingToolsEditorOnly/Public/SimplifyMeshTool.h
+++ b/Engine/Plugins/Runtime/MeshModelingToolset/Source/MeshModelingToolsEditorOnly/Public/SimplifyMeshTool.h
@@ -101,6 +101,16 @@ public:
 				EditCondition = "bGeometricConstraint && SimplifierType != ESimplifyType::UEStandard && SimplifierType != ESimplifyType::MinimalPolygroup"))
 	float GeometricTolerance;
 
+	/** Optional regularization term to favor vertices remaining near their original locations. Using a small regularization weight can give cleaner topology in flat regions. */
+	UPROPERTY(EditAnywhere, Category = Options, meta = (UIMin = "0.0", UIMax = "0.01", ClampMin = "0.0", ClampMax = "100.0",
+		EditCondition = "SimplifierType != ESimplifyType::UEStandard && SimplifierType != ESimplifyType::MinimalPolygroup && SimplifierType != ESimplifyType::ClusterBased"))
+	float RegularizeWeight = 0.0001;
+
+	/** Optionally scale the regularization with per-vertex weights, if a weight map is specified. Can be used to target simplification via the weights. */
+	UPROPERTY(EditAnywhere, Category = Options, meta = (
+		EditCondition = "RegularizeWeight > 0 && SimplifierType != ESimplifyType::UEStandard && SimplifierType != ESimplifyType::MinimalPolygroup && SimplifierType != ESimplifyType::ClusterBased"))
+	bool bScaleRegularizationWithWeightMap = true;
+
 	/** Display colors corresponding to the mesh's polygon groups */
 	UPROPERTY(EditAnywhere, Category = Display)
 	bool bShowGroupColors = false;
diff --git a/Engine/Plugins/Runtime/MeshModelingToolset/Source/ModelingOperatorsEditorOnly/Private/CleaningOps/SimplifyMeshOp.cpp b/Engine/Plugins/Runtime/MeshModelingToolset/Source/ModelingOperatorsEditorOnly/Private/CleaningOps/SimplifyMeshOp.cpp
--- a/Engine/Plugins/Runtime/MeshModelingToolset/Source/ModelingOperatorsEditorOnly/Private/CleaningOps/SimplifyMeshOp.cpp
+++ b/Engine/Plugins/Runtime/MeshModelingToolset/Source/ModelingOperatorsEditorOnly/Private/CleaningOps/SimplifyMeshOp.cpp
@@ -41,9 +41,11 @@ void ComputeSimplify(FDynamicMesh3* TargetMesh, const bool bReproject,
 	                 typename SimplificationType::ESimplificationCollapseModes CollapseMode,
 					 bool bUseQuadricMemory,
 					 float GeometricTolerance,
+					 float RegularizeWeight,
 	TFunction<double(const FDynamicMesh3&, int VertexA, int VertexB)> CustomGeometricErrorScaleF,
 	TFunction<double(const FDynamicMesh3&, int VertexA, int VertexB)> CustomQuadricErrorScaleF,
-	TFunction<double(const FDynamicMesh3&, int VertexA, int VertexB)> CustomEdgeLengthScaleF)
+	TFunction<double(const FDynamicMesh3&, int VertexA, int VertexB)> CustomEdgeLengthScaleF,
+	TFunction<double(const FDynamicMesh3&, int Vertex)> CustomRegularizationScaleF)
 {
 	SimplificationType Reducer(TargetMesh);
 
@@ -59,6 +61,9 @@ void ComputeSimplify(FDynamicMesh3* TargetMesh, const bool bReproject,
 	Reducer.CustomGeometricErrorScaleF = CustomGeometricErrorScaleF;
 	Reducer.CustomQuadricErrorScaleF = CustomQuadricErrorScaleF;
 	Reducer.CustomEdgeLengthScaleF = CustomEdgeLengthScaleF;
+	Reducer.CustomRegularizationScaleF = CustomRegularizationScaleF;
+
+	Reducer.RegularizeWeight = RegularizeWeight;
 
 	if (bAllowSeamCollapse)
 	{
@@ -251,9 +256,11 @@ void FSimplifyMeshOp::CalculateResult(FProgressCancel* Progress)
 											TargetMode, TargetPercentage, TargetCount, TargetEdgeLength, TargetQuadricError, MinimalPlanarAngleThresh,
 											FQEMSimplification::ESimplificationCollapseModes::MinimalQuadricPositionError, bUseQuadricMemory, 
 											UseGeometricTolerance,
+											RegularizeWeight,
 											CustomGeometricErrorScaleF,
 											CustomQuadricErrorScaleF,
-											CustomEdgeLengthScaleF);
+											CustomEdgeLengthScaleF,
+											CustomRegularizationScaleF);
 	}
 	else if (SimplifierType == ESimplifyType::Attribute)
 	{
@@ -271,9 +278,11 @@ void FSimplifyMeshOp::CalculateResult(FProgressCancel* Progress)
 													TargetMode, TargetPercentage, TargetCount, TargetEdgeLength, TargetQuadricError, MinimalPlanarAngleThresh,
 													FAttrMeshSimplification::ESimplificationCollapseModes::MinimalQuadricPositionError, bUseQuadricMemory,
 													UseGeometricTolerance,
+													RegularizeWeight,
 													CustomGeometricErrorScaleF,
 													CustomQuadricErrorScaleF,
-													CustomEdgeLengthScaleF);
+													CustomEdgeLengthScaleF,
+													CustomRegularizationScaleF);
 	}
 	else if (SimplifierType == ESimplifyType::MinimalPlanar)
 	{
@@ -291,9 +300,11 @@ void FSimplifyMeshOp::CalculateResult(FProgressCancel* Progress)
 			ESimplifyTargetType::MinimalPlanar, TargetPercentage, TargetCount, TargetEdgeLength, TargetQuadricError, MinimalPlanarAngleThresh,
 			FQEMSimplification::ESimplificationCollapseModes::MinimalQuadricPositionError, bUseQuadricMemory,
 			UseGeometricTolerance,
+			RegularizeWeight,
 			CustomGeometricErrorScaleF,
 			CustomQuadricErrorScaleF,
-			CustomEdgeLengthScaleF);
+			CustomEdgeLengthScaleF,
+			CustomRegularizationScaleF);
 	}
 	else if (SimplifierType == ESimplifyType::MinimalExistingVertex)
 	{
@@ -307,9 +318,11 @@ void FSimplifyMeshOp::CalculateResult(FProgressCancel* Progress)
 			TargetMode, TargetPercentage, TargetCount, TargetEdgeLength, TargetQuadricError, MinimalPlanarAngleThresh,
 			FQEMSimplification::ESimplificationCollapseModes::MinimalExistingVertexError, bUseQuadricMemory,
 			UseGeometricTolerance,
+			RegularizeWeight,
 			CustomGeometricErrorScaleF,
 			CustomQuadricErrorScaleF,
-			CustomEdgeLengthScaleF);
+			CustomEdgeLengthScaleF,
+			CustomRegularizationScaleF);
 	}
 	else if (SimplifierType == ESimplifyType::MinimalPolygroup)
 	{
diff --git a/Engine/Plugins/Runtime/MeshModelingToolset/Source/ModelingOperatorsEditorOnly/Public/CleaningOps/SimplifyMeshOp.h b/Engine/Plugins/Runtime/MeshModelingToolset/Source/ModelingOperatorsEditorOnly/Public/CleaningOps/SimplifyMeshOp.h
--- a/Engine/Plugins/Runtime/MeshModelingToolset/Source/ModelingOperatorsEditorOnly/Public/CleaningOps/SimplifyMeshOp.h
+++ b/Engine/Plugins/Runtime/MeshModelingToolset/Source/ModelingOperatorsEditorOnly/Public/CleaningOps/SimplifyMeshOp.h
@@ -100,10 +100,14 @@ public:
 	bool bGeometricDeviationConstraint = false;
 	float GeometricTolerance = 0.0f;
 
+	// Weight of optional regularization term, for QEM-based simplifiers
+	float RegularizeWeight = 0.f;
+
 	// QEM simplifier custom scales to apply to *measured* errors and edge lengths. Note this is in effect the inverse of scaling the tolerances / target lengths.
 	TFunction<double(const FDynamicMesh3&, int VertexA, int VertexB)> CustomGeometricErrorScaleF;
 	TFunction<double(const FDynamicMesh3&, int VertexA, int VertexB)> CustomQuadricErrorScaleF;
 	TFunction<double(const FDynamicMesh3&, int VertexA, int VertexB)> CustomEdgeLengthScaleF;
+	TFunction<double(const FDynamicMesh3&, int Vertex)> CustomRegularizationScaleF;
 
 	// Cluster simplifier custom scales apply to the *target* edge lengths, and so is conceptually the inverse of the above scales.
 	TFunction<double(int32)> ClusterTargetEdgeLengthScaleF = nullptr;
diff --git a/Engine/Source/Runtime/GeometryCore/Public/QuadricError.h b/Engine/Source/Runtime/GeometryCore/Public/QuadricError.h
--- a/Engine/Source/Runtime/GeometryCore/Public/QuadricError.h
+++ b/Engine/Source/Runtime/GeometryCore/Public/QuadricError.h
@@ -242,6 +242,27 @@ struct TQuadricError
 		return SolveAxEqualsb(OutResult, -bx, -by, -bz, minThresh);
 	}
 
+	// Add regularizing term that penalizes movement away from Pos
+	void AddPositionRegularizer(const TVector<RealType>& Pos, RealType Weight)
+	{
+		Axx += Weight;
+		Ayy += Weight;
+		Azz += Weight;
+		bx -= Pos.X * Weight;
+		by -= Pos.Y * Weight;
+		bz -= Pos.Z * Weight;
+		c += Pos.Dot(Pos) * Weight;
+	}
+
+	// Add regularizing term that penalizes movement away from the line through Pos along LineDir
+	void AddLineRegularizer(const TVector<RealType>& Pos, const TVector<RealType>& LineDir, RealType Weight)
+	{
+		TVector<RealType> N1, N2;
+		VectorUtil::MakePerpVectors(LineDir, N1, N2);
+		Add(Weight * .5, TQuadricError<RealType>(N1, Pos));
+		Add(Weight * .5, TQuadricError<RealType>(N2, Pos));
+	}
+
 };
 
 
@@ -402,8 +423,18 @@ public:
 		}
 
 		return bValid;
+	}
 
+	// Add regularizing term that penalizes movement away from Pos
+	void AddPositionRegularizer(const UE::Math::TVector<RealType>& Pos, RealType Weight)
+	{
+		BaseStruct::AddPositionRegularizer(Pos, Weight);
+	}
 
+	// Add regularizing term that penalizes movement away from the line through Pos along LineDir
+	void AddLineRegularizer(const UE::Math::TVector<RealType>& Pos, const UE::Math::TVector<RealType>& LineDir, RealType Weight)
+	{
+		BaseStruct::AddLineRegularizer(Pos, LineDir, Weight);
 	}
 };
 
@@ -696,6 +727,18 @@ public:
 		return Evaluate(point, attr);
 	}
 
+	// Add regularizing term that penalizes movement away from Pos
+	void AddPositionRegularizer(const UE::Math::TVector<RealType>& Pos, RealType Weight)
+	{
+		BaseStruct::AddPositionRegularizer(Pos, Weight);
+	}
+
+	// Add regularizing term that penalizes movement away from the line through Pos along LineDir
+	void AddLineRegularizer(const UE::Math::TVector<RealType>& Pos, const UE::Math::TVector<RealType>& LineDir, RealType Weight)
+	{
+		BaseStruct::AddLineRegularizer(Pos, LineDir, Weight);
+	}
+
 public:
 
 	// accumulated area
