guard against divide by zero on empty mesh case for ue standard simplifier paths

diff --git a/Engine/Plugins/Runtime/GeometryScripting/Source/GeometryScriptingCore/Private/MeshSimplifyFunctions.cpp b/Engine/Plugins/Runtime/GeometryScripting/Source/GeometryScriptingCore/Private/MeshSimplifyFunctions.cpp
--- a/Engine/Plugins/Runtime/GeometryScripting/Source/GeometryScriptingCore/Private/MeshSimplifyFunctions.cpp
+++ b/Engine/Plugins/Runtime/GeometryScripting/Source/GeometryScriptingCore/Private/MeshSimplifyFunctions.cpp
@@ -136,7 +136,8 @@ namespace UE::PrivateSimplifyHelper
 
 		FDynamicMeshToMeshDescription Converter;
 		Converter.Convert(&Mesh, SrcMeshDescription, true);
-		float Percent = (float)TargetCount / (float)(bTargetIsTriCount ? Mesh.TriangleCount() : Mesh.VertexCount());
+		float Total = (float)(bTargetIsTriCount ? Mesh.TriangleCount() : Mesh.VertexCount());
+		float Percent = Total > 0 ? (float)TargetCount / Total : 0.f;
 		return FSimplifyMeshOp::ComputeStandardSimplifier(MeshReduction, SrcMeshDescription, Mesh, Percent, bTargetIsTriCount, false, nullptr);
 #else
 		return false;
diff --git a/Engine/Plugins/Runtime/MeshModelingToolset/Source/ModelingOperatorsEditorOnly/Private/CleaningOps/SimplifyMeshOp.cpp b/Engine/Plugins/Runtime/MeshModelingToolset/Source/ModelingOperatorsEditorOnly/Private/CleaningOps/SimplifyMeshOp.cpp
--- a/Engine/Plugins/Runtime/MeshModelingToolset/Source/ModelingOperatorsEditorOnly/Private/CleaningOps/SimplifyMeshOp.cpp
+++ b/Engine/Plugins/Runtime/MeshModelingToolset/Source/ModelingOperatorsEditorOnly/Private/CleaningOps/SimplifyMeshOp.cpp
@@ -400,12 +400,12 @@ void FSimplifyMeshOp::CalculateResult(FProgressCancel* Progress)
 		else if (TargetMode == ESimplifyTargetType::TriangleCount)
 		{
 			int32 NumTris = SrcMeshDescription->Polygons().Num();
-			PercentReduction = (float)TargetCount / (float)NumTris;
+			PercentReduction = NumTris > 0 ? (float)TargetCount / (float)NumTris : 0.f;
 		}
 		else if (TargetMode == ESimplifyTargetType::VertexCount)
 		{
 			int32 NumVerts = SrcMeshDescription->Vertices().Num();
-			PercentReduction = (float)TargetCount / (float)NumVerts;
+			PercentReduction = NumVerts > 0 ? (float)TargetCount / (float)NumVerts : 0.f;
 		}
 		bool bTargetIsTriangleCount = TargetMode != ESimplifyTargetType::VertexCount;
 		
