Fix for PerformSceneRender getting stuck in an infinite loop.

Some of the compiling managers; like FNiagaraSystemCompilingManager, need  to have their Async Tasks Processed to finish completely.

diff --git a/Engine/Plugins/Runtime/MeshModelingToolset/Source/ModelingComponents/Private/Scene/WorldRenderCapture.cpp b/Engine/Plugins/Runtime/MeshModelingToolset/Source/ModelingComponents/Private/Scene/WorldRenderCapture.cpp
--- a/Engine/Plugins/Runtime/MeshModelingToolset/Source/ModelingComponents/Private/Scene/WorldRenderCapture.cpp
+++ b/Engine/Plugins/Runtime/MeshModelingToolset/Source/ModelingComponents/Private/Scene/WorldRenderCapture.cpp
@@ -388,6 +388,9 @@ static void PerformSceneRender(FCanvas* Canvas, FSceneViewFamily* ViewFamily)
 		if (bCompiledAssets)
 		{
 			FAssetCompilingManager::Get().FinishAllCompilation();
+
+			// Ensure all deferred construction scripts are executed
+			FAssetCompilingManager::Get().ProcessAsyncTasks();
 		}
 	} while (bCompiledAssets);
 }
