Spline: guard against empty mesh when projecting the spline

diff --git a/Engine/Plugins/Runtime/MeshModelingToolset/Source/ModelingComponents/Private/SplineUtil.cpp b/Engine/Plugins/Runtime/MeshModelingToolset/Source/ModelingComponents/Private/SplineUtil.cpp
--- a/Engine/Plugins/Runtime/MeshModelingToolset/Source/ModelingComponents/Private/SplineUtil.cpp
+++ b/Engine/Plugins/Runtime/MeshModelingToolset/Source/ModelingComponents/Private/SplineUtil.cpp
@@ -221,12 +221,15 @@ void UE::Geometry::SplineUtil::ProjectSplineToSurface(FInterpCurveVector& Output
 		ProjectedPoint.OutVal = SplineNearestPoint;
 		ProjectedPoint.InterpMode = InterpMode;
 
-		// Spline point tangent = spline derivative projected onto the tangent plane of the surface
-		const FVector3d SplineDerivative = Spline.EvalDerivative(SplinePointT, FVector3d(0));
-		const FVector3d TriangleNormal = SurfaceAABBTree.GetMesh()->GetTriNormal(NearestTriangleID);
-		const FVector3d Tangent = SplineDerivative - SplineDerivative.Dot(TriangleNormal) * TriangleNormal;
-		ProjectedPoint.ArriveTangent = 0.5 * Tangent;
-		ProjectedPoint.LeaveTangent = 0.5 * Tangent;
+		if (ensure(SurfaceAABBTree.GetMesh()->IsTriangle(NearestTriangleID)))
+		{
+			// Spline point tangent = spline derivative projected onto the tangent plane of the surface
+			const FVector3d SplineDerivative = Spline.EvalDerivative(SplinePointT, FVector3d(0));
+			const FVector3d TriangleNormal = SurfaceAABBTree.GetMesh()->GetTriNormal(NearestTriangleID);
+			const FVector3d Tangent = SplineDerivative - SplineDerivative.Dot(TriangleNormal) * TriangleNormal;
+			ProjectedPoint.ArriveTangent = 0.5 * Tangent;
+			ProjectedPoint.LeaveTangent = 0.5 * Tangent;
+		}
 
 		return ProjectedPoint;
 	};
