Dataflow : allow much more smaller minimum brush relative size in the weight map tool

diff --git a/Engine/Plugins/Experimental/Dataflow/Source/DataflowEditor/Private/DataflowEditorTools/DataflowEditorVertexAttributePaintTool.cpp b/Engine/Plugins/Experimental/Dataflow/Source/DataflowEditor/Private/DataflowEditorTools/DataflowEditorVertexAttributePaintTool.cpp
--- a/Engine/Plugins/Experimental/Dataflow/Source/DataflowEditor/Private/DataflowEditorTools/DataflowEditorVertexAttributePaintTool.cpp
+++ b/Engine/Plugins/Experimental/Dataflow/Source/DataflowEditor/Private/DataflowEditorTools/DataflowEditorVertexAttributePaintTool.cpp
@@ -613,6 +613,8 @@ void UDataflowEditorVertexAttributePaintTool::SetupBrushEditBehaviorSetup(ULocal
 
 void UDataflowEditorVertexAttributePaintTool::Setup()
 {
+	MinimumBrushAdaptiveSizeRatio = 0.001;
+
 	UMeshSculptToolBase::Setup();
 
 	// todo(ccaillaud) : this should be a customized by the inherited class 
diff --git a/Engine/Plugins/Runtime/MeshModelingToolset/Source/MeshModelingTools/Private/Sculpting/MeshSculptToolBase.cpp b/Engine/Plugins/Runtime/MeshModelingToolset/Source/MeshModelingTools/Private/Sculpting/MeshSculptToolBase.cpp
--- a/Engine/Plugins/Runtime/MeshModelingToolset/Source/MeshModelingTools/Private/Sculpting/MeshSculptToolBase.cpp
+++ b/Engine/Plugins/Runtime/MeshModelingToolset/Source/MeshModelingTools/Private/Sculpting/MeshSculptToolBase.cpp
@@ -1300,7 +1300,7 @@ void UMeshSculptToolBase::UpdateStrokeReferencePlaneFromWorkPlane()
 void UMeshSculptToolBase::InitializeBrushSizeRange(const FAxisAlignedBox3d& TargetBounds)
 {
 	double MaxDimension = TargetBounds.MaxDim();
-	BrushRelativeSizeRange = FInterval1d(MaxDimension * 0.01, MaxDimension);
+	BrushRelativeSizeRange = FInterval1d(MaxDimension * MinimumBrushAdaptiveSizeRatio, MaxDimension);
 	BrushProperties->BrushSize.InitializeWorldSizeRange(
 		TInterval<float>((float)BrushRelativeSizeRange.Min, (float)BrushRelativeSizeRange.Max));
 	CalculateBrushRadius();
diff --git a/Engine/Plugins/Runtime/MeshModelingToolset/Source/MeshModelingTools/Public/Sculpting/MeshSculptToolBase.h b/Engine/Plugins/Runtime/MeshModelingToolset/Source/MeshModelingTools/Public/Sculpting/MeshSculptToolBase.h
--- a/Engine/Plugins/Runtime/MeshModelingToolset/Source/MeshModelingTools/Public/Sculpting/MeshSculptToolBase.h
+++ b/Engine/Plugins/Runtime/MeshModelingToolset/Source/MeshModelingTools/Public/Sculpting/MeshSculptToolBase.h
@@ -468,6 +468,7 @@ protected:
 	FFrame3d LastBrushFrameWorld;
 	FFrame3d LastBrushFrameLocal;
 	int32 LastBrushTriangleID = IndexConstants::InvalidID;
+	double MinimumBrushAdaptiveSizeRatio = 0.01;
 
 	const FFrame3d& GetBrushFrameWorld() const { return LastBrushFrameWorld; }
 	const FFrame3d& GetBrushFrameLocal() const { return LastBrushFrameLocal; }
