Skin Weight Paint Tool: fix regression from 46296928

#rnx
#Synced-CL 48455833

diff --git a/Engine/Plugins/Runtime/MeshModelingToolset/Source/MeshModelingToolsEditorOnly/Private/SkeletalMesh/SkinWeightsPaintTool.cpp b/Engine/Plugins/Runtime/MeshModelingToolset/Source/MeshModelingToolsEditorOnly/Private/SkeletalMesh/SkinWeightsPaintTool.cpp
--- a/Engine/Plugins/Runtime/MeshModelingToolset/Source/MeshModelingToolsEditorOnly/Private/SkeletalMesh/SkinWeightsPaintTool.cpp
+++ b/Engine/Plugins/Runtime/MeshModelingToolset/Source/MeshModelingToolsEditorOnly/Private/SkeletalMesh/SkinWeightsPaintTool.cpp
@@ -636,7 +636,7 @@ void UWeightToolTransferManager::TransferWeightsFromOtherMeshOrSubset()
 			TmpSourceCopy = PartialSubMesh.GetSubmesh();
 			// By default, submesh only have matching attributes enabled, so we need to explicitly copy the ones we need
 			FDynamicMeshAttributeSet* SourceAttributes = TmpSourceCopy.Attributes();
-			SourceAttributes->CopyBoneAttributes(*EditedMesh.Attributes());
+			SourceAttributes->CopyBoneAttributes(*SourceMeshOrig.Attributes());
 		}
 		else
 		{
