MeshSceneAdapter - Validate that the static mesh render data is valid before falling back to render meshes

#Synced-CL 48784656

diff --git a/Engine/Plugins/Runtime/MeshModelingToolset/Source/ModelingComponents/Private/Scene/MeshSceneAdapter.cpp b/Engine/Plugins/Runtime/MeshModelingToolset/Source/ModelingComponents/Private/Scene/MeshSceneAdapter.cpp
--- a/Engine/Plugins/Runtime/MeshModelingToolset/Source/ModelingComponents/Private/Scene/MeshSceneAdapter.cpp
+++ b/Engine/Plugins/Runtime/MeshModelingToolset/Source/ModelingComponents/Private/Scene/MeshSceneAdapter.cpp
@@ -1386,6 +1386,7 @@ static TUniquePtr<IMeshSpatialWrapper> SpatialWrapperFactory( const FMeshTypeCon
 		{
 #if WITH_EDITOR
 			bool bUseRenderMeshes = BuildOptions.bIgnoreStaticMeshSourceData || StaticMesh->GetOutermost()->bIsCookedForEditor;
+			bUseRenderMeshes &= StaticMesh->GetRenderData() && StaticMesh->GetRenderData()->GetCurrentFirstLOD(0);
 #else
 			bool bUseRenderMeshes = true;
 #endif
