[ModelingTools] Fixed geo selection not changing selectable rendering when changing element types

#Synced-CL 48737829

diff --git a/Engine/Plugins/Runtime/MeshModelingToolset/Source/ModelingComponents/Private/Selection/GeometrySelectionManager.cpp b/Engine/Plugins/Runtime/MeshModelingToolset/Source/ModelingComponents/Private/Selection/GeometrySelectionManager.cpp
--- a/Engine/Plugins/Runtime/MeshModelingToolset/Source/ModelingComponents/Private/Selection/GeometrySelectionManager.cpp
+++ b/Engine/Plugins/Runtime/MeshModelingToolset/Source/ModelingComponents/Private/Selection/GeometrySelectionManager.cpp
@@ -211,7 +211,8 @@ void UGeometrySelectionManager::SetSelectionElementTypeInternal(EGeometryElement
 			Target->SelectionEditor->UpdateQueryConfig(GetCurrentSelectionQueryConfig(), bEnableTopologyFilter);
 		}
 
-		MarkRenderCachesDirty(false);
+		constexpr bool bMarkSelectableDirty = true;
+		MarkRenderCachesDirty(bMarkSelectableDirty);
 		ClearActivePreview();
 	}
 }
