MeshModelingToolset: Fix AssetUtils::ReadTexture() for TSF_R32F to only provide values in first channel, and add support for TSF_16F

#Synced-CL 49387652

diff --git a/Engine/Plugins/Runtime/MeshModelingToolset/Source/ModelingComponents/Private/AssetUtils/Texture2DUtil.cpp b/Engine/Plugins/Runtime/MeshModelingToolset/Source/ModelingComponents/Private/AssetUtils/Texture2DUtil.cpp
--- a/Engine/Plugins/Runtime/MeshModelingToolset/Source/ModelingComponents/Private/AssetUtils/Texture2DUtil.cpp
+++ b/Engine/Plugins/Runtime/MeshModelingToolset/Source/ModelingComponents/Private/AssetUtils/Texture2DUtil.cpp
@@ -218,6 +218,17 @@ static bool ReadTexture_SourceData(
 		}
 		return true;
 	}
+	else if (SourceFormat == TSF_R16F)
+	{
+		check(BytesPerPixel == sizeof(FFloat16));
+		for (int64 i = 0; i < Num; ++i)
+		{
+			const uint8* PixelPtr = SourceDataPtr + (i * BytesPerPixel);
+			const float Value = reinterpret_cast<const FFloat16*>(PixelPtr)->GetFloat();
+			DestImage.SetPixel(i, FVector4f(Value, 0.0f, 0.0f, 1.f));
+		}
+		return true;
+	}
 	else if (SourceFormat == TSF_R32F)
 	{
 		check(BytesPerPixel == sizeof(float));
@@ -225,11 +236,11 @@ static bool ReadTexture_SourceData(
 		{
 			const uint8* PixelPtr = SourceDataPtr + (i * BytesPerPixel);
 			const float Value = *reinterpret_cast<const float*>(PixelPtr);
-			DestImage.SetPixel(i, FVector4f(Value, Value, Value, 1.f));
+			DestImage.SetPixel(i, FVector4f(Value, 0.0f, 0.0f, 1.f));
 		}
 		return true;
 	}
-	else if ((SourceFormat == TSF_G16))
+	else if (SourceFormat == TSF_G16)
 	{
 		check(BytesPerPixel == 2);
 		for (int64 i = 0; i < Num; ++i)
