MeshModelingToolset: Add benchmarks to AssetUtils ReadTexture tests, and send any output to the Message Log window instead of using UE_LOG

#rnx
#Synced-CL 49433390

diff --git a/Engine/Plugins/Runtime/MeshModelingToolset/Source/ModelingComponents/Private/Tests/AssetUtilsTests.cpp b/Engine/Plugins/Runtime/MeshModelingToolset/Source/ModelingComponents/Private/Tests/AssetUtilsTests.cpp
--- a/Engine/Plugins/Runtime/MeshModelingToolset/Source/ModelingComponents/Private/Tests/AssetUtilsTests.cpp
+++ b/Engine/Plugins/Runtime/MeshModelingToolset/Source/ModelingComponents/Private/Tests/AssetUtilsTests.cpp
@@ -6,6 +6,7 @@
 #include "Algo/ForEach.h"
 #include "AssetUtils/Texture2DUtil.h"
 #include "Engine/Texture2D.h"
+#include "Logging/MessageLog.h"
 #include "Misc/AutomationTest.h"
 
 #if WITH_DEV_AUTOMATION_TESTS
@@ -142,6 +143,10 @@ bool FAssetUtilsReadTextureTest::RunTest(const FString& Parameters)
 	constexpr int64 NumPixels = Width * Height;
 	constexpr uint32 Seed = 0;
 
+	// Set this to true to time 100 iterations for each test, and print the results summary in the message log.
+	constexpr bool bRunBenchmark = false;
+	constexpr uint32 BenchmarkNumIterations = 100;
+
 	// The current implementation for platform data converts the texture to TC_VectorDisplacementmap, which seems to introduce all kinds of artifacts that make
 	// a comparison to the original source data intractable.
 	constexpr bool bPreferPlatformData = false;
@@ -157,6 +162,8 @@ bool FAssetUtilsReadTextureTest::RunTest(const FString& Parameters)
 		TSF_RGBE8_DEPRECATED
 	};
 
+	FMessageLog MessageLog("AutomationTestingLog");
+
 	for (const auto SourceFormat : SourceFormats)
 	{
 		if (Algo::Find(UnsupportedFormats, SourceFormat))
@@ -242,8 +249,37 @@ bool FAssetUtilsReadTextureTest::RunTest(const FString& Parameters)
 
 					if (!TestTrue(What("Check Pixels"), NumErrors == 0))
 					{
-						UE_LOG(LogAutomationTestFramework, Warning, TEXT("%s - Num Errors: %lld (%.2f%%) - Max Error: %.10f"),
-						       *What("Check Pixels"), NumErrors, static_cast<float>(NumErrors) / NumPixels * 100.0f, MaxError);
+						MessageLog.Warning(FText::FromString(FString::Printf(
+							TEXT("%s - Num Errors: %lld (%.2f%%) - Max Error: %.10f"),
+							*What("Check Pixels"), NumErrors, static_cast<float>(NumErrors) / NumPixels * 100.0f, MaxError)));
+					}
+				}
+
+				if constexpr (bRunBenchmark)
+				{
+					TArray<double> Timings;
+					Timings.Reserve(BenchmarkNumIterations);
+					for (int32 Iteration = 0; Iteration < BenchmarkNumIterations; ++Iteration)
+					{
+						const uint64 StartCycles = FPlatformTime::Cycles64();
+						const bool bSuccess = UE::AssetUtils::ReadTexture(Texture.Get(), Output, bPreferPlatformData);
+						const uint64 EndCycles = FPlatformTime::Cycles64();
+
+						if (bSuccess)
+						{
+							Timings.Add(FPlatformTime::ToMilliseconds64(EndCycles - StartCycles) * 1000.0);
+						}
+					}
+					if (Timings.Num())
+					{
+						Timings.Sort();
+						MessageLog.Info(FText::FromString(FString::Printf(
+							TEXT("%s Min: %.1f us, Median: %.1f us, Max: %.1f us"),
+							*What("Benchmark"), Timings[0], Timings[Timings.Num() / 2], Timings[Timings.Num() - 1])));
+					}
+					else
+					{
+						MessageLog.Error(FText::FromString(FString::Printf(TEXT("%s failed!"), *What("Benchmark"))));
 					}
 				}
 			}
