Fix Bounds computation for Point components
- the size of the points was added to the local bounds before scaling it by the LocalToWorld transform  causing the bounds to grow very large is the scale of the transform was also large

#Synced-CL 49563940

diff --git a/Engine/Plugins/Runtime/MeshModelingToolset/Source/ModelingComponents/Private/Drawing/PointSetComponent.cpp b/Engine/Plugins/Runtime/MeshModelingToolset/Source/ModelingComponents/Private/Drawing/PointSetComponent.cpp
--- a/Engine/Plugins/Runtime/MeshModelingToolset/Source/ModelingComponents/Private/Drawing/PointSetComponent.cpp
+++ b/Engine/Plugins/Runtime/MeshModelingToolset/Source/ModelingComponents/Private/Drawing/PointSetComponent.cpp
@@ -409,21 +409,23 @@ FBoxSphereBounds UPointSetComponent::CalcBounds(const FTransform& LocalToWorld)
 		}
 		Bounds = FBoxSphereBounds(Box);
 		bBoundsDirty = false;
-
-		// TODO: This next bit is not ideal because the point size is specified in onscreen pixels,
-		// so the true amount by which we would need to expand bounds depends on camera location, FOV, etc.
-		// We mainly do this as a hack against a problem in ortho viewports, which cull small items
-		// based on their bounds, and a set consisting of a single point will always be culled due
-		// to having 0-sized bounds. It's worth noting that when zooming out sufficiently far, the
-		// point will still be culled even with this hack, however.
-		// The proper solution is to be able to opt out of the ortho culling behavior, which is something
-		// we need to add.
-		if (Points.Num() > 0)
-		{
-			Bounds = Bounds.ExpandBy(Points[0].Size);
-		}
 	}
-	return Bounds.TransformBy(LocalToWorld);
+	FBoxSphereBounds TransformBounds = Bounds.TransformBy(LocalToWorld);
+
+	// TODO: This next bit is not ideal because the point size is specified in onscreen pixels,
+	// so the true amount by which we would need to expand bounds depends on camera location, FOV, etc.
+	// We mainly do this as a hack against a problem in ortho viewports, which cull small items
+	// based on their bounds, and a set consisting of a single point will always be culled due
+	// to having 0-sized bounds. It's worth noting that when zooming out sufficiently far, the
+	// point will still be culled even with this hack, however.
+	// The proper solution is to be able to opt out of the ortho culling behavior, which is something
+	// we need to add.
+	if (Points.Num() > 0)
+	{
+		TransformBounds = TransformBounds.ExpandBy(Points[0].Size);
+	}
+
+	return TransformBounds;
 }
 
 
