fix crash due to RemeshOp discarding a weight map that it needs, when bDiscardAttributes is set but the target edge length is weightmap-based

#Synced-CL 49560868

diff --git a/Engine/Plugins/ChaosClothAssetEditor/Source/ChaosClothAssetDataflowNodes/Private/ChaosClothAsset/RemeshNode.cpp b/Engine/Plugins/ChaosClothAssetEditor/Source/ChaosClothAssetDataflowNodes/Private/ChaosClothAsset/RemeshNode.cpp
--- a/Engine/Plugins/ChaosClothAssetEditor/Source/ChaosClothAssetDataflowNodes/Private/ChaosClothAsset/RemeshNode.cpp
+++ b/Engine/Plugins/ChaosClothAssetEditor/Source/ChaosClothAssetDataflowNodes/Private/ChaosClothAsset/RemeshNode.cpp
@@ -868,6 +868,7 @@ namespace UE::Chaos::ClothAsset::Private
 				{
 					DensityMapLayerIndex = WeightLayerIndex;
 					bFoundDensityMapLayer = true;
+					RemeshOp.bNeedsWeightMaps = true;
 					break;
 				}
 			}
diff --git a/Engine/Plugins/Runtime/GeometryScripting/Source/GeometryScriptingCore/Private/MeshRemeshFunctions.cpp b/Engine/Plugins/Runtime/GeometryScripting/Source/GeometryScriptingCore/Private/MeshRemeshFunctions.cpp
--- a/Engine/Plugins/Runtime/GeometryScripting/Source/GeometryScriptingCore/Private/MeshRemeshFunctions.cpp
+++ b/Engine/Plugins/Runtime/GeometryScripting/Source/GeometryScriptingCore/Private/MeshRemeshFunctions.cpp
@@ -193,6 +193,7 @@ UDynamicMesh* UGeometryScriptLibrary_RemeshingFunctions::ApplyAdaptiveRemesh(
 		if (SourceMesh->HasAttributes())
 		{
 			const int32 WeightMapAttributeLayerIndex = WeightMapHandle.WeightMapAttributeLayerIndex;
+			RemeshOp.bNeedsWeightMaps = true;
 
 			RemeshOp.CustomEdgeLengthScaleF = [WeightMapAttributeLayerIndex, RelativeDensity = AdaptiveOptions.RelativeDensity](const FDynamicMesh3& Mesh, int VertexA, int VertexB) -> double
 			{
diff --git a/Engine/Plugins/Runtime/MeshModelingToolset/Source/MeshModelingTools/Private/RemeshMeshTool.cpp b/Engine/Plugins/Runtime/MeshModelingToolset/Source/MeshModelingTools/Private/RemeshMeshTool.cpp
--- a/Engine/Plugins/Runtime/MeshModelingToolset/Source/MeshModelingTools/Private/RemeshMeshTool.cpp
+++ b/Engine/Plugins/Runtime/MeshModelingToolset/Source/MeshModelingTools/Private/RemeshMeshTool.cpp
@@ -270,6 +270,7 @@ TUniquePtr<FDynamicMeshOperator> URemeshMeshTool::MakeNewOperator()
 
 		if (WeightMapAttributeLayerIndex >= 0)
 		{
+			Op->bNeedsWeightMaps = true;
 			Op->CustomEdgeLengthScaleF = [WeightMapAttributeLayerIndex, RelativeDensity = WeightMapProperties->RelativeDensity](const FDynamicMesh3& Mesh, int VertexA, int VertexB) -> double
 			{
 				if (Mesh.HasAttributes() && Mesh.Attributes()->GetWeightLayer(WeightMapAttributeLayerIndex))
diff --git a/Engine/Plugins/Runtime/MeshModelingToolset/Source/ModelingOperators/Private/CleaningOps/RemeshMeshOp.cpp b/Engine/Plugins/Runtime/MeshModelingToolset/Source/ModelingOperators/Private/CleaningOps/RemeshMeshOp.cpp
--- a/Engine/Plugins/Runtime/MeshModelingToolset/Source/ModelingOperators/Private/CleaningOps/RemeshMeshOp.cpp
+++ b/Engine/Plugins/Runtime/MeshModelingToolset/Source/ModelingOperators/Private/CleaningOps/RemeshMeshOp.cpp
@@ -72,9 +72,43 @@ void FRemeshMeshOp::CalculateResult(FProgressCancel* Progress)
 	{
 		return;
 	}
+	
 
-	bool bDiscardAttributesImmediately = bDiscardAttributes && !bPreserveSharpEdges;
-	ResultMesh->Copy(*OriginalMesh, true, true, true, !bDiscardAttributesImmediately);
+
+	bool bHasDiscardedAttributes = bDiscardAttributes && !bPreserveSharpEdges && !bNeedsWeightMaps;
+	ResultMesh->Copy(*OriginalMesh, true, true, true, !bHasDiscardedAttributes); // if we'll never need attributes, copying w/out them is fastest
+
+	// helper to discard only attributes that we don't need
+	auto SafeDiscardAttributes = [this, &bHasDiscardedAttributes](FDynamicMesh3& Mesh)
+	{
+		if (bHasDiscardedAttributes || !Mesh.HasAttributes())
+		{
+			return;
+		}
+		bHasDiscardedAttributes = true;
+		if (bNeedsWeightMaps)
+		{
+			// manually turn off attributes aside from the weights
+			Mesh.Attributes()->SetNumNormalLayers(0);
+			Mesh.Attributes()->SetNumPolygroupLayers(0);
+			Mesh.Attributes()->SetNumUVLayers(0);
+			Mesh.Attributes()->DisableBones();
+			Mesh.Attributes()->DiscardSculptLayers();
+			Mesh.Attributes()->DisablePrimaryColors();
+			Mesh.Attributes()->DisableMaterialID();
+		}
+		else
+		{
+			// discard the entire attribute set
+			Mesh.DiscardAttributes();
+		}
+	};
+
+	// if we copied w/ attributes but we just needed the weight maps, discard all else
+	if (bDiscardAttributes && bNeedsWeightMaps && !bPreserveSharpEdges)
+	{
+		SafeDiscardAttributes(*ResultMesh);
+	}
 
 	if (Progress && Progress->Cancelled())
 	{
@@ -186,9 +220,9 @@ void FRemeshMeshOp::CalculateResult(FProgressCancel* Progress)
 
 	Remesher->Progress = Progress;
 
-	if (bDiscardAttributes && !bDiscardAttributesImmediately)
+	if (bDiscardAttributes && !bHasDiscardedAttributes)
 	{
-		TargetMesh->DiscardAttributes();
+		SafeDiscardAttributes(*TargetMesh);
 	}
 
 	{
@@ -235,7 +269,15 @@ void FRemeshMeshOp::CalculateResult(FProgressCancel* Progress)
 		}
 	}
 
-	if (!TargetMesh->HasAttributes() && bResultMustHaveAttributesEnabled)
+	if (bHasDiscardedAttributes && TargetMesh->HasAttributes())
+	{
+		// return the default single UV and normal layers, if we discarded to 0 above,
+		// since we typically expect PrimaryUV and PrimaryNormals to be valid
+		TargetMesh->Attributes()->SetNumUVLayers(1);
+		TargetMesh->Attributes()->SetNumNormalLayers(1);
+	}
+
+	if ((bHasDiscardedAttributes || !TargetMesh->HasAttributes()) && bResultMustHaveAttributesEnabled)
 	{
 		TargetMesh->EnableAttributes();
 		if (TargetMesh->HasVertexUVs())
diff --git a/Engine/Plugins/Runtime/MeshModelingToolset/Source/ModelingOperators/Public/CleaningOps/RemeshMeshOp.h b/Engine/Plugins/Runtime/MeshModelingToolset/Source/ModelingOperators/Public/CleaningOps/RemeshMeshOp.h
--- a/Engine/Plugins/Runtime/MeshModelingToolset/Source/ModelingOperators/Public/CleaningOps/RemeshMeshOp.h
+++ b/Engine/Plugins/Runtime/MeshModelingToolset/Source/ModelingOperators/Public/CleaningOps/RemeshMeshOp.h
@@ -69,6 +69,8 @@ public:
 	double TargetEdgeLength = 1.0f;
 	ERemeshSmoothingType SmoothingType;
 	bool bDiscardAttributes = false;
+	// Note: if bDiscardAttributes and bNeedsWeightMaps are both true, the weight map attributes will be kept through remeshing (but other attributes will be discarded)
+	bool bNeedsWeightMaps = false;
 	bool bPreserveSharpEdges = true;
 	bool bFlips = true;
 	bool bSplits = true;
