[SkeletalMeshModelingTools] Fix an issue where the wrong bone was renamed when a name was committed by changing the selection and losing focus

#Synced-CL 49723230

diff --git a/Engine/Plugins/Runtime/MeshModelingToolset/Source/MeshModelingToolsEditorOnly/Private/SkeletalMesh/SReferenceSkeletonTree.cpp b/Engine/Plugins/Runtime/MeshModelingToolset/Source/MeshModelingToolsEditorOnly/Private/SkeletalMesh/SReferenceSkeletonTree.cpp
--- a/Engine/Plugins/Runtime/MeshModelingToolset/Source/MeshModelingToolsEditorOnly/Private/SkeletalMesh/SReferenceSkeletonTree.cpp
+++ b/Engine/Plugins/Runtime/MeshModelingToolset/Source/MeshModelingToolsEditorOnly/Private/SkeletalMesh/SReferenceSkeletonTree.cpp
@@ -74,6 +74,7 @@ void SBoneItem::Construct(const FArguments& InArgs)
 	
 	WeakTreeElement = InArgs._TreeElement;
 	WeakModifier = InArgs._WeakModifier;
+	OnBoneElementNameCommitted = InArgs._Delegates.OnBoneElementNameCommitted;
 
 	const TSharedPtr<FBoneElement> TreeElement = WeakTreeElement.Pin();
 
@@ -108,7 +109,7 @@ void SBoneItem::Construct(const FArguments& InArgs)
 		.ColorAndOpacity_Static(&SBoneItem::GetTextColor, InArgs._IsSelected)
 		.MaximumLength(NAME_SIZE-1)
 		.OnVerifyTextChanged(this, &SBoneItem::OnVerifyNameChanged)
-		.OnTextCommitted(InArgs._Delegates.OnBoneNameCommitted)
+		.OnTextCommitted(this, &SBoneItem::OnTextCommitted)
 		.MultiLine(false)
     	.OnEnterEditingMode(this, &SBoneItem::OnEnterEditingMode)
 		.OnExitEditingMode(this, &SBoneItem::OnExitEditingMode)
@@ -176,6 +177,14 @@ bool SBoneItem::OnVerifyNameChanged(const FText& InText, FText& OutErrorMessage)
 	return BoneIndex == INDEX_NONE;
 }
 
+void SBoneItem::OnTextCommitted(const FText& Text, ETextCommit::Type TextCommitType)
+{
+	if (WeakTreeElement.IsValid())
+	{
+		OnBoneElementNameCommitted.ExecuteIfBound(Text, TextCommitType, WeakTreeElement.Pin().ToSharedRef());
+	}
+}
+
 void SBoneItem::OnEnterEditingMode()
 {
 	bEditing = true;
@@ -313,9 +322,12 @@ void SReferenceSkeletonTree::Construct(const FArguments& InArgs)
 	Delegates.OnCanAcceptDrop = !Delegates.OnCanAcceptDrop.IsBound() ? FOnRefSkeletonTreeCanAcceptDrop::CreateSP(this, &SReferenceSkeletonTree::OnCanAcceptDrop) : Delegates.OnCanAcceptDrop;
 	Delegates.OnAcceptDrop = !Delegates.OnAcceptDrop.IsBound() ? FOnRefSkeletonTreeAcceptDrop::CreateSP(this, &SReferenceSkeletonTree::OnAcceptDrop) : Delegates.OnAcceptDrop;
 	Delegates.OnDragDetected = !Delegates.OnDragDetected.IsBound() ? FOnDragDetected::CreateSP(this, &SReferenceSkeletonTree::OnDragDetected) : Delegates.OnDragDetected;
-	Delegates.OnBoneRenamed = !Delegates.OnBoneRenamed.IsBound() ? FOnBoneRenamed::CreateSP(this, &SReferenceSkeletonTree::OnBoneRenamed) : Delegates.OnBoneRenamed;
-	Delegates.OnBoneNameCommitted =!Delegates.OnBoneNameCommitted.IsBound() ? FOnTextCommitted::CreateSP(this, &SReferenceSkeletonTree::OnNewBoneNameCommitted) : Delegates.OnBoneNameCommitted;
-	
+	Delegates.OnBoneRenamed = !Delegates.OnBoneRenamed.IsBound() ? FOnBoneRenamed::CreateSP(this, &SReferenceSkeletonTree::OnBoneRenamed) : Delegates.OnBoneRenamed;	
+	if (!Delegates.OnBoneElementNameCommitted.IsBound())
+	{
+		Delegates.OnBoneElementNameCommitted = FRefSkeletonTreeDelegates::FOnBoneElementNameCommitted::CreateSP(this, &SReferenceSkeletonTree::OnBoneElementNameCommitted);
+	}
+
 	CommandList = MakeShared<FUICommandList>();
 	BindCommands();
 
@@ -800,7 +812,7 @@ void SReferenceSkeletonTree::OnBoneRenamed(const FName InOldName, const FName In
 	}
 }
 
-void SReferenceSkeletonTree::OnNewBoneNameCommitted(const FText& InText, ETextCommit::Type InCommitType)
+void SReferenceSkeletonTree::OnBoneElementNameCommitted(const FText& InText, ETextCommit::Type InCommitType, TSharedRef<FBoneElement> BoneElement)
 {
 	if (InCommitType == ETextCommit::OnCleared)
 	{
@@ -810,6 +822,13 @@ void SReferenceSkeletonTree::OnNewBoneNameCommitted(const FText& InText, ETextCo
 	TArray<FName> BoneNames;
 	GetSelectedBoneNames(BoneNames);
 
+	// If the selection does not contain the bone name, the user committed by changing the selection.
+	// Handle the case by reanming the element which invoked this commit and not the selection.
+	if (!BoneNames.Contains(BoneElement->BoneName))
+	{
+		BoneNames = { BoneElement->BoneName };
+	}
+
 	static const FString NamePattern(TEXT("$name"));
 	static constexpr TCHAR Slash(TEXT('/'));
 	
diff --git a/Engine/Plugins/Runtime/MeshModelingToolset/Source/MeshModelingToolsEditorOnly/Public/SkeletalMesh/SReferenceSkeletonTree.h b/Engine/Plugins/Runtime/MeshModelingToolset/Source/MeshModelingToolsEditorOnly/Public/SkeletalMesh/SReferenceSkeletonTree.h
--- a/Engine/Plugins/Runtime/MeshModelingToolset/Source/MeshModelingToolsEditorOnly/Public/SkeletalMesh/SReferenceSkeletonTree.h
+++ b/Engine/Plugins/Runtime/MeshModelingToolset/Source/MeshModelingToolsEditorOnly/Public/SkeletalMesh/SReferenceSkeletonTree.h
@@ -85,13 +85,26 @@ using FOnRefSkeletonTreeAcceptDrop = STableRow<TSharedPtr<FBoneElement>>::FOnAcc
 using FOnBoneRenamed = FBoneElement::FOnRenamed;
 
 struct FRefSkeletonTreeDelegates
-{
+{	
+	// Avoid deprecation warnings when the struct is constructed or copied with default methods
+	PRAGMA_DISABLE_DEPRECATION_WARNINGS
+	FRefSkeletonTreeDelegates() = default;
+	FRefSkeletonTreeDelegates(const FRefSkeletonTreeDelegates&) = default;
+	FRefSkeletonTreeDelegates(FRefSkeletonTreeDelegates&&) = default;
+	FRefSkeletonTreeDelegates& operator=(const FRefSkeletonTreeDelegates&) = default;
+	FRefSkeletonTreeDelegates& operator=(FRefSkeletonTreeDelegates&&) = default;
+	PRAGMA_ENABLE_DEPRECATION_WARNINGS
+
 	FOnDragDetected					OnDragDetected;
 	FOnRefSkeletonTreeCanAcceptDrop	OnCanAcceptDrop;
 	FOnRefSkeletonTreeAcceptDrop	OnAcceptDrop;
-	FOnBoneRenamed					OnBoneRenamed;
+	FOnBoneRenamed					OnBoneRenamed;	
+	UE_DEPRECATED(5.8, TEXT("Instead use OnBoneElementNameCommitted which is explicit about the renamed element"))
 	FOnTextCommitted				OnBoneNameCommitted;
 
+	DECLARE_DELEGATE_ThreeParams(FOnBoneElementNameCommitted, const FText& /** ElementName */, ETextCommit::Type /** TextCommitType */, TSharedRef<FBoneElement> /** BoneElement */)
+	FOnBoneElementNameCommitted	OnBoneElementNameCommitted;
+
 	// Tools can use these to extend the skeleton tree with extra commands or widgets
 	DECLARE_DELEGATE_RetVal_OneParam(TSharedRef<SWidget>, FOnGetRowWidget, FName)
 	FOnGetRowWidget	OnGetRowWidget;
@@ -123,6 +136,7 @@ class SBoneItem : public SCompoundWidget
 private:
 
 	bool OnVerifyNameChanged(const FText& InText, FText& OutErrorMessage) const;
+	void OnTextCommitted(const FText& Text, ETextCommit::Type TextCommitType);
 
 	void OnEnterEditingMode();
 	void OnExitEditingMode();
@@ -137,6 +151,9 @@ private:
 	
 	bool bEditing = false;
 	TOptional<ReferenceSkeletonTreeLocals::EBoneRenameType> OptRenameType;
+
+	/** Delegate executed when the bone name was committed */
+	FRefSkeletonTreeDelegates::FOnBoneElementNameCommitted OnBoneElementNameCommitted;
 };
 
 class FBoneItemDragDropOp : public FDecoratedDragDropOp
@@ -270,7 +287,7 @@ private:
 	FReply OnAcceptDrop(const FDragDropEvent& DragDropEvent, EItemDropZone DropZone, TSharedPtr<FBoneElement> TargetItem);
 	void OnBoneRenamed(const FName InOldName, const FName InNewName) const;
 	
-	void OnNewBoneNameCommitted(const FText& InText, ETextCommit::Type InCommitType);
+	void OnBoneElementNameCommitted(const FText& InText, ETextCommit::Type InCommitType, TSharedRef<FBoneElement> BoneElement);
 
 	// The modifier that actually does modifications on the reference skeleton
 	TWeakObjectPtr<USkeletonModifier> Modifier;
