[SkeletalMeshModellingTools] Fix build UnrealEditor Mac, SReferenceSkeletonTree.h, 'deprecated' attribute requires a string

#rnx

#Synced-CL 49723230

diff --git a/Engine/Plugins/Runtime/MeshModelingToolset/Source/MeshModelingToolsEditorOnly/Public/SkeletalMesh/SReferenceSkeletonTree.h b/Engine/Plugins/Runtime/MeshModelingToolset/Source/MeshModelingToolsEditorOnly/Public/SkeletalMesh/SReferenceSkeletonTree.h
--- a/Engine/Plugins/Runtime/MeshModelingToolset/Source/MeshModelingToolsEditorOnly/Public/SkeletalMesh/SReferenceSkeletonTree.h
+++ b/Engine/Plugins/Runtime/MeshModelingToolset/Source/MeshModelingToolsEditorOnly/Public/SkeletalMesh/SReferenceSkeletonTree.h
@@ -99,7 +99,7 @@ struct FRefSkeletonTreeDelegates
 	FOnRefSkeletonTreeCanAcceptDrop	OnCanAcceptDrop;
 	FOnRefSkeletonTreeAcceptDrop	OnAcceptDrop;
 	FOnBoneRenamed					OnBoneRenamed;	
-	UE_DEPRECATED(5.8, TEXT("Instead use OnBoneElementNameCommitted which is explicit about the renamed element"))
+	UE_DEPRECATED(5.8, "Instead use OnBoneElementNameCommitted which is explicit about the renamed element")
 	FOnTextCommitted				OnBoneNameCommitted;
 
 	DECLARE_DELEGATE_ThreeParams(FOnBoneElementNameCommitted, const FText& /** ElementName */, ETextCommit::Type /** TextCommitType */, TSharedRef<FBoneElement> /** BoneElement */)
