[SkeletalMeshModellingTools] Collapsing nodes in the skeleton tree no longer unselects children of the collapsed node

#Synced-CL 49723230

diff --git a/Engine/Plugins/Runtime/MeshModelingToolset/Source/MeshModelingToolsEditorOnly/Private/SkeletalMesh/SReferenceSkeletonTree.cpp b/Engine/Plugins/Runtime/MeshModelingToolset/Source/MeshModelingToolsEditorOnly/Private/SkeletalMesh/SReferenceSkeletonTree.cpp
--- a/Engine/Plugins/Runtime/MeshModelingToolset/Source/MeshModelingToolsEditorOnly/Private/SkeletalMesh/SReferenceSkeletonTree.cpp
+++ b/Engine/Plugins/Runtime/MeshModelingToolset/Source/MeshModelingToolsEditorOnly/Private/SkeletalMesh/SReferenceSkeletonTree.cpp
@@ -377,6 +377,7 @@ void SReferenceSkeletonTree::Construct(const FArguments& InArgs)
 					SAssignNew(TreeView, SRefSkeletonTreeView)
 					.TreeItemsSource(&RootElements)
 					.SelectionMode(ESelectionMode::Multi)
+					.AllowInvisibleItemSelection(true)
 					.OnGenerateRow_Lambda( [this](
 						TSharedPtr<FBoneElement> InItem,
 						const TSharedRef<STableViewBase>& OwnerTable) -> TSharedRef<ITableRow>
