[SkeletalMeshModellingTools] Bones are now mirrored in global space instead of the first parent's space

#Synced-CL 49723230

diff --git a/Engine/Plugins/Runtime/MeshModelingToolset/Source/SkeletalMeshModifiers/Private/SkeletonModifier.cpp b/Engine/Plugins/Runtime/MeshModelingToolset/Source/SkeletalMeshModifiers/Private/SkeletonModifier.cpp
--- a/Engine/Plugins/Runtime/MeshModelingToolset/Source/SkeletalMeshModifiers/Private/SkeletonModifier.cpp
+++ b/Engine/Plugins/Runtime/MeshModelingToolset/Source/SkeletalMeshModifiers/Private/SkeletonModifier.cpp
@@ -1533,55 +1533,28 @@ void USkeletonModifier::GetMirroredTransforms(
 	{
 		return;
 	}
-	
-	const TArray<FMeshBoneInfo>& BoneInfos = ReferenceSkeleton->GetRawRefBoneInfo();
-	auto FindFirstNotMirroredParent = [&](const int32 RefBoneIndex) -> int32
-	{
-		if (RefBoneIndex == INDEX_NONE)
-		{
-			return INDEX_NONE;
-		}
-		
-		int32 RefParentIndex = BoneInfos[RefBoneIndex].ParentIndex;
-		if (RefParentIndex == INDEX_NONE)
-		{
-			return INDEX_NONE;
-		}
 
-		int32 ParentMirroredIndex = InBonesToMirror.IndexOfByKey(RefParentIndex);
-		while (ParentMirroredIndex != INDEX_NONE)
-		{
-			RefParentIndex = BoneInfos[RefParentIndex].ParentIndex;
-			ParentMirroredIndex = InBonesToMirror.IndexOfByKey(RefParentIndex);
-		}
-		return RefParentIndex;
-	};
+	TArray<int32> BonesAndBonesToMirror;
+	BonesAndBonesToMirror.Append(InBonesToMirror);
+	BonesAndBonesToMirror.Append(InMirroredBones);
 
 	// compute global mirrored transforms
-	TArray<FTransform> MirroredGlobal; MirroredGlobal.Reserve(NumBonesToMirror);
+	TArray<FTransform> MirroredGlobal; 
+	MirroredGlobal.Reserve(NumBonesToMirror);
 	for (int32 Index = 0; Index < NumBonesToMirror; Index++)
 	{
 		// bone global
 		const int32 RefBoneIndex = InBonesToMirror[Index];
-		FTransform Global = TransformComposer->GetGlobalTransform(RefBoneIndex);
-	
-		// first parent not mirrored global
-		const int32 FirstNotMirroredParent = FindFirstNotMirroredParent(RefBoneIndex);
-		const FTransform& FirstNotMirroredGlobal = TransformComposer->GetGlobalTransform(FirstNotMirroredParent);
-
-		// switch to first parent not mirrored (translation only)
-		Global.AddToTranslation(-FirstNotMirroredGlobal.GetTranslation());
+		const FTransform GlobalTransform = TransformComposer->GetGlobalTransform(RefBoneIndex);
 
 		// mirror
-		Global = InOptions.MirrorTransform(Global);
-
-		// switch back to global (translation only)
-		Global.AddToTranslation(FirstNotMirroredGlobal.GetTranslation());
+		const FTransform MirroredGlobalTransform = InOptions.MirrorTransform(GlobalTransform);
 
-		MirroredGlobal.Add(Global);
+		MirroredGlobal.Add(MirroredGlobalTransform);
 	}
 
 	// switch back to local
+	const TArray<FMeshBoneInfo>& BoneInfos = ReferenceSkeleton->GetRawRefBoneInfo();
 	OutMirroredTransforms.Reserve(NumBonesToMirror);
 	for (int32 Index = 0; Index < NumBonesToMirror; Index++)
 	{
