[SkeletalMeshModellingTools] Improve selecting affected vertices of a Bone when Shift- or Ctrl-clicking the Affected button when editing Weights

* Fix Shift- and Ctrl-clicking Affected sometimes had no effect
* Ctrl-clicking Affected now toggles the selection of the affected vertices of a bone

#Synced-CL 49723230

diff --git a/Engine/Plugins/Runtime/MeshModelingToolset/Source/MeshModelingToolsEditorOnly/Private/SkeletalMesh/SkinWeightsPaintTool.cpp b/Engine/Plugins/Runtime/MeshModelingToolset/Source/MeshModelingToolsEditorOnly/Private/SkeletalMesh/SkinWeightsPaintTool.cpp
--- a/Engine/Plugins/Runtime/MeshModelingToolset/Source/MeshModelingToolsEditorOnly/Private/SkeletalMesh/SkinWeightsPaintTool.cpp
+++ b/Engine/Plugins/Runtime/MeshModelingToolset/Source/MeshModelingToolsEditorOnly/Private/SkeletalMesh/SkinWeightsPaintTool.cpp
@@ -2,7 +2,9 @@
 
 #include "SkeletalMesh/SkinWeightsPaintTool.h"
 
+#include "Algo/AllOf.h"
 #include "AssetViewerSettings.h"
+#include "Framework/Application/SlateApplication.h"
 
 #include "Math/UnrealMathUtility.h"
 #include "Async/Async.h"
@@ -4748,16 +4750,33 @@ void USkinWeightsPaintTool::SelectAffected() const
 
 	// optionally add/remove/replace selection based on modifier key state
 	const FGroupTopologySelection& CurrentSelection = SelectionMechanic->GetActiveSelection();
-	if (bShiftToggle)
+	
+	const bool bIsShiftDown = FSlateApplication::IsInitialized() ? FSlateApplication::Get().GetModifierKeys().IsShiftDown() : false;
+	const bool bIsControlDown = FSlateApplication::IsInitialized() ? FSlateApplication::Get().GetModifierKeys().IsControlDown() : false;
+	if (bIsShiftDown)
 	{
 		// ADD to selection
 		Selection.SelectedCornerIDs.Append(CurrentSelection.SelectedCornerIDs);
 		Selection.SelectedCornerIDs.Append(AffectedVertices);
 	}
-	else if (bCtrlToggle)
-	{
-		// REMOVE from selection
-		Selection.SelectedCornerIDs = CurrentSelection.SelectedCornerIDs.Difference(AffectedVertices);
+	else if (bIsControlDown)
+	{		
+		// REMOVE from selection if all affected vertices are contained in the current selection
+		// ADD to the selection otherwise
+		const bool bRemoveFromSelection = Algo::AllOf(AffectedVertices, [&CurrentSelection](const int32 VertexIndex)
+			{
+				return CurrentSelection.SelectedCornerIDs.Contains(VertexIndex);
+			});
+
+		if (bRemoveFromSelection)
+		{
+			Selection.SelectedCornerIDs = CurrentSelection.SelectedCornerIDs.Difference(AffectedVertices);
+		}
+		else
+		{
+			Selection.SelectedCornerIDs.Append(CurrentSelection.SelectedCornerIDs);
+			Selection.SelectedCornerIDs.Append(AffectedVertices);
+		}
 	}
 	else
 	{
