[SkeletalMeshModellingTools] Fix a crash when trying to edit Attributes of a Skeletal Mesh when its UVLayerNamesList is empty

#rnx

#Synced-CL 50157960

diff --git a/Engine/Plugins/Runtime/MeshModelingToolset/Source/MeshModelingToolsEditorOnly/Private/AttributeEditorTool.cpp b/Engine/Plugins/Runtime/MeshModelingToolset/Source/MeshModelingToolsEditorOnly/Private/AttributeEditorTool.cpp
--- a/Engine/Plugins/Runtime/MeshModelingToolset/Source/MeshModelingToolsEditorOnly/Private/AttributeEditorTool.cpp
+++ b/Engine/Plugins/Runtime/MeshModelingToolset/Source/MeshModelingToolsEditorOnly/Private/AttributeEditorTool.cpp
@@ -322,7 +322,10 @@ void UAttributeEditorTool::InitializeAttributeLists()
 	{
 		UVActions->UVLayerNamesList.Add(FString::Printf(TEXT("UV%d"), k));
 	}
-	UVActions->UVLayer = UVActions->UVLayerNamesList[0];
+	if (!UVActions->UVLayerNamesList.IsEmpty())
+	{
+		UVActions->UVLayer = UVActions->UVLayerNamesList[0];
+	}
 
 	if (bTargetIsStaticMesh && LightmapUVActions != nullptr)
 	{
