[SkeletalMeshModellingTools] Switching between the isolated and the full mesh no longer clears the selection when editing weights

#Synced-CL 50157960

diff --git a/Engine/Plugins/Runtime/MeshModelingToolset/Source/MeshModelingToolsEditorOnly/Private/SkeletalMesh/SkinWeightsPaintTool.cpp b/Engine/Plugins/Runtime/MeshModelingToolset/Source/MeshModelingToolsEditorOnly/Private/SkeletalMesh/SkinWeightsPaintTool.cpp
--- a/Engine/Plugins/Runtime/MeshModelingToolset/Source/MeshModelingToolsEditorOnly/Private/SkeletalMesh/SkinWeightsPaintTool.cpp
+++ b/Engine/Plugins/Runtime/MeshModelingToolset/Source/MeshModelingToolsEditorOnly/Private/SkeletalMesh/SkinWeightsPaintTool.cpp
@@ -4624,18 +4624,47 @@ void UWeightToolSelectionIsolator::RestoreFullMesh()
 		WeightTool->GetWeights().ApplyCurrentWeightsToMesh(FullMesh, PartialIndexToFullIndex);
 	}
 
+	// retain selection (a convenience for the user that allows for easily adjusting the isolation by going back/forth as needed)
+	struct FScopedRetainSelection
+	{
+		FScopedRetainSelection(USkinWeightsPaintTool& InWeightTool, const UE::Geometry::FDynamicSubmesh3& PartialSubMesh)
+		{
+			UToolMeshSelector* MeshSelector = InWeightTool.GetMainMeshSelector();
+			SelectionMechanic = MeshSelector ? MeshSelector->GetSelectionMechanic() : nullptr;
+			if (ensure(SelectionMechanic))
+			{
+				UE::Geometry::FGeometrySelection SelectionInParitalMesh;
+				SelectionMechanic->GetSelection_AsTriangleTopology(SelectionInParitalMesh);
+
+				SelectionInFullMesh = SelectionInParitalMesh;
+				SelectionInFullMesh.Selection.Reset();
+
+				Algo::Transform(SelectionInParitalMesh.Selection, SelectionInFullMesh.Selection,
+					[&PartialSubMesh](const int32 PartialIndex)
+					{
+						return PartialSubMesh.MapTriangleToBaseMesh(PartialIndex);
+					});
+			}
+		}
+
+		~FScopedRetainSelection()
+		{
+			if (SelectionMechanic)
+			{
+				SelectionMechanic->SetSelection_AsTriangleTopology(SelectionInFullMesh);
+			}
+		}
+
+		private:			
+			UE::Geometry::FGeometrySelection SelectionInFullMesh;
+			UPolygonSelectionMechanic* SelectionMechanic = nullptr;
+	};
+	const FScopedRetainSelection RetainSeletionScope(*WeightTool, PartialSubMesh);
+
 	// reinitialize with full mesh
 	// this resizes Weights to the full mesh size
 	WeightTool->UpdateCurrentlyEditedMesh(FullMesh);
 
-	// restore selection (a convenience for the user that allows for easily adjusting the isolation by going back/forth as needed)
-	if (UPolygonSelectionMechanic* SelectionMechanic = WeightTool->GetMainMeshSelector()->GetSelectionMechanic())
-	{
-		SelectionMechanic->SetSelection_AsTriangleTopology(IsolatedSelectionToRestoreVertices);
-		SelectionMechanic->SetSelection_AsTriangleTopology(IsolatedSelectionToRestoreEdges);
-		SelectionMechanic->SetSelection_AsTriangleTopology(IsolatedSelectionToRestoreFaces);
-	}
-
 	PartialSubMesh = {};
 	CurrentlyIsolatedTriangles.Reset();
 }
@@ -4689,23 +4718,22 @@ void UWeightToolSelectionIsolator::CreatePartialMesh()
 	// put into ref pose, BEFORE copying the mesh, so that submesh deformer initializes with vertices in ref pose
 	Weights.Deformer.SetToRefPose(WeightTool);
 
-	// store selection to be restored
-	IsolatedSelectionToRestoreVertices.Reset();
-	IsolatedSelectionToRestoreEdges.Reset();
-	IsolatedSelectionToRestoreFaces.Reset();
-	IsolatedSelectionToRestoreVertices.ElementType = UE::Geometry::EGeometryElementType::Vertex;
-	IsolatedSelectionToRestoreEdges.ElementType = UE::Geometry::EGeometryElementType::Edge;
-	IsolatedSelectionToRestoreFaces.ElementType = UE::Geometry::EGeometryElementType::Face;
-	SelectionMechanic->GetSelection_AsTriangleTopology(IsolatedSelectionToRestoreVertices);
-	SelectionMechanic->GetSelection_AsTriangleTopology(IsolatedSelectionToRestoreEdges);
-	SelectionMechanic->GetSelection_AsTriangleTopology(IsolatedSelectionToRestoreFaces);
-	
 	// create a partial sub-mesh from a subset of triangles on the dynamic mesh
 	PartialSubMesh = UE::Geometry::FDynamicSubmesh3(&FullMesh, CurrentlyIsolatedTriangles);
 
-
 	// reinitialize all mesh data structures
 	WeightTool->UpdateCurrentlyEditedMesh(PartialSubMesh.GetSubmesh());
+	
+	// retain selection (a convenience for the user that allows for easily adjusting the isolation by going back/forth as needed)
+	UE::Geometry::FGeometrySelection SelectionInParitalMesh;
+	SelectionInParitalMesh.ElementType = UE::Geometry::EGeometryElementType::Face;
+	SelectionInParitalMesh.TopologyType = UE::Geometry::EGeometryTopologyType::Triangle;
+	for (int32 TriangleID : PartialSubMesh.GetSubmesh().TriangleIndicesItr())
+	{
+		SelectionInParitalMesh.Selection.Add(TriangleID);
+	}	
+
+	SelectionMechanic->SetSelection_AsTriangleTopology(SelectionInParitalMesh);
 }
 
 
diff --git a/Engine/Plugins/Runtime/MeshModelingToolset/Source/MeshModelingToolsEditorOnly/Public/SkeletalMesh/SkinWeightsPaintTool.h b/Engine/Plugins/Runtime/MeshModelingToolset/Source/MeshModelingToolsEditorOnly/Public/SkeletalMesh/SkinWeightsPaintTool.h
--- a/Engine/Plugins/Runtime/MeshModelingToolset/Source/MeshModelingToolsEditorOnly/Public/SkeletalMesh/SkinWeightsPaintTool.h
+++ b/Engine/Plugins/Runtime/MeshModelingToolset/Source/MeshModelingToolsEditorOnly/Public/SkeletalMesh/SkinWeightsPaintTool.h
@@ -615,10 +615,10 @@ private:
 	// when selection is isolated, we hide the full mesh and show a submesh
 	// when islated selection is unhidden, we remap all changes from the submesh back to the full mesh
 	TArray<int32> CurrentlyIsolatedTriangles;
-	UE::Geometry::FGeometrySelection IsolatedSelectionToRestoreVertices;
-	UE::Geometry::FGeometrySelection IsolatedSelectionToRestoreEdges;
-	UE::Geometry::FGeometrySelection IsolatedSelectionToRestoreFaces;
+
+	/** Flag to indicate the isolated mesh needs to be updated, handled on the next OnTick */
 	bool bIsolatedMeshNeedsUpdated = false;
+
 	// isolate selection sub-meshes
 	UE::Geometry::FDynamicSubmesh3 PartialSubMesh;
 };
