Dataflow : add support to lasso to the paint tool by leveraging the newly added feature in the mesh selector

#Synced-CL 50358557

diff --git a/Engine/Plugins/Experimental/Dataflow/Source/DataflowEditor/Private/DataflowEditorTools/DataflowEditorVertexAttributePaintTool.cpp b/Engine/Plugins/Experimental/Dataflow/Source/DataflowEditor/Private/DataflowEditorTools/DataflowEditorVertexAttributePaintTool.cpp
--- a/Engine/Plugins/Experimental/Dataflow/Source/DataflowEditor/Private/DataflowEditorTools/DataflowEditorVertexAttributePaintTool.cpp
+++ b/Engine/Plugins/Experimental/Dataflow/Source/DataflowEditor/Private/DataflowEditorTools/DataflowEditorVertexAttributePaintTool.cpp
@@ -35,6 +35,7 @@
 #include "Polygon2.h"
 #include "PreviewMesh.h"
 #include "StaticMeshAttributes.h"
+#include "Selection/PolygonSelectionMechanic.h"
 #include "Selections/MeshConnectedComponents.h"
 #include "Selections/MeshVertexSelection.h"
 #include "Spatial/PointHashGrid3.h"
@@ -725,7 +726,9 @@ void UDataflowEditorVertexAttributePaintTool::Setup()
 
 	// make sure the selector is in sync with the settings
 	MeshSelector->SetComponentSelectionMode(ToolProperties->SelectionProperties.ComponentSelectionMode);
-
+	MeshSelector->SetSelectionTool(EMeshSelectorTool::Marquee);
+	MeshSelector->GetSelectionMechanic()->SetSelectionDragToolUpdateType(ESelectionDragToolUpdateType::OnTickAndRelease);
+	
 	InitializeIndicator();
 	InitializeMirrorIndicator();
 
@@ -1161,6 +1164,22 @@ EAxis::Type UDataflowEditorVertexAttributePaintTool::GetMirrorAxis() const
 	return ToolProperties ? ToolProperties->MirrorProperties.MirrorAxis.GetValue() : EAxis::X;
 }
 
+void UDataflowEditorVertexAttributePaintTool::SetMeshSelectionTool(EMeshSelectorTool SelectorTool)
+{
+	if (MeshSelector)
+	{
+		MeshSelector->SetSelectionTool(SelectorTool);
+	}
+}
+
+EMeshSelectorTool UDataflowEditorVertexAttributePaintTool::GetMeshSelectionTool() const
+{
+	return MeshSelector
+		? MeshSelector->GetSelectionTool()
+		: EMeshSelectorTool::Ray;
+}
+
+
 FInputChord UDataflowEditorVertexAttributePaintTool::GreyscaleDisplayInputChord = FInputChord(EModifierKey::Control, EKeys::NumPadSeven);
 FInputChord UDataflowEditorVertexAttributePaintTool::ColorMapDisplayInputChord = FInputChord(EModifierKey::Control, EKeys::NumPadEight);
 FInputChord UDataflowEditorVertexAttributePaintTool::MaterialColorDisplayInputChord = FInputChord(EModifierKey::Control, EKeys::NumPadSix);
@@ -1931,85 +1950,6 @@ namespace UE::DataflowEditorVertexAttributePaintTool::Private
 	}
 }
 
-
-void UDataflowEditorVertexAttributePaintTool::OnPolyLassoFinished(const FCameraPolyLasso& Lasso, bool bCanceled)
-{
-	// construct polyline
-	TArray<FVector2f> Polyline;
-	for (FVector2D Pos : Lasso.Polyline)
-	{
-		Polyline.Add((FVector2f)Pos);
-	}
-	int32 N = Polyline.Num();
-	if (N < 2)
-	{
-		return;
-	}
-
-	// Try to clip polyline to be closed, or closed-enough for winding evaluation to work.
-	// If that returns false, the polyline is "too open". In that case we will extend
-	// outwards from the endpoints and then try to create a closed very large polygon
-	if (UE::DataflowEditorVertexAttributePaintTool::Private::ApproxSelfClipPolyline(Polyline) == false)
-	{
-		FVector2f StartDirOut = UE::Geometry::Normalized(Polyline[0] - Polyline[1]);
-		UE::Geometry::FLine2f StartLine(Polyline[0], StartDirOut);
-		FVector2f EndDirOut = UE::Geometry::Normalized(Polyline[N - 1] - Polyline[N - 2]);
-		UE::Geometry::FLine2f EndLine(Polyline[N - 1], EndDirOut);
-
-		// if we did not intersect, we are in ambiguous territory. Check if a segment along either end-direction
-		// intersects the polyline. If it does, we have something like a spiral and will be OK. 
-		// If not, make a closed polygon by interpolating outwards from each endpoint, and then in perp-directions.
-		UE::Geometry::FPolygon2f Polygon(Polyline);
-		float PerpSign = Polygon.IsClockwise() ? -1.0 : 1.0;
-
-		Polyline.Insert(StartLine.PointAt(10000.0f), 0);
-		Polyline.Insert(Polyline[0] + 1000 * PerpSign * UE::Geometry::PerpCW(StartDirOut), 0);
-
-		Polyline.Add(EndLine.PointAt(10000.0f));
-		Polyline.Add(Polyline.Last() + 1000 * PerpSign * UE::Geometry::PerpCW(EndDirOut));
-		FVector2f StartPos = Polyline[0];
-		Polyline.Add(StartPos);		// close polyline (cannot use Polyline[0] in case Add resizes!)
-	}
-
-	N = Polyline.Num();
-
-	// project each mesh vertex to view plane and evaluate winding integral of polyline
-	const FDynamicMesh3* Mesh = GetSculptMesh();
-	TempROIBuffer.SetNum(Mesh->MaxVertexID());
-	ParallelFor(Mesh->MaxVertexID(), [&](int32 VertexIdx)
-		{
-			if (Mesh->IsVertex(VertexIdx))
-			{
-				FVector3d WorldPos = CurTargetTransform.TransformPosition(Mesh->GetVertex(VertexIdx));
-				FVector2f PlanePos = (FVector2f)Lasso.GetProjectedPoint((FVector)WorldPos);
-
-				double WindingSum = 0;
-				FVector2f a = Polyline[0] - PlanePos, b = FVector2f::Zero();
-				for (int32 i = 1; i < N; ++i)
-				{
-					b = Polyline[i] - PlanePos;
-					WindingSum += (double)FMathf::Atan2(a.X * b.Y - a.Y * b.X, a.X * b.X + a.Y * b.Y);
-					a = b;
-				}
-				WindingSum /= FMathd::TwoPi;
-				bool bInside = FMathd::Abs(WindingSum) > 0.3;
-				TempROIBuffer[VertexIdx] = bInside ? 1 : 0;
-			}
-			else
-			{
-				TempROIBuffer[VertexIdx] = -1;
-			}
-		});
-
-	// convert to vertex selection, and then select fully-enclosed faces
-	UE::Geometry::FMeshVertexSelection VertexSelection(Mesh);
-	VertexSelection.SelectByVertexID([&](int32 VertexIdx) { return TempROIBuffer[VertexIdx] == 1; });
-
-	const double SetWeightValue = ToolProperties? ToolProperties->BrushProperties.AttributeValue: 0.0;
-	SetVerticesToWeightMap(VertexSelection.AsSet(), SetWeightValue);
-}
-
-
 void UDataflowEditorVertexAttributePaintTool::ComputeGradient()
 {
 	//if (!ensure(ActiveWeightMap))
@@ -2556,12 +2496,6 @@ bool UDataflowEditorVertexAttributePaintTool::OnUpdateHover(const FInputDeviceRa
 
 void UDataflowEditorVertexAttributePaintTool::DrawHUD(FCanvas* Canvas, IToolsContextRenderAPI* RenderAPI)
 {
-	if (PolyLassoMechanic)
-	{
-		PolyLassoMechanic->LineColor = FLinearColor::Green;
-		PolyLassoMechanic->DrawHUD(Canvas, RenderAPI);
-	}
-
 	if (ToolProperties->EditingMode == EDataflowEditorToolEditMode::Mesh)
 	{
 		if (MeshSelector)
@@ -2653,9 +2587,26 @@ void UDataflowEditorVertexAttributePaintTool::DrawHUD(FCanvas* Canvas, IToolsCon
 				const float MaxValue = ToolProperties->DisplayProperties.MaxValue;
 				const float ValueRange = MaxValue - MinValue;
 
-				for (int32 VertexID : Mesh->VertexIndicesItr())
+				UFont* EngineLargeFont = GEngine->GetLargeFont();
+
+				const int32 MaxVertexId = Mesh->MaxVertexID();
+				TArray<bool> VertexVisibility;
+				VertexVisibility.Init(false, MaxVertexId);
+
+				ParallelFor(Mesh->MaxVertexID(), [this, &VertexVisibility, &Mesh, IsVertexVisible](int32 VertexID)
+					{
+						if (Mesh->IsVertex(VertexID))
+						{
+							if (IsVertexVisible(*Mesh, VertexID))
+							{
+								VertexVisibility[VertexID] = true;
+							}
+						}
+					});
+
+				for (int32 VertexID = 0; VertexID < MaxVertexId; ++VertexID)
 				{
-					if (IsVertexVisible(*Mesh, VertexID))
+					if (VertexVisibility[VertexID])
 					{
 						const float NormalizedValue = VertexData.GetValue(VertexID);
 						const float RemappedValue = MinValue + NormalizedValue * ValueRange;
@@ -2671,7 +2622,7 @@ void UDataflowEditorVertexAttributePaintTool::DrawHUD(FCanvas* Canvas, IToolsCon
 						TextItem.bCentreX = true;
 						Canvas->DrawItem(TextItem);
 					}
-				}
+				};
 			}
 		}
 	}
@@ -2715,20 +2666,15 @@ void UDataflowEditorVertexAttributePaintTool::OnTick(float DeltaTime)
 	UMeshSculptToolBase::OnTick(DeltaTime);
 	MeshElementsDisplay->OnTick(DeltaTime);
 
-	bool bIsLasso = false; // TODO(ccaillaud) const bool bIsLasso = (FilterProperties->SubToolType == EDataflowEditorWeightMapPaintInteractionType::PolyLasso);
-	if (PolyLassoMechanic)
-	{
-		PolyLassoMechanic->SetIsEnabled(bIsLasso);
-	}
-
 	const bool bIsMeshEditMode = (ToolProperties->EditingMode == EDataflowEditorToolEditMode::Mesh);
 	if (MeshSelector)
 	{
 		MeshSelector->SetIsEnabled(bIsMeshEditMode);
+		MeshSelector->Tick(DeltaTime);
 	}
 
 	ConfigureIndicator(IsVolumetricBrush());
-	SetIndicatorVisibility(!bIsLasso && !bIsMeshEditMode && bBrushIsHoverMesh && !bPendingPickWeight);
+	SetIndicatorVisibility(!bIsMeshEditMode && bBrushIsHoverMesh && !bPendingPickWeight);
 	UpdateMirrorIndicator();
 
 	SCOPE_CYCLE_COUNTER(VertexAttributePaintTool_Tick);
@@ -3040,8 +2986,9 @@ int32 UDataflowEditorVertexAttributePaintTool::GetBrushNearestVertexAccurate() c
 		}
 
 		double MinDist = UE_BIG_NUMBER;
-		for (const int32 VertexIndex : VertexROI)
+		for (int32 TriangleVertexIndex = 0; TriangleVertexIndex < 3; ++TriangleVertexIndex)
 		{
+			const int32 VertexIndex = Vertices[TriangleVertexIndex];
 			const FVector3d& VertexPosition = Mesh->GetVertex(VertexIndex);
 			const double CurrDist = FVector3d::Distance(VertexPosition, PointOnSurface);
 			if (CurrDist < MinDist)
diff --git a/Engine/Plugins/Experimental/Dataflow/Source/DataflowEditor/Private/DataflowEditorTools/DataflowEditorVertexAttributePaintToolDetailCustomization.cpp b/Engine/Plugins/Experimental/Dataflow/Source/DataflowEditor/Private/DataflowEditorTools/DataflowEditorVertexAttributePaintToolDetailCustomization.cpp
--- a/Engine/Plugins/Experimental/Dataflow/Source/DataflowEditor/Private/DataflowEditorTools/DataflowEditorVertexAttributePaintToolDetailCustomization.cpp
+++ b/Engine/Plugins/Experimental/Dataflow/Source/DataflowEditor/Private/DataflowEditorTools/DataflowEditorVertexAttributePaintToolDetailCustomization.cpp
@@ -1178,6 +1178,49 @@ namespace UE::Dataflow::Editor
 
 	///////////////////////////////////////////////////////////////////////////////////////////////////////////////
 
+	TSharedRef<SWidget> FVertexAttributePaintToolDetailCustomization::MakeDragToolToolbar() const
+	{
+		FSlimHorizontalToolBarBuilder ToolbarBuilder(MakeShared<FUICommandList>(), FMultiBoxCustomization::None);
+		ToolbarBuilder.SetStyle(FModelingToolsEditorModeStyle::Get().Get(), "PolyEd.SelectionToolbar");
+		ToolbarBuilder.SetLabelVisibility(EVisibility::Collapsed);
+
+		ToolbarBuilder.BeginSection("SelectionDragTool");
+		ToolbarBuilder.BeginBlockGroup();
+
+		auto AddToggleButtonForBool = [&ToolbarBuilder, this](EMeshSelectorTool SelectionTool, const FText& Label, const FText& Tooltip, const FName IconName)
+			{
+				ToolbarBuilder.AddToolBarButton(FUIAction(
+					FExecuteAction::CreateLambda([this, SelectionTool]()
+						{
+							if (Tool.IsValid())
+							{
+								Tool->SetMeshSelectionTool(SelectionTool);
+							}
+						}),
+					FCanExecuteAction::CreateLambda([this]()
+						{
+							return ToolProperties.IsValid() ? ToolProperties->EditingMode == EDataflowEditorToolEditMode::Mesh : false;
+						}),
+					FIsActionChecked::CreateLambda([this, SelectionTool]()
+						{
+							return Tool.IsValid() ? (Tool->GetMeshSelectionTool() == SelectionTool) : false;
+						})),
+					NAME_None,	// Extension hook
+					Label,		// Label
+					Tooltip,	// Tooltip
+					FSlateIcon(FAppStyle::GetAppStyleSetName(), IconName),
+					EUserInterfaceActionType::ToggleButton);
+			};
+
+
+
+
+		ToolbarBuilder.EndBlockGroup();
+		ToolbarBuilder.EndSection();
+
+		return ToolbarBuilder.MakeWidget();
+	}
+
 	TSharedRef<SWidget> FVertexAttributePaintToolDetailCustomization::MakeSelectionElementsToolbar() const
 	{
 		FSlimHorizontalToolBarBuilder ToolbarBuilder(MakeShared<FUICommandList>(), FMultiBoxCustomization::None);
@@ -1187,7 +1230,45 @@ namespace UE::Dataflow::Editor
 		ToolbarBuilder.BeginSection("SelectionFilter");
 		ToolbarBuilder.BeginBlockGroup();
 
-		auto AddToggleButtonForBool = [&ToolbarBuilder, this](EComponentSelectionMode Mode, const FText& Label, const FText& Tooltip, const FName IconName)
+		auto AddToggleButtonSelectionTool = [&ToolbarBuilder, this](EMeshSelectorTool SelectionTool, const FText& Label, const FText& Tooltip, const FSlateIcon& Icon)
+			{
+				ToolbarBuilder.AddToolBarButton(FUIAction(
+					FExecuteAction::CreateLambda([this, SelectionTool]()
+						{
+							if (Tool.IsValid())
+							{
+								Tool->SetMeshSelectionTool(SelectionTool);
+							}
+						}),
+					FCanExecuteAction::CreateLambda([this]()
+						{
+							return ToolProperties.IsValid() ? ToolProperties->EditingMode == EDataflowEditorToolEditMode::Mesh : false;
+						}),
+					FIsActionChecked::CreateLambda([this, SelectionTool]()
+						{
+							return Tool.IsValid() ? (Tool->GetMeshSelectionTool() == SelectionTool) : false;
+						})),
+					NAME_None,	// Extension hook
+					Label,		// Label
+					Tooltip,	// Tooltip
+					Icon,
+					EUserInterfaceActionType::ToggleButton);
+			};
+
+		AddToggleButtonSelectionTool(
+			EMeshSelectorTool::Marquee,
+			LOCTEXT("RectangleSelectionToolLabel", "Rectangle"),
+			LOCTEXT("RectangleSelectionToolTooltip", "Use rectangle selection tool on click and drag."),
+			FSlateIcon(FAppStyle::GetAppStyleSetName(), "LevelEditor.StrictBoxSelect")
+		);
+		AddToggleButtonSelectionTool(
+			EMeshSelectorTool::Lasso,
+			LOCTEXT("LassoSelectionToolLabel", "Lasso"),
+			LOCTEXT("LassoSelectionToolTooltip", "Use lasso selection tool on click and drag."),
+			FSlateIcon(FAppStyle::GetAppStyleSetName(), "FoliageEditMode.SetLassoSelect")
+		);
+
+		auto AddToggleButtonComponentMode = [&ToolbarBuilder, this](EComponentSelectionMode Mode, const FText& Label, const FText& Tooltip, const FSlateIcon& Icon)
 			{
 				ToolbarBuilder.AddToolBarButton(FUIAction(
 					FExecuteAction::CreateLambda([this, Mode]()
@@ -1208,25 +1289,29 @@ namespace UE::Dataflow::Editor
 					NAME_None,	// Extension hook
 					Label,		// Label
 					Tooltip,	// Tooltip
-					FSlateIcon(FModelingToolsEditorModeStyle::Get()->GetStyleSetName(), IconName),
+					Icon,
 					EUserInterfaceActionType::ToggleButton);
 			};
 
-		AddToggleButtonForBool(
+		ToolbarBuilder.AddSeparator();
+		AddToggleButtonComponentMode(
 			EComponentSelectionMode::Vertices,
 			LOCTEXT("VerticesLabel", "Vertices"),
 			LOCTEXT("VerticesTooltip", "Select mesh vertices."),
-			"PolyEd.SelectCorners");
-		AddToggleButtonForBool(
+			FSlateIcon(FModelingToolsEditorModeStyle::Get()->GetStyleSetName(), "PolyEd.SelectCorners")
+		);
+		AddToggleButtonComponentMode(
 			EComponentSelectionMode::Edges,
 			LOCTEXT("EdgesLabel", "Edges"),
 			LOCTEXT("EdgesTooltip", "Select mesh edges."),
-			"PolyEd.SelectEdges");
-		AddToggleButtonForBool(
+			FSlateIcon(FModelingToolsEditorModeStyle::Get()->GetStyleSetName(), "PolyEd.SelectEdges")
+		);
+		AddToggleButtonComponentMode(
 			EComponentSelectionMode::Faces,
 			LOCTEXT("FacesLabel", "Faces"),
 			LOCTEXT("FacesTooltip", "Select mesh faces."),
-			"PolyEd.SelectFaces");
+			FSlateIcon(FModelingToolsEditorModeStyle::Get()->GetStyleSetName(), "PolyEd.SelectFaces")
+		);
 
 		ToolbarBuilder.EndBlockGroup();
 		ToolbarBuilder.EndSection();
diff --git a/Engine/Plugins/Experimental/Dataflow/Source/DataflowEditor/Private/DataflowEditorTools/DataflowEditorVertexAttributePaintToolDetailCustomization.h b/Engine/Plugins/Experimental/Dataflow/Source/DataflowEditor/Private/DataflowEditorTools/DataflowEditorVertexAttributePaintToolDetailCustomization.h
--- a/Engine/Plugins/Experimental/Dataflow/Source/DataflowEditor/Private/DataflowEditorTools/DataflowEditorVertexAttributePaintToolDetailCustomization.h
+++ b/Engine/Plugins/Experimental/Dataflow/Source/DataflowEditor/Private/DataflowEditorTools/DataflowEditorVertexAttributePaintToolDetailCustomization.h
@@ -58,6 +58,7 @@ namespace UE::Dataflow::Editor
 		void AddVisibilityFilterRow(IDetailCategoryBuilder& BrushCategory) const;
 
 		void AddSelectionUI(IDetailLayoutBuilder& DetailBuilder) const;
+		TSharedRef<SWidget> MakeDragToolToolbar() const;
 		TSharedRef<SWidget> MakeSelectionElementsToolbar() const;
 		TSharedRef<SWidget> MakeSelectionIsolationWidget() const;
 		TSharedRef<SWidget> MakeSelectionEditActionsToolbar() const;
diff --git a/Engine/Plugins/Experimental/Dataflow/Source/DataflowEditor/Public/DataflowEditorTools/DataflowEditorVertexAttributePaintTool.h b/Engine/Plugins/Experimental/Dataflow/Source/DataflowEditor/Public/DataflowEditorTools/DataflowEditorVertexAttributePaintTool.h
--- a/Engine/Plugins/Experimental/Dataflow/Source/DataflowEditor/Public/DataflowEditorTools/DataflowEditorVertexAttributePaintTool.h
+++ b/Engine/Plugins/Experimental/Dataflow/Source/DataflowEditor/Public/DataflowEditorTools/DataflowEditorVertexAttributePaintTool.h
@@ -9,7 +9,6 @@
 #include "BaseTools/MeshSurfacePointMeshEditingTool.h"
 #include "Components/DynamicMeshComponent.h"
 #include "PropertySets/WeightMapSetProperties.h"
-#include "Mechanics/PolyLassoMarqueeMechanic.h"
 
 #include "Sculpting/MeshSculptToolBase.h"
 #include "Sculpting/MeshBrushOpBase.h"
@@ -561,17 +560,10 @@ protected:
 	UE::Geometry::FFrame3d ComputeMirroredFrame(const UE::Geometry::FFrame3d& InLocalSpaceFrame) const;
 	void OnMirrorCommand();
 
-protected:
-	//
-	// Marquee Support
-	//
-public:
-	UPROPERTY()
-	TObjectPtr<UPolyLassoMarqueeMechanic> PolyLassoMechanic;
+	void SetMeshSelectionTool(EMeshSelectorTool SelectorTool);
+	EMeshSelectorTool GetMeshSelectionTool() const;
 
 protected:
-	UE_API void OnPolyLassoFinished(const FCameraPolyLasso& Lasso, bool bCanceled);
-
 
 	// 
 	// Gradient Support
diff --git a/Engine/Plugins/Experimental/MeshModelingToolsetExp/Source/MeshModelingToolsExp/Private/EditUVIslandsTool.cpp b/Engine/Plugins/Experimental/MeshModelingToolsetExp/Source/MeshModelingToolsExp/Private/EditUVIslandsTool.cpp
--- a/Engine/Plugins/Experimental/MeshModelingToolsetExp/Source/MeshModelingToolsExp/Private/EditUVIslandsTool.cpp
+++ b/Engine/Plugins/Experimental/MeshModelingToolsetExp/Source/MeshModelingToolsExp/Private/EditUVIslandsTool.cpp
@@ -93,7 +93,7 @@ void UEditUVIslandsTool::Setup()
 	SelectionMechanic->Setup(this);
 	SelectionMechanic->Properties->bSelectEdges = false;
 	SelectionMechanic->Properties->bSelectVertices = false;
-	SelectionMechanic->Properties->bEnableMarquee = false;
+	SelectionMechanic->Properties->SelectionDragTool = EMeshTopologySelectionMechanicDragTool::None;
 	SelectionMechanic->OnSelectionChanged.AddUObject(this, &UEditUVIslandsTool::OnSelectionModifiedEvent);
 
 	// initialize AABBTree
diff --git a/Engine/Plugins/Experimental/MeshModelingToolsetExp/Source/MeshModelingToolsExp/Private/ExtractSplineTool.cpp b/Engine/Plugins/Experimental/MeshModelingToolsetExp/Source/MeshModelingToolsExp/Private/ExtractSplineTool.cpp
--- a/Engine/Plugins/Experimental/MeshModelingToolsetExp/Source/MeshModelingToolsExp/Private/ExtractSplineTool.cpp
+++ b/Engine/Plugins/Experimental/MeshModelingToolsetExp/Source/MeshModelingToolsExp/Private/ExtractSplineTool.cpp
@@ -131,7 +131,7 @@ void UExtractSplineTool::Setup()
 	SelectionMechanic->Properties->bCanSelectVertices = false;
 	SelectionMechanic->Properties->bCanSelectEdges = true;
 	SelectionMechanic->Properties->bDisplayPolygroupReliantControls = true;
-	SelectionMechanic->Properties->bEnableMarquee = false;
+	SelectionMechanic->Properties->SelectionDragTool = EMeshTopologySelectionMechanicDragTool::None;
 
 	SelectionMechanic->OnSelectionChanged.AddUObject(this, &UExtractSplineTool::InvalidatePreviews);
 	SelectionMechanic->SetShouldAddToSelectionFunc([]() {return false; });
diff --git a/Engine/Plugins/Experimental/MeshResizing/Source/MeshResizingEditorTools/Private/MeshResizing/MeshWrapLandmarkSelectionTool.cpp b/Engine/Plugins/Experimental/MeshResizing/Source/MeshResizingEditorTools/Private/MeshResizing/MeshWrapLandmarkSelectionTool.cpp
--- a/Engine/Plugins/Experimental/MeshResizing/Source/MeshResizingEditorTools/Private/MeshResizing/MeshWrapLandmarkSelectionTool.cpp
+++ b/Engine/Plugins/Experimental/MeshResizing/Source/MeshResizingEditorTools/Private/MeshResizing/MeshWrapLandmarkSelectionTool.cpp
@@ -175,7 +175,7 @@ void UMeshWrapLandmarkSelectionTool::Setup()
 	SelectionMechanic->Properties->bSelectEdges = false;
 	SelectionMechanic->Properties->bSelectFaces = false;
 
-	SelectionMechanic->Properties->bEnableMarquee = false;
+	SelectionMechanic->Properties->SelectionDragTool = EMeshTopologySelectionMechanicDragTool::None;
 	
 	SelectionMechanic->SetShowEdges(false);
 	SelectionMechanic->SetShowSelectableCorners(false);
diff --git a/Engine/Plugins/Runtime/MeshModelingToolset/Source/MeshModelingToolsEditorOnly/Private/SkeletalMesh/SkeletonEditingTool.cpp b/Engine/Plugins/Runtime/MeshModelingToolset/Source/MeshModelingToolsEditorOnly/Private/SkeletalMesh/SkeletonEditingTool.cpp
--- a/Engine/Plugins/Runtime/MeshModelingToolset/Source/MeshModelingToolsEditorOnly/Private/SkeletalMesh/SkeletonEditingTool.cpp
+++ b/Engine/Plugins/Runtime/MeshModelingToolset/Source/MeshModelingToolsEditorOnly/Private/SkeletalMesh/SkeletonEditingTool.cpp
@@ -336,7 +336,7 @@ void USkeletonEditingTool::SetupComponentsSelection()
 	SelectionMechanic->SetIsEnabled(false);
 	// the property is disabled and managed thru customization
 	SetToolPropertySourceEnabled(SelectionMechanic->Properties, false);
-	SelectionMechanic->SetMarqueeSelectionUpdateType(EMarqueeSelectionUpdateType::OnRelease);
+	SelectionMechanic->SetSelectionDragToolUpdateType(ESelectionDragToolUpdateType::OnRelease);
 	SelectionMechanic->Properties->bDisplayPolygroupReliantControls = false;
 
 	SelectionMechanic->Properties->bSelectVertices = false;
diff --git a/Engine/Plugins/Runtime/MeshModelingToolset/Source/MeshModelingToolsEditorOnly/Private/ToolMeshSelector.cpp b/Engine/Plugins/Runtime/MeshModelingToolset/Source/MeshModelingToolsEditorOnly/Private/ToolMeshSelector.cpp
--- a/Engine/Plugins/Runtime/MeshModelingToolset/Source/MeshModelingToolsEditorOnly/Private/ToolMeshSelector.cpp
+++ b/Engine/Plugins/Runtime/MeshModelingToolset/Source/MeshModelingToolsEditorOnly/Private/ToolMeshSelector.cpp
@@ -134,6 +134,47 @@ void UToolMeshSelector::SetIsEnabled(bool bIsEnabled)
 	PolygonSelectionMechanic->SetIsEnabled(bIsEnabled, bIsEnabled);
 }
 
+void UToolMeshSelector::SetSelectionTool(EMeshSelectorTool InTool)
+{
+	if (!PolygonSelectionMechanic || !PolygonSelectionMechanic->Properties)
+	{
+		return;
+	}
+
+	switch (InTool)
+	{
+	case EMeshSelectorTool::Marquee:
+		PolygonSelectionMechanic->Properties->SelectionDragTool = EMeshTopologySelectionMechanicDragTool::Marquee;
+		break;
+
+	case EMeshSelectorTool::Lasso:
+		PolygonSelectionMechanic->Properties->SelectionDragTool = EMeshTopologySelectionMechanicDragTool::Lasso;
+		break;
+
+	case EMeshSelectorTool::Ray:
+	default:
+		PolygonSelectionMechanic->Properties->SelectionDragTool = EMeshTopologySelectionMechanicDragTool::None;
+		break;
+	}
+}
+
+EMeshSelectorTool UToolMeshSelector::GetSelectionTool() const
+{
+	if (!PolygonSelectionMechanic || !PolygonSelectionMechanic->Properties)
+	{
+		return EMeshSelectorTool::Ray;
+	}
+	switch (PolygonSelectionMechanic->Properties->SelectionDragTool)
+	{
+	case EMeshTopologySelectionMechanicDragTool::Marquee:
+		return EMeshSelectorTool::Marquee;
+	case EMeshTopologySelectionMechanicDragTool::Lasso:
+		return EMeshSelectorTool::Lasso;
+	default:
+		return EMeshSelectorTool::Ray;
+	}
+}
+
 void UToolMeshSelector::SetComponentSelectionMode(EComponentSelectionMode InMode)
 {
 	if (!(PolygonSelectionMechanic))
@@ -164,6 +205,14 @@ void UToolMeshSelector::UpdateAfterMeshDeformation()
 	PolygonSelectionMechanic->GetTopologySelector()->Invalidate(bTopologyDeformed, bTopologyModified);
 }
 
+void UToolMeshSelector::Tick(float DeltaTime)
+{
+	if (PolygonSelectionMechanic)
+	{
+		PolygonSelectionMechanic->Tick(DeltaTime);
+	}
+}
+
 void UToolMeshSelector::DrawHUD(FCanvas* Canvas, IToolsContextRenderAPI* RenderAPI)
 {
 	if (!PolygonSelectionMechanic)
diff --git a/Engine/Plugins/Runtime/MeshModelingToolset/Source/MeshModelingToolsEditorOnly/Public/ToolMeshSelector.h b/Engine/Plugins/Runtime/MeshModelingToolset/Source/MeshModelingToolsEditorOnly/Public/ToolMeshSelector.h
--- a/Engine/Plugins/Runtime/MeshModelingToolset/Source/MeshModelingToolsEditorOnly/Public/ToolMeshSelector.h
+++ b/Engine/Plugins/Runtime/MeshModelingToolset/Source/MeshModelingToolsEditorOnly/Public/ToolMeshSelector.h
@@ -32,6 +32,14 @@ enum class EComponentSelectionMode : uint8
 	Faces
 };
 
+UENUM()
+enum class EMeshSelectorTool : uint8
+{
+	Ray,
+	Marquee,
+	Lasso
+};
+
 using VertexIndex = int32;
 
 // this class wraps the all the components to enable selection on a single mesh in the skin weights tool
@@ -65,6 +73,9 @@ public:
 	UE_API void SetComponentSelectionMode(EComponentSelectionMode InMode);
 	UE_API void SetTransform(const FTransform3d& InTargetTransform);
 
+	UE_API void SetSelectionTool(EMeshSelectorTool InTool);
+	UE_API EMeshSelectorTool GetSelectionTool() const;
+
 	// viewport 
 	UE_API void DrawHUD(FCanvas* Canvas, IToolsContextRenderAPI* RenderAPI);
 	UE_API void Render(IToolsContextRenderAPI* RenderAPI);
@@ -81,6 +92,8 @@ public:
 	UE_API void FloodSelection() const;
 	UE_API void SelectBorder() const;
 
+	UE_API void Tick(float DeltaTime);
+
 	// get access to the selection mechanic
 	UPolygonSelectionMechanic* GetSelectionMechanic() { return PolygonSelectionMechanic; };
 
diff --git a/Engine/Plugins/Runtime/MeshModelingToolset/Source/ModelingComponents/Private/Selection/MeshTopologySelectionMechanic.cpp b/Engine/Plugins/Runtime/MeshModelingToolset/Source/ModelingComponents/Private/Selection/MeshTopologySelectionMechanic.cpp
--- a/Engine/Plugins/Runtime/MeshModelingToolset/Source/ModelingComponents/Private/Selection/MeshTopologySelectionMechanic.cpp
+++ b/Engine/Plugins/Runtime/MeshModelingToolset/Source/ModelingComponents/Private/Selection/MeshTopologySelectionMechanic.cpp
@@ -6,10 +6,12 @@
 #include "BaseBehaviors/MouseHoverBehavior.h"
 #include "Engine/World.h"
 #include "InteractiveToolManager.h"
-#include "Util/ColorConstants.h"
+#include "Mechanics/PolyLassoMarqueeMechanic.h"
+#include "Polygon2.h"
 #include "Selection/PersistentMeshSelection.h"
 #include "ToolSceneQueriesUtil.h"
 #include "ToolSetupUtil.h"
+#include "Util/ColorConstants.h"
 
 #include UE_INLINE_GENERATED_CPP_BY_NAME(MeshTopologySelectionMechanic)
 
@@ -58,6 +60,16 @@ void UMeshTopologySelectionMechanic::Setup(UInteractiveTool* ParentToolIn)
 	MarqueeMechanic->OnDragRectangleFinished.AddUObject(this, &UMeshTopologySelectionMechanic::OnDragRectangleFinished);
 	MarqueeMechanic->SetBasePriority(BasePriority.MakeLower());
 
+	// Set up control points mechanic
+	PolyLassoMechanic = NewObject<UPolyLassoMarqueeMechanic>(this);
+	PolyLassoMechanic->Setup(ParentToolIn);
+	PolyLassoMechanic->SetIsEnabled(false);
+	PolyLassoMechanic->SpacingTolerance = 10.0f;
+	PolyLassoMechanic->OnDrawPolyLassoStarted.AddUObject(this, &UMeshTopologySelectionMechanic::OnPolyLassoStarted);
+	PolyLassoMechanic->OnDrawPolyLassoChanged.AddUObject(this, &UMeshTopologySelectionMechanic::OnPolyLassoChanged);
+	PolyLassoMechanic->OnDrawPolyLassoFinished.AddUObject(this, &UMeshTopologySelectionMechanic::OnPolyLassoFinished);
+	PolyLassoMechanic->SetBasePriority(BasePriority.MakeLower());
+
 	ClickOrDragBehavior = NewObject<USingleClickOrDragInputBehavior>();
 	ClickOrDragBehavior->Initialize(this, MarqueeMechanic);
 	ClickOrDragBehavior->Modifiers.RegisterModifier(ShiftModifierID, FInputDeviceState::IsShiftKeyDown);
@@ -72,22 +84,22 @@ void UMeshTopologySelectionMechanic::Setup(UInteractiveTool* ParentToolIn)
 		AddToolPropertySource(Properties);
 	}
 	Properties->WatchProperty(Properties->bSelectVertices, [this](bool bSelectVertices) {
-		UpdateMarqueeEnabled(); 
+		UpdateSelectorEnabled();
 	});
 	Properties->WatchProperty(Properties->bSelectEdges, [this](bool bSelectVertices) { 
-		UpdateMarqueeEnabled(); 
+		UpdateSelectorEnabled();
 	});
 	Properties->WatchProperty(Properties->bSelectFaces, [this](bool bSelectFaces) {
-		UpdateMarqueeEnabled();
+		UpdateSelectorEnabled();
 	});
 	Properties->WatchProperty(Properties->bSelectEdgeLoops, [this](bool bSelectEdgeLoops) {
-		UpdateMarqueeEnabled();
+		UpdateSelectorEnabled();
 	});
 	Properties->WatchProperty(Properties->bSelectEdgeRings, [this](bool bSelectEdgeRings) {
-		UpdateMarqueeEnabled();
+		UpdateSelectorEnabled();
 	});
-	Properties->WatchProperty(Properties->bEnableMarquee, [this](bool bEnableMarquee) { 
-		UpdateMarqueeEnabled(); 
+	Properties->WatchProperty(Properties->SelectionDragTool, [this](EMeshTopologySelectionMechanicDragTool DragTool) {
+		UpdateSelectorEnabled();
 	});
 
 	// set up visualizers
@@ -123,6 +135,7 @@ void UMeshTopologySelectionMechanic::DisableBehaviors(UInteractiveTool* ParentTo
 {
 	ParentToolIn->RemoveInputBehaviorsBySource(this);
 	ParentToolIn->RemoveInputBehaviorsBySource(MarqueeMechanic);
+	ParentToolIn->RemoveInputBehaviorsBySource(PolyLassoMechanic);
 
 	// TODO: Is it worth adding a way to remove the property watchers for marquee?
 }
@@ -131,7 +144,7 @@ void UMeshTopologySelectionMechanic::SetIsEnabled(bool bBehaviorEnabledIn, bool
 {
 	bIsEnabled = bBehaviorEnabledIn;
 	bRenderTopology = bRenderTopologyIn;
-	UpdateMarqueeEnabled();
+	UpdateSelectorEnabled();
 }
 
 void UMeshTopologySelectionMechanic::SetTransform(const FTransform3d& InTargetTransform)
@@ -139,9 +152,28 @@ void UMeshTopologySelectionMechanic::SetTransform(const FTransform3d& InTargetTr
 	TargetTransform = InTargetTransform;
 }
 
+PRAGMA_DISABLE_DEPRECATION_WARNINGS
 void UMeshTopologySelectionMechanic::SetMarqueeSelectionUpdateType(EMarqueeSelectionUpdateType InType)
 {
 	MarqueeSelectionUpdateType = InType;
+	switch (MarqueeSelectionUpdateType)
+	{
+	case EMarqueeSelectionUpdateType::OnDrag:
+		SetSelectionDragToolUpdateType(ESelectionDragToolUpdateType::OnDrag);
+		break;
+	case EMarqueeSelectionUpdateType::OnRelease:
+		SetSelectionDragToolUpdateType(ESelectionDragToolUpdateType::OnRelease);
+		break;
+	case EMarqueeSelectionUpdateType::OnTickAndRelease:
+		SetSelectionDragToolUpdateType(ESelectionDragToolUpdateType::OnTickAndRelease);
+		break;
+	}
+}
+PRAGMA_ENABLE_DEPRECATION_WARNINGS
+
+void UMeshTopologySelectionMechanic::SetSelectionDragToolUpdateType(ESelectionDragToolUpdateType InType)
+{
+	SelectionDragToolUpdateType = InType;
 }
 
 void UMeshTopologySelectionMechanic::SetBasePriority(const FInputCapturePriority &Priority)
@@ -155,9 +187,14 @@ void UMeshTopologySelectionMechanic::SetBasePriority(const FInputCapturePriority
 	{
 		HoverBehavior->SetDefaultPriority(Priority);
 	}
+	const FInputCapturePriority MarqueeAndLassoPriority = Priority.MakeLower();
 	if (MarqueeMechanic)
 	{
-		MarqueeMechanic->SetBasePriority(Priority.MakeLower());
+		MarqueeMechanic->SetBasePriority(MarqueeAndLassoPriority);
+	}
+	if (PolyLassoMechanic)
+	{
+		PolyLassoMechanic->SetBasePriority(MarqueeAndLassoPriority);
 	}
 }
 
@@ -177,6 +214,7 @@ void UMeshTopologySelectionMechanic::Render(IToolsContextRenderAPI* RenderAPI)
 	if (bIsEnabled)
 	{
 		MarqueeMechanic->Render(RenderAPI);
+		PolyLassoMechanic->Render(RenderAPI);
 	}
 	
 	if (!bRenderTopology) 
@@ -237,13 +275,14 @@ void UMeshTopologySelectionMechanic::DrawHUD(FCanvas* Canvas, IToolsContextRende
 	}
 	
 	MarqueeMechanic->DrawHUD(Canvas, RenderAPI);
+	PolyLassoMechanic->DrawHUD(Canvas, RenderAPI);
 }
 
 void UMeshTopologySelectionMechanic::Tick(float DeltaTime)
 {
 	Super::Tick(DeltaTime);
 
-	if (MarqueeSelectionUpdateType == EMarqueeSelectionUpdateType::OnTickAndRelease && PendingSelectionFunction)
+	if (SelectionDragToolUpdateType == ESelectionDragToolUpdateType::OnTickAndRelease && PendingSelectionFunction)
 	{
 		PendingSelectionFunction();
 		PendingSelectionFunction.Reset();
@@ -617,9 +656,12 @@ void UMeshTopologySelectionMechanic::OnUpdateModifierState(int ModifierID, bool
 	}
 }
 
-void UMeshTopologySelectionMechanic::OnDragRectangleStarted()
+//////////////////////////////////////////////////////////////////////////////////////////////////////
+
+void UMeshTopologySelectionMechanic::OnSelectionToolDragStarted()
 {
 	bCurrentlyMarqueeDragging = true;
+
 	PendingSelectionFunction.Reset();
 
 	ParentTool->GetToolManager()->BeginUndoTransaction(LOCTEXT("SelectionChange", "Selection"));
@@ -628,12 +670,42 @@ void UMeshTopologySelectionMechanic::OnDragRectangleStarted()
 	PreDragPersistentSelection = PersistentSelection;
 	LastUpdateRectangleSelection = PersistentSelection;
 	PreDragTopoSelectorSettings = GetTopoSelectorSettings(false);
-	PreDragTopoSelectorSettings.bIgnoreOcclusion = Properties->bMarqueeIgnoreOcclusion; // uses a separate setting for marquee
+
+	PreDragTopoSelectorSettings.bIgnoreOcclusion = Properties->bIgnoreOcclusionForDragTools;
+}
+
+void UMeshTopologySelectionMechanic::OnSelectionToolDragFinished()
+{
+	if (PendingSelectionFunction)
+	{
+		PendingSelectionFunction();
+		PendingSelectionFunction.Reset();
+	}
+
+	bCurrentlyMarqueeDragging = false;
+
+	TriIsOccludedCache.Reset();
+
+	if (PersistentSelection != PreDragPersistentSelection)
+	{
+		SelectionTimestamp++;
+		OnSelectionChanged.Broadcast();
+	}
+
+	EndChangeAndEmitIfModified();
+	ParentTool->GetToolManager()->EndUndoTransaction();
+}
+
+//////////////////////////////////////////////////////////////////////////////////////////////////////
+
+void UMeshTopologySelectionMechanic::OnDragRectangleStarted()
+{
+	OnSelectionToolDragStarted();
 }
 
 void UMeshTopologySelectionMechanic::OnDragRectangleChanged(const FCameraRectangle& CurrentRectangle)
 {
-	if (MarqueeSelectionUpdateType == EMarqueeSelectionUpdateType::OnDrag)
+	if (SelectionDragToolUpdateType == ESelectionDragToolUpdateType::OnDrag)
 	{
 		HandleRectangleChanged(CurrentRectangle);
 	}
@@ -649,33 +721,148 @@ void UMeshTopologySelectionMechanic::OnDragRectangleChanged(const FCameraRectang
 
 void UMeshTopologySelectionMechanic::OnDragRectangleFinished(const FCameraRectangle& Rectangle, bool bCancelled)
 {
-	if (PendingSelectionFunction)
+	OnSelectionToolDragFinished();
+}
+
+//////////////////////////////////////////////////////////////////////////////////////////////////////
+
+void UMeshTopologySelectionMechanic::OnPolyLassoStarted()
+{
+	OnSelectionToolDragStarted();
+}
+
+void UMeshTopologySelectionMechanic::OnPolyLassoChanged(const FCameraPolyLasso& InLasso)
+{
+	if (SelectionDragToolUpdateType == ESelectionDragToolUpdateType::OnDrag)
 	{
-		PendingSelectionFunction();
-		PendingSelectionFunction.Reset();
+		HandleLassoChanged(InLasso);
 	}
-	
-	bCurrentlyMarqueeDragging = false;
+	else
+	{
+		// Defer the selection on tick or release  
+		PendingSelectionFunction = [this, InLasso]()
+			{
+				HandleLassoChanged(InLasso);
+			};
+	}
+}
 
-	TriIsOccludedCache.Reset();
+namespace MeshTopologySelectionMechanic::Private
+{
 
-	if (PersistentSelection != PreDragPersistentSelection)
+	static FCameraPolyLasso MakeClosedLasso(const FCameraPolyLasso& InLasso)
 	{
-		SelectionTimestamp++;
-		OnSelectionChanged.Broadcast();
+		FCameraPolyLasso ClosedLasso{ InLasso };
+
+		if (ClosedLasso.Polyline.Num() > 0)
+		{
+			const FVector2D FirstPoint = ClosedLasso.Polyline[0];
+			if (FirstPoint != ClosedLasso.Polyline.Last())
+			{
+				ClosedLasso.Polyline.Add(FirstPoint);
+			}
+		}
+		return ClosedLasso;
 	}
+}
 
-	EndChangeAndEmitIfModified();
-	ParentTool->GetToolManager()->EndUndoTransaction();
+void UMeshTopologySelectionMechanic::OnPolyLassoFinished(const FCameraPolyLasso& InLasso, bool bCanceled)
+{
+	OnSelectionToolDragFinished();
+}
+
+void UMeshTopologySelectionMechanic::HandleLassoChanged(const FCameraPolyLasso& InLasso)
+{
+	if (!bIsEnabled)
+	{
+		return;
+	}
+
+	FCameraPolyLasso ClosedLasso = MeshTopologySelectionMechanic::Private::MakeClosedLasso(InLasso);
+	if (ClosedLasso.Polyline.Num() < 3)
+	{
+		// nothing to select really
+		return;
+	}
+
+	FGroupTopologySelection RectangleSelection;
+
+	TopoSelector->FindSelectedElement(PreDragTopoSelectorSettings, ClosedLasso, TargetTransform,
+		RectangleSelection, &TriIsOccludedCache);
+
+	if (ShouldAddToSelectionFunc())
+	{
+		PersistentSelection = PreDragPersistentSelection;
+		if (ShouldRemoveFromSelectionFunc())
+		{
+			PersistentSelection.Toggle(RectangleSelection);
+		}
+		else
+		{
+			PersistentSelection.Append(RectangleSelection);
+		}
+	}
+	else if (ShouldRemoveFromSelectionFunc())
+	{
+		PersistentSelection = PreDragPersistentSelection;
+		PersistentSelection.Remove(RectangleSelection);
+	}
+	else
+	{
+		// Neither key pressed.
+		PersistentSelection = RectangleSelection;
+	}
+
+	// If we modified the currently selected edges/vertices, they will be properly displayed in our
+	// Render() call. However, the mechanic is not responsible for face highlighting, so if we modified
+	// that, we need to notify the user so that they can update the highlighting (since OnSelectionChanged
+	// only gets broadcast at rectangle end).
+	if ((!PersistentSelection.SelectedGroupIDs.IsEmpty() || !LastUpdateRectangleSelection.SelectedGroupIDs.IsEmpty()) // if groups are involved
+		&& PersistentSelection != LastUpdateRectangleSelection)
+	{
+		LastUpdateRectangleSelection = PersistentSelection;
+		OnFaceSelectionPreviewChanged.Broadcast();
+	}
+}
+
+////////////////////////////////////////////////////////////////////////////////////////////////////////////////////
+
+void UMeshTopologySelectionMechanic::UpdateSelectorEnabled()
+{
+	UpdateMarqueeEnabled();
+	UpdateLassoEnabled();
 }
 
 void UMeshTopologySelectionMechanic::UpdateMarqueeEnabled()
 {
-	MarqueeMechanic->SetIsEnabled(
-		bIsEnabled 
-		&& Properties->bEnableMarquee
+	const bool bIsMarqueeEnabled = bIsEnabled
+		&& Properties->SelectionDragTool == EMeshTopologySelectionMechanicDragTool::Marquee
 		&& (Properties->bSelectVertices || Properties->bSelectEdges || Properties->bSelectFaces
-			|| Properties->bSelectEdgeLoops || Properties->bSelectEdgeRings));
+			|| Properties->bSelectEdgeLoops || Properties->bSelectEdgeRings);
+	if (bIsMarqueeEnabled != MarqueeMechanic->IsEnabled())
+	{
+		MarqueeMechanic->SetIsEnabled(bIsMarqueeEnabled);
+		if (bIsMarqueeEnabled)
+		{
+			ClickOrDragBehavior->SetDragTarget(MarqueeMechanic.Get());
+		}
+	}
+}
+
+void UMeshTopologySelectionMechanic::UpdateLassoEnabled()
+{
+	const bool bIsPolyLassoEnabled = bIsEnabled
+		&& Properties->SelectionDragTool == EMeshTopologySelectionMechanicDragTool::Lasso
+		&& (Properties->bSelectVertices || Properties->bSelectEdges || Properties->bSelectFaces
+			|| Properties->bSelectEdgeLoops || Properties->bSelectEdgeRings);
+	if (bIsPolyLassoEnabled != PolyLassoMechanic->IsEnabled())
+	{
+		PolyLassoMechanic->SetIsEnabled(bIsPolyLassoEnabled);
+		if (bIsPolyLassoEnabled)
+		{
+			ClickOrDragBehavior->SetDragTarget(PolyLassoMechanic.Get());
+		}
+	}
 }
 
 void UMeshTopologySelectionMechanic::RenderMarquee(IToolsContextRenderAPI* RenderAPI)
diff --git a/Engine/Plugins/Runtime/MeshModelingToolset/Source/ModelingComponents/Private/Selection/MeshTopologySelector.cpp b/Engine/Plugins/Runtime/MeshModelingToolset/Source/ModelingComponents/Private/Selection/MeshTopologySelector.cpp
--- a/Engine/Plugins/Runtime/MeshModelingToolset/Source/ModelingComponents/Private/Selection/MeshTopologySelector.cpp
+++ b/Engine/Plugins/Runtime/MeshModelingToolset/Source/ModelingComponents/Private/Selection/MeshTopologySelector.cpp
@@ -6,6 +6,7 @@
 #include "Spatial/SpatialInterfaces.h"
 #include "Engine/Polys.h"
 #include "Mechanics/RectangleMarqueeMechanic.h" // FCameraRectangle
+#include "Mechanics/PolyLassoMarqueeMechanic.h" // FCameraPolyLasso
 #include "ConvexVolume.h"
 
 using namespace UE::Geometry;
@@ -783,8 +784,207 @@ bool FMeshTopologySelector::FindSelectedElement(const FSelectionSettings& Settin
 		Spatial->DoTraversal(Traversal, Options);
 	}
 
-	return ResultOut.IsEmpty();
+	return !ResultOut.IsEmpty();
 }
 
+bool FMeshTopologySelector::FindSelectedElement(const FSelectionSettings& Settings,
+	const FCameraPolyLasso& CameraLasso, FTransform3d TargetTransform, FGroupTopologySelection& ResultOut,
+	TMap<int32, bool>* TriIsOccludedCache)
+{
+	// One minor easy thing we can do to speed up the below is to detect cases where transform does not have a non-uniform
+	// scale, and in those cases we transform the rectangle once rather than transforming all the query points/curves. 
+	// This isn't worth doing for non-uniform scale transforms because the rectangle basis stops being orthonormal and 
+	// becomes a problem to work with.
+
+	ResultOut.Clear();
+	FDynamicMeshAABBTree3* Spatial = GetSpatial();
+
+	// Needed for occlusion test, which happens in local space.
+	const FVector3d LocalCameraOrigin = TargetTransform.InverseTransformPosition((FVector3d)CameraLasso.CameraOrigin);
+
+	auto IsPointInsidePolyLine = [](const TArray<FVector2D>& InPolyLine, const FVector2D& Vertex2DPos)
+		{
+			const int32 NumPolyPoints = InPolyLine.Num();
+			if (NumPolyPoints > 0)
+			{
+				double WindingSum = 0;
+				FVector2D Prev = (InPolyLine[0] - Vertex2DPos);
+				FVector2D Curr = FVector2D::Zero();
+				for (int32 PointIndex = 1; PointIndex < NumPolyPoints; ++PointIndex)
+				{
+					Curr = (InPolyLine[PointIndex] - Vertex2DPos);
+					WindingSum += (double)FMathf::Atan2(Prev.X * Curr.Y - Prev.Y * Curr.X, Prev.X * Curr.X + Prev.Y * Curr.Y);
+					Prev = Curr;
+				}
+				WindingSum /= FMathd::TwoPi;
+				return FMathd::Abs(WindingSum) > 0.3;
+			}
+			return false;
+		};
+
+	auto VertexFilter = [&CameraLasso, &TargetTransform, &Settings, &LocalCameraOrigin, Spatial, &IsPointInsidePolyLine]
+	(const FVector3d& PointPosition)
+
+		{
+			const FVector TransformedPoint = TargetTransform.TransformPosition(PointPosition);
+			const FVector2D Vertex2DPos = CameraLasso.GetProjectedPoint(TransformedPoint);
+			const bool bIsInside = IsPointInsidePolyLine(CameraLasso.Polyline, Vertex2DPos);
+			return bIsInside
+				&& (Settings.bIgnoreOcclusion || !IsOccluded(PointPosition, LocalCameraOrigin,
+					Spatial, Settings.bHitBackFaces));
+		};
+
+	// Corner selection takes priority over edges.
+	if (Settings.bEnableCornerHits)
+	{
+		GeometrySet.ParallelFindAllPointsSatisfying(VertexFilter, ResultOut.SelectedCornerIDs);
+	}
+
+	auto SegmentIntersectPolyLine = [IsPointInsidePolyLine](const TArray<FVector2D>& InPolyLine, const FSegment2d& Segment2D, bool TestPoints = true)
+		{
+			if (IsPointInsidePolyLine(InPolyLine, Segment2D.StartPoint()))
+			{
+				return true;
+			}
+			if (IsPointInsidePolyLine(InPolyLine, Segment2D.EndPoint()))
+			{
+				return true;
+			}
+
+			const int32 NumPolyPoints = InPolyLine.Num();
+			if (NumPolyPoints > 1)
+			{
+				for (int32 PointIndex = 1; PointIndex < NumPolyPoints; ++PointIndex)
+				{
+					const FSegment2d PolyLineSegment(InPolyLine[PointIndex - 1], InPolyLine[PointIndex]);
+					if (PolyLineSegment.Intersects(Segment2D))
+					{
+						return true;
+					}
+				}
+			}
+			return false;
+		};
+
+	auto CurveFilter = [&CameraLasso, &TargetTransform, &Settings, &LocalCameraOrigin, Spatial, SegmentIntersectPolyLine](const FPolyline3d& Curve)
+		{
+			// Testing occlusion properly seems like it would be a pain, so we'll just consider something occluded
+			// if one of the endpoints is occluded. This will handle the common case of not wanting to select a
+			// hidden edge that is connected to a visible corner.
+			if (!Settings.bIgnoreOcclusion
+				&& (IsOccluded(Curve.Start(), LocalCameraOrigin, Spatial, Settings.bHitBackFaces)
+					|| IsOccluded(Curve.End(), LocalCameraOrigin, Spatial, Settings.bHitBackFaces)))
+			{
+				return false;
+			}
+
+			// Check whether any of the component segments intersect the rectangle
+			const TArray<FVector3d>& Verts = Curve.GetVertices();
+			FVector2D CurrentVert = CameraLasso.GetProjectedPoint(TargetTransform.TransformPosition(Verts[0]));
+			for (int32 VertIndex = 1; VertIndex < Verts.Num(); ++VertIndex)
+			{
+				const FVector2D NextVert = CameraLasso.GetProjectedPoint(TargetTransform.TransformPosition(Verts[VertIndex]));
+				const FSegment2d Segment(CurrentVert, NextVert);
+				if (SegmentIntersectPolyLine(CameraLasso.Polyline, Segment))
+				{
+					// We consider an edge non occluded if at least one point is not
+					// This is not perfect but it is a good approximation most most use case and still remain easy to evaluate
+					if (Settings.bIgnoreOcclusion ||
+						!IsOccluded(Verts[VertIndex - 1], LocalCameraOrigin, Spatial, Settings.bHitBackFaces) ||
+						!IsOccluded(Verts[VertIndex], LocalCameraOrigin, Spatial, Settings.bHitBackFaces)
+						)
+					{
+						return true;
+					}
+				}
+				CurrentVert = NextVert;
+			}
+			return false;
+		};
+
+	// If we didn't get corners, look for edges.
+	if (ResultOut.SelectedCornerIDs.IsEmpty() && Settings.bEnableEdgeHits)
+	{
+		GeometrySet.ParallelFindAllCurvesSatisfying(CurveFilter, ResultOut.SelectedEdgeIDs);
+	}
+
+	if (Settings.bEnableFaceHits && ResultOut.IsEmpty() && Mesh && TopologyProvider)
+	{
+		// check if the edges of triangle intersect with the lasso or if any points are inside of it 
+		auto TriangleIntersectPolyLine = [IsPointInsidePolyLine](const TArray<FVector2D>& InPolyLine, const FVector2D& Vtx0, const FVector2D& Vtx1, const FVector2D& Vtx2)
+			{
+				const int32 NumPolyPoints = InPolyLine.Num();
+				if (NumPolyPoints > 1)
+				{
+					const FSegment2d Segment0(Vtx0, Vtx1);
+					const FSegment2d Segment1(Vtx1, Vtx2);
+					const FSegment2d Segment2(Vtx2, Vtx0);
+
+					FVector WindingSums{ 0,0,0 };
+					for (int32 PointIndex = 1; PointIndex < NumPolyPoints; ++PointIndex)
+					{
+						const FVector2d PrevVert = InPolyLine[PointIndex - 1];
+						const FVector2d CurrVert = InPolyLine[PointIndex + 0];
+						const FSegment2d PolyLineSegment(PrevVert, CurrVert);
+						if (PolyLineSegment.Intersects(Segment0) ||
+							PolyLineSegment.Intersects(Segment1) ||
+							PolyLineSegment.Intersects(Segment2) )
+						{
+							return true;
+						}
+						const FVector2d Prev0 = PrevVert - Vtx0;
+						const FVector2d Curr0 = CurrVert - Vtx0;
+						WindingSums[0] += (double)FMathf::Atan2(Prev0.X * Curr0.Y - Prev0.Y * Curr0.X, Prev0.X * Curr0.X + Prev0.Y * Curr0.Y);
+						const FVector2d Prev1 = PrevVert - Vtx1;
+						const FVector2d Curr1 = CurrVert - Vtx1;
+						WindingSums[1] += (double)FMathf::Atan2(Prev1.X * Curr1.Y - Prev1.Y * Curr1.X, Prev1.X * Curr1.X + Prev1.Y * Curr1.Y);
+						const FVector2d Prev2 = PrevVert - Vtx2;
+						const FVector2d Curr2 = CurrVert - Vtx2;
+						WindingSums[2] += (double)FMathf::Atan2(Prev2.X * Curr2.Y - Prev2.Y * Curr2.X, Prev2.X * Curr2.X + Prev2.Y * Curr2.Y);
+					}
+					WindingSums /= FMathd::TwoPi;
+
+					return(
+						FMathd::Abs(WindingSums[0]) > 0.3 ||
+						FMathd::Abs(WindingSums[1]) > 0.3 ||
+						FMathd::Abs(WindingSums[2]) > 0.3
+						);
+				}
+				return false;
+			};
+
+		const int32 MaxTriangleID = Mesh->MaxTriangleID();
+
+		TArray<bool> SelectedTriangles;
+		SelectedTriangles.Init(false, MaxTriangleID);
+
+		ParallelFor(Mesh->MaxTriangleID(), 
+			[this, &SelectedTriangles/*, &VerticesInLasso*/, TriangleIntersectPolyLine, &CameraLasso, &TargetTransform](int32 TriId)
+			{
+				if (Mesh->IsTriangle(TriId))
+				{
+					const FIndex3i TriIndices = Mesh->GetTriangle(TriId);
+					const FVector2D Vtx0 = CameraLasso.GetProjectedPoint(TargetTransform.TransformPosition(Mesh->GetVertex(TriIndices[0])));
+					const FVector2D Vtx1 = CameraLasso.GetProjectedPoint(TargetTransform.TransformPosition(Mesh->GetVertex(TriIndices[1])));
+					const FVector2D Vtx2 = CameraLasso.GetProjectedPoint(TargetTransform.TransformPosition(Mesh->GetVertex(TriIndices[2])));
+					if (TriangleIntersectPolyLine(CameraLasso.Polyline, Vtx0, Vtx1, Vtx2))
+					{
+						SelectedTriangles[TriId] = true;
+					}
+				}
+			});
+
+		// Now add to the selection array 
+		for (int32 TriId = 0; TriId < MaxTriangleID; ++TriId)
+		{
+			if (SelectedTriangles[TriId])
+			{
+				ResultOut.SelectedGroupIDs.Add(TopologyProvider->GetGroupIDForTriangle(TriId));
+			}
+		}
+	}
+
+	return !ResultOut.IsEmpty();
+}
 
 #undef LOCTEXT_NAMESPACE
diff --git a/Engine/Plugins/Runtime/MeshModelingToolset/Source/ModelingComponents/Public/Selection/MeshTopologySelectionMechanic.h b/Engine/Plugins/Runtime/MeshModelingToolset/Source/ModelingComponents/Public/Selection/MeshTopologySelectionMechanic.h
--- a/Engine/Plugins/Runtime/MeshModelingToolset/Source/ModelingComponents/Public/Selection/MeshTopologySelectionMechanic.h
+++ b/Engine/Plugins/Runtime/MeshModelingToolset/Source/ModelingComponents/Public/Selection/MeshTopologySelectionMechanic.h
@@ -22,11 +22,23 @@ class UPersistentMeshSelection;
 class UMouseHoverBehavior;
 class UMeshTopologySelectionMechanic;
 class URectangleMarqueeMechanic;
+class UPolyLassoMarqueeMechanic;
 class USingleClickOrDragInputBehavior;
 
 using UE::Geometry::FDynamicMeshAABBTree3;
 PREDECLARE_USE_GEOMETRY_CLASS(FCompactMaps);
 
+/**
+* define the tool to use when dragging the mouse when selecting 
+*/
+UENUM()
+enum class EMeshTopologySelectionMechanicDragTool : uint8
+{
+	None,
+	Marquee,
+	Lasso
+};
+
 UCLASS(MinimalAPI)
 class UMeshTopologySelectionMechanicProperties : public UInteractiveToolPropertySet
 {
@@ -55,11 +67,21 @@ public:
 	bool bHitBackFaces = true;
 
 	UPROPERTY(EditAnywhere, Category = AdditionalSelectionOptions, AdvancedDisplay)
+	EMeshTopologySelectionMechanicDragTool SelectionDragTool = EMeshTopologySelectionMechanicDragTool::Marquee;
+
+	UPROPERTY(EditAnywhere, Category = AdditionalSelectionOptions, meta = (EditCondition = "SelectionDragTool != UMeshTopologySelectionMechanicDragTool::None", EditConditionHides))
+	bool bIgnoreOcclusionForDragTools = true;
+
+#if WITH_EDITORONLY_DATA
+	UE_DEPRECATED(5.8, "bEnableMarquee has been deprecated, use SelectionDragTool instead")
+	UPROPERTY(meta = (DeprecatedProperty, DeprecationMessage = "bEnableMarquee has been deprecated, use SelectionDragTool instead"))
 	bool bEnableMarquee = true;
 
 	/** Determines whether vertices should be checked for occlusion in marquee select (Note: marquee select currently only works with edges and vertices) */
-	UPROPERTY(EditAnywhere, Category = AdditionalSelectionOptions, meta = (EditCondition = "bEnableMarquee", EditConditionHides))
+	UE_DEPRECATED(5.8, "bMarqueeIgnoreOcclusion has been deprecated, use bIgnoreOcclusionForDragTools instead")
+	UPROPERTY(meta = (DeprecatedProperty, DeprecationMessage = "bMarqueeIgnoreOcclusion has been deprecated, use bIgnoreOcclusionForDragTools instead"))
 	bool bMarqueeIgnoreOcclusion = true;
+#endif
 
 	// The following were originally in their own category, all marked as AdvancedDisplay. However, since there wasn't a non-AdvancedDisplay
 	// property in the category, they started out as expanded and could not be collapsed.
@@ -108,13 +130,23 @@ protected:
  */
 
 UENUM()
-enum class EMarqueeSelectionUpdateType
+enum class UE_DEPRECATED(5.0, "EMarqueeSelectionUpdateType is deprecated, use ESelectionDragToolUpdateType instead") EMarqueeSelectionUpdateType
 {
 	OnDrag,
 	OnTickAndRelease,
 	OnRelease
 };
 
+/*
+ * Selection update type that defines when to update the selection when using a drag selection tool
+ */
+UENUM()
+enum class ESelectionDragToolUpdateType: uint8
+{
+	OnDrag,				// will update the selection for every mouse input events ( can be multiple time per tick )
+	OnTickAndRelease,	// will update the selection only once per tick and on the tool release
+	OnRelease			// will only update the selection on ethe tool release
+};
 
 /**
  * Base class mechanic for selecting a subset of mesh elements (edge loops, groups, corners, etc.)
@@ -177,7 +209,15 @@ public:
 	/**
 	 * Sets how/when the selection updates are handled.
 	 */
+	PRAGMA_DISABLE_DEPRECATION_WARNINGS
+	UE_DEPRECATED(5.8, TEXT("deprecated, use SetSelectionDragToolUpdateType instead"))
 	UE_API void SetMarqueeSelectionUpdateType(EMarqueeSelectionUpdateType InType);
+	PRAGMA_ENABLE_DEPRECATION_WARNINGS
+
+	/**
+	 * Sets how/when the selection updates are handled.
+	 */
+	UE_API void SetSelectionDragToolUpdateType(ESelectionDragToolUpdateType InType);
 
 	/**
 	 * Sets the base priority so that tools can make sure that their own behaviors are higher
@@ -356,12 +396,24 @@ public:
 	UE_API virtual void OnUpdateModifierState(int ModifierID, bool bIsOn) override;
 
 protected:
+	// common handling of the selection tools (marquee/lasso) 
+	void OnSelectionToolDragStarted();
+	void OnSelectionToolDragFinished();
+
 	// These get bound to marquee mechanic delegates.
 	UE_API virtual void OnDragRectangleStarted();
 	UE_API virtual void OnDragRectangleChanged(const FCameraRectangle& CurrentRectangle);
 	UE_API virtual void OnDragRectangleFinished(const FCameraRectangle& Rectangle, bool bCancelled);
 
+	// These get bound to lasso mechanic delegates.
+	void OnPolyLassoStarted();
+	void OnPolyLassoChanged(const FCameraPolyLasso& Lasso);
+	void OnPolyLassoFinished(const FCameraPolyLasso& Lasso, bool bCanceled);
+	
+	void UpdateSelectorEnabled();
+
 	UE_API virtual void UpdateMarqueeEnabled();
+	UE_API virtual void UpdateLassoEnabled();
 
 public:
 	/** 
@@ -422,15 +474,23 @@ protected:
 	UPROPERTY()
 	TObjectPtr<URectangleMarqueeMechanic> MarqueeMechanic;
 
+	UPROPERTY()
+	TObjectPtr<UPolyLassoMarqueeMechanic> PolyLassoMechanic;
+
+	PRAGMA_DISABLE_DEPRECATION_WARNINGS
+	UE_DEPRECATED(5.8, "MarqueeSelectionUpdateType is deprecated, use SelectionDragToolUpdateType (see also SetSelectionDragToolUpdateType)")
+	UPROPERTY(meta = (DeprecatedProperty, DeprecationMessage = "MarqueeSelectionUpdateType is deprecated, use SelectionDragToolUpdateType (see also SetSelectionDragToolUpdateType)"))
+	EMarqueeSelectionUpdateType MarqueeSelectionUpdateType = EMarqueeSelectionUpdateType::OnDrag;
+	PRAGMA_ENABLE_DEPRECATION_WARNINGS
 	/**
 	 * Selection update type (default is OnDrag) as it may not need to be triggered for every rectangle change
 	 * This can drastically improve the responsiveness of the UI for meshes high density meshes.
 	 * - OnDrag: calls HandleRectangleChanged when dragging
 	 * - OnTick: stores a PendingSelection function when dragging and calls it when ticking and on release (if any)
 	 * - OnRelease: stores a PendingSelection function when dragging and calls it on release (if any)
-	 */
+	*/
 	UPROPERTY()
-	EMarqueeSelectionUpdateType MarqueeSelectionUpdateType = EMarqueeSelectionUpdateType::OnDrag;
+	ESelectionDragToolUpdateType SelectionDragToolUpdateType = ESelectionDragToolUpdateType::OnDrag;
 	
 	FInputCapturePriority BasePriority = FInputCapturePriority(FInputCapturePriority::DEFAULT_TOOL_PRIORITY);
 
@@ -455,6 +515,8 @@ protected:
 	/** Calls actual selection using the input marquee rectangle. **/
 	UE_API void HandleRectangleChanged(const FCameraRectangle& InRectangle);
 
+	void HandleLassoChanged(const FCameraPolyLasso& InLasso);
+
 	/**
 	 * Get the topology selector settings to use given the current selection settings.
 	 * 
diff --git a/Engine/Plugins/Runtime/MeshModelingToolset/Source/ModelingComponents/Public/Selection/MeshTopologySelector.h b/Engine/Plugins/Runtime/MeshModelingToolset/Source/ModelingComponents/Public/Selection/MeshTopologySelector.h
--- a/Engine/Plugins/Runtime/MeshModelingToolset/Source/ModelingComponents/Public/Selection/MeshTopologySelector.h
+++ b/Engine/Plugins/Runtime/MeshModelingToolset/Source/ModelingComponents/Public/Selection/MeshTopologySelector.h
@@ -10,6 +10,7 @@
 #define UE_API MODELINGCOMPONENTS_API
 
 struct FCameraRectangle;
+struct FCameraPolyLasso;
 class FToolDataVisualizer;
 struct FViewCameraState;
 
@@ -184,6 +185,28 @@ public:
 		FTransform3d TargetTransform, FGroupTopologySelection& ResultOut, 
 		TMap<int32, bool>* TriIsOccludedCache = nullptr);
 
+	/**
+	 * Given a camera lasso (from a lasso selection) gives the corners/edges/groups that are contained in
+	 * the closed lasso curve. Only selects one of the types, and prefers corners to edges to groups.
+	 *
+	 * For group edges, occlusion is checked at endpoints, and if either is occluded, the edge is considered occluded.
+	 * For faces, occlusion is checked at triangle centroids.
+	 *
+	 * @param Settings Settings that determine what is selected and whether occlusion is tested.
+	 * @param CameraLasso The camera lasso .
+	 * @param TargetTransform Transform from points to world space, since the rectangle is in world space.
+	 * @param ResultOut Output selection (cleared before use).
+	 * @param TriIsOccludedCache Cached values of whether a given Tid has an occluded centroid, since occlusion
+	 *  values will stay the same if the rectangle is resized as long as the camera is not moved. This should
+	 *  be used when doing occlusion-based marquee face selection on meshes with many triangles, because the
+	 *  occlusion checks are expensive.
+	 *
+	 * @return true if something was selected.
+	 */
+	UE_API bool FindSelectedElement(const FSelectionSettings& Settings, const FCameraPolyLasso& CameraLasso, 
+		FTransform3d TargetTransform, FGroupTopologySelection& ResultOut,
+		TMap<int32, bool>* TriIsOccludedCache = nullptr);
+	
 	enum ECornerDrawStyle
 	{
 		Circle,
