Dataflow : remove TEXT() macro from deprecation message

#Synced-CL 50358557

diff --git a/Engine/Plugins/Runtime/MeshModelingToolset/Source/ModelingComponents/Public/Selection/MeshTopologySelectionMechanic.h b/Engine/Plugins/Runtime/MeshModelingToolset/Source/ModelingComponents/Public/Selection/MeshTopologySelectionMechanic.h
--- a/Engine/Plugins/Runtime/MeshModelingToolset/Source/ModelingComponents/Public/Selection/MeshTopologySelectionMechanic.h
+++ b/Engine/Plugins/Runtime/MeshModelingToolset/Source/ModelingComponents/Public/Selection/MeshTopologySelectionMechanic.h
@@ -210,7 +210,7 @@ public:
 	 * Sets how/when the selection updates are handled.
 	 */
 	PRAGMA_DISABLE_DEPRECATION_WARNINGS
-	UE_DEPRECATED(5.8, TEXT("deprecated, use SetSelectionDragToolUpdateType instead"))
+	UE_DEPRECATED(5.8, "deprecated, use SetSelectionDragToolUpdateType instead")
 	UE_API void SetMarqueeSelectionUpdateType(EMarqueeSelectionUpdateType InType);
 	PRAGMA_ENABLE_DEPRECATION_WARNINGS
 
