[GeometrySelection] Added basic support for soft selection

#Synced-CL 50430604

diff --git a/Engine/Plugins/Animation/SkeletalMeshModelingTools/Source/SkeletalMeshModelingTools/Private/Selection/SKMGeometrySelectionTransformTweaker.cpp b/Engine/Plugins/Animation/SkeletalMeshModelingTools/Source/SkeletalMeshModelingTools/Private/Selection/SKMGeometrySelectionTransformTweaker.cpp
--- a/Engine/Plugins/Animation/SkeletalMeshModelingTools/Source/SkeletalMeshModelingTools/Private/Selection/SKMGeometrySelectionTransformTweaker.cpp
+++ b/Engine/Plugins/Animation/SkeletalMeshModelingTools/Source/SkeletalMeshModelingTools/Private/Selection/SKMGeometrySelectionTransformTweaker.cpp
@@ -65,37 +65,11 @@ void USkeletalMeshGeometrySelectionTransformTweaker::TweakTransform(FVector& InD
 	CachedSelectionFrameTransform.SetRotation(InRot.Quaternion() * CachedSelectionFrameTransform.GetRotation());
 	Scale += InScale;
 
-	if (bInIsWorldSpace)
-	{
-		// For a world frame gizmo, the scaling needs to happen in world aligned gizmo space, but the 
-		// rotation is still encoded in the local gizmo frame change.
-		FQuat RotationToApply = CachedSelectionFrameTransform.GetRotation() * StartSelectionFrameTransform.GetRotation().Inverse();
-		SelectionManager->UpdateTransformation([this, RotationToApply](int32 VertexID, const FVector3d& InitialPosition, const FTransform3d& WorldTransform)
-			{
-				FVector3d PosLocal = InitialPosition;
-				FVector3d PosWorld = WorldTransform.TransformPosition(PosLocal);
-				FVector3d PosWorldGizmo = PosWorld - StartSelectionFrameTransform.GetLocation();
-
-				FVector3d NewPosWorld = RotationToApply * (Scale * PosWorldGizmo) + CachedSelectionFrameTransform.GetLocation();
-				FVector3d NewPosLocal = WorldTransform.InverseTransformPosition(NewPosWorld);
-				return NewPosLocal;
-			});	
-	}
-	else
-	{
-		SelectionManager->UpdateTransformation([this](int32 VertexID, const FVector3d& InitialPosition, const FTransform3d& WorldTransform)
-			{
-				// For a local gizmo, we just get the coordinates in the original frame, scale in that frame,
-				// then interpret them as coordinates in the new frame.
-				FVector3d PosLocal = InitialPosition;
-				FVector3d PosWorld = WorldTransform.TransformPosition(PosLocal);
-				FVector3d PosGizmo = StartSelectionFrameTransform.InverseTransformPosition(PosWorld);
-				PosGizmo = Scale * PosGizmo;
-				FVector3d NewPosWorld = CachedSelectionFrameTransform.TransformPosition(PosGizmo);
-				FVector3d NewPosLocal = WorldTransform.InverseTransformPosition(NewPosWorld);
-				return NewPosLocal;
-			});	
-	}
+	FQuat RotationToApply = CachedSelectionFrameTransform.GetRotation() * StartSelectionFrameTransform.GetRotation().Inverse();
+	FVector TranslationToApply = CachedSelectionFrameTransform.GetLocation() - StartSelectionFrameTransform.GetLocation();
+	FVector ScaleToApply = Scale - 1;
+
+	SelectionManager->UpdateTransformation(StartSelectionFrameTransform, TranslationToApply, RotationToApply, ScaleToApply, bInIsWorldSpace);
 	
 }
 
diff --git a/Engine/Plugins/Animation/SkeletalMeshModelingTools/Source/SkeletalMeshModelingTools/Private/Selection/SKMPreviewDynaMeshSelector.cpp b/Engine/Plugins/Animation/SkeletalMeshModelingTools/Source/SkeletalMeshModelingTools/Private/Selection/SKMPreviewDynaMeshSelector.cpp
--- a/Engine/Plugins/Animation/SkeletalMeshModelingTools/Source/SkeletalMeshModelingTools/Private/Selection/SKMPreviewDynaMeshSelector.cpp
+++ b/Engine/Plugins/Animation/SkeletalMeshModelingTools/Source/SkeletalMeshModelingTools/Private/Selection/SKMPreviewDynaMeshSelector.cpp
@@ -25,7 +25,7 @@ IGeometrySelectionTransformer* FSkeletalMeshPreviewDynamicMeshSelector::Initiali
 		ActiveTransformer = MakeShared<FBasicDynamicMeshSelectionTransformer>();
 	}
 	
-	ActiveTransformer->bEnableSelectionTransformDrawing = true;
+	ActiveTransformer->bEnableSelectionTransformDrawing = false;
 
 	ActiveTransformer->Initialize(this);
 	ActiveTransformer->OnEndTransformFunc = [this](IToolsContextTransactionsAPI* TransactionsAPI) 
diff --git a/Engine/Plugins/Animation/SkeletalMeshModelingTools/Source/SkeletalMeshModelingTools/Private/SkeletalMeshModelingToolsEditorMode.cpp b/Engine/Plugins/Animation/SkeletalMeshModelingTools/Source/SkeletalMeshModelingTools/Private/SkeletalMeshModelingToolsEditorMode.cpp
--- a/Engine/Plugins/Animation/SkeletalMeshModelingTools/Source/SkeletalMeshModelingTools/Private/SkeletalMeshModelingToolsEditorMode.cpp
+++ b/Engine/Plugins/Animation/SkeletalMeshModelingTools/Source/SkeletalMeshModelingTools/Private/SkeletalMeshModelingToolsEditorMode.cpp
@@ -23,6 +23,7 @@
 #include "ViewportInteractions/SKMModelingToolsGeometryClickSelection.h"
 #include "ViewportInteractions/SKMModelingToolsGeometryFrustumSelection.h"
 #include "ViewportInteractions/SKMModelingToolsGeometryHoverInteraction.h"
+#include "ViewportInteractions/SKMModelingToolsGeometrySoftSelectRadiusDrag.h"
 #include "ViewportInteractions/ViewportInteractionsBehaviorSource.h"
 #include "ViewportInteractions/ViewportInteractionsUtilities.h"
 #include "ViewportInteractions/ViewportMoveYawInteraction.h"
@@ -89,6 +90,7 @@
 #include "SkeletalMeshToolMenuContext.h"
 #include "Commands/ModifyGeometrySelectionCommand.h"
 #include "Selection/SKMPreviewDynaMeshSelector.h"
+#include "Widgets/Input/SSpinBox.h"
 
 #define LOCTEXT_NAMESPACE "SkeletalMeshModelingToolsEditorMode"
 
@@ -396,17 +398,22 @@ void USkeletalMeshModelingToolsEditorMode::Enter()
 			Source->AddInteraction<USkeletalMeshModelingToolsClickSelection>();
 				
 			check(WeakThis.IsValid());
-			USkeletalMeshModelingToolsGeometryFrustumSelection* GeometryFrustumSelection = Source->AddInteraction<
-				USkeletalMeshModelingToolsGeometryFrustumSelection>();
+			USkeletalMeshModelingToolsGeometryFrustumSelection* GeometryFrustumSelection =
+				Source->AddInteraction<USkeletalMeshModelingToolsGeometryFrustumSelection>();
 			GeometryFrustumSelection->BindSelectionManager(WeakThis->SelectionManager);
 
-			USkeletalMeshModelingToolsGeometryClickSelection* GeometryClickSelection = Source->AddInteraction<
-				USkeletalMeshModelingToolsGeometryClickSelection>();
+			USkeletalMeshModelingToolsGeometryClickSelection* GeometryClickSelection =
+				Source->AddInteraction<USkeletalMeshModelingToolsGeometryClickSelection>();
 			GeometryClickSelection->BindSelectionManager(WeakThis->SelectionManager);
 
-			USkeletalMeshModelingToolsGeometryHoverInteraction* GeometryHoverInteraction = Source->AddInteraction<
-				USkeletalMeshModelingToolsGeometryHoverInteraction>();
+			USkeletalMeshModelingToolsGeometryHoverInteraction* GeometryHoverInteraction =
+				Source->AddInteraction<USkeletalMeshModelingToolsGeometryHoverInteraction>();
 			GeometryHoverInteraction->BindSelectionManager(WeakThis->SelectionManager);
+
+			USkeletalMeshModelingToolsGeometrySoftSelectRadiusDragInteraction* GeometrySoftSelectDrag =
+				Source->AddInteraction<USkeletalMeshModelingToolsGeometrySoftSelectRadiusDragInteraction>();
+			GeometrySoftSelectDrag->BindSelectionManager(WeakThis->SelectionManager);
+				
 				
 			// Since viewport interactions are rebuilt whenever an editor mode is changed (including modes that can coexists like the persona skeleton selection edit mode)
 			// we have to check again to make sure the correct interactions are enabled when that happens
@@ -881,7 +888,7 @@ void USkeletalMeshModelingToolsEditorMode::InitializeGeometrySelectionSystems()
 	const UModelingToolsModeCustomizationSettings* ModelingEditorSettings = GetDefault<UModelingToolsModeCustomizationSettings>();
 	
 	// Colors initialized here any time Modeling mode is entered
-	SelectionManager->SetSelectionColors(ModelingEditorSettings->UnselectedColor, ModelingEditorSettings->HoverOverSelectedColor, ModelingEditorSettings->HoverOverUnselectedColor, ModelingEditorSettings->GeometrySelectedColor);
+	SelectionManager->SetSelectionColors(ModelingEditorSettings->UnselectedColor, ModelingEditorSettings->HoverOverSelectedColor, ModelingEditorSettings->HoverOverUnselectedColor, ModelingEditorSettings->GeometrySelectedColor, ModelingEditorSettings->GeometrySoftSelectedColor);
 
 	// This function registers and tracks an active UGeometrySelectionEditCommand and it's associated UICommand
 	auto RegisterSelectionCommand = [this, &CommandList](UGeometrySelectionEditCommand* Command, TSharedPtr<FUICommandInfo> UICommand, bool bAlwaysVisible)
@@ -962,6 +969,27 @@ void USkeletalMeshModelingToolsEditorMode::InitializeGeometrySelectionSystems()
 		}));
 	CommandList->MapAction(ToolManagerCommands.SelectionHitBackFaces, ToggleHitBackFacesAction);
 
+	// Register Soft Selection Toggle
+	FUIAction ToggleSoftSelectionAction = FUIAction(
+		FExecuteAction::CreateLambda([WeakThis]
+		{
+			if (!WeakThis.IsValid())
+			{
+				return;
+			}
+			bool bIsSoftSelectionEnabled = WeakThis->SelectionManager->IsSoftSelectionEnabled();
+			WeakThis->SelectionManager->SetSoftSelectionEnabled(!bIsSoftSelectionEnabled);
+		}),
+		FCanExecuteAction(),
+		FIsActionChecked::CreateLambda([WeakThis]()
+		{
+			if (!WeakThis.IsValid())
+			{
+				return false;
+			}
+			return WeakThis->SelectionManager->IsSoftSelectionEnabled();
+		}));
+	CommandList->MapAction(ToolManagerCommands.SelectionEnableSoftSelection, ToggleSoftSelectionAction);
 	
 	// set up selection type toggles
 	auto RegisterSelectionMode = [WeakThis](UGeometrySelectionManager::EMeshTopologyMode TopoMode, UE::Geometry::EGeometryElementType ElementMode, TSharedPtr<FUICommandInfo> UICommand)
@@ -1019,14 +1047,15 @@ void USkeletalMeshModelingToolsEditorMode::PopulateViewportToolbar()
 	using namespace SkeletalMeshModelingToolsEditorModeLocals;
 	FToolMenuOwnerScoped ScopeOwner(ViewportToolbarExtensionOwnerName);
 	UToolMenu* const Menu = UToolMenus::Get()->ExtendMenu("AnimationEditor.ViewportToolbar");
-	
+
+	TWeakObjectPtr<USkeletalMeshModelingToolsEditorMode> WeakThis(this);
 	
 	FToolMenuSection& MeshElementSelectionMenu = Menu->FindOrAddSection("Left");
 	FToolMenuEntry Entry = FToolMenuEntry::InitSubMenu(
 		"Mesh Element Selection",
 		LOCTEXT("MeshElementSelectionSubmenuLabel", "Mesh Element Selection"),
 		LOCTEXT("MeshElementSelectionSubmenuTooltip", "Mesh Element Selection settings in the viewport"),
-		FNewToolMenuDelegate::CreateLambda([ WeakCommandList = Toolkit->GetToolkitCommands().ToWeakPtr()](UToolMenu* Submenu) -> void
+		FNewToolMenuDelegate::CreateLambda([WeakThis, WeakCommandList = Toolkit->GetToolkitCommands().ToWeakPtr()](UToolMenu* Submenu) -> void
 		{
 			FToolMenuSection& MeshElementSelectionOptionsSection =
 			Submenu->FindOrAddSection("Element Selection", LOCTEXT("ElementSelectionLabel", "Element Selection"));
@@ -1097,6 +1126,49 @@ void USkeletalMeshModelingToolsEditorMode::PopulateViewportToolbar()
 								FModelingToolsManagerCommands::Get().SelectionHitBackFaces, CommandList));
 				}
 
+				// Section for Selection Filters
+				{
+					FToolMenuSection& SoftSelectionSection =
+						Submenu->FindOrAddSection("Soft Selection", LOCTEXT("SoftSelectionLabel", "Soft Selection"));
+					SoftSelectionSection.AddEntry(
+							FToolMenuEntry::InitMenuEntryWithCommandList(
+								FModelingToolsManagerCommands::Get().SelectionEnableSoftSelection, CommandList));
+
+
+					
+					TSharedRef<SWidget> SoftSelectionRadiusWidget =
+						SNew(SHorizontalBox)
+						+ SHorizontalBox::Slot()
+						.FillWidth(0.9f)
+						[
+							SNew(SSpinBox<double>)
+								.MinValue(0.0f)
+								.Value_Lambda(
+									[WeakThis]() -> double 
+									{
+										if (!WeakThis.IsValid())
+										{
+											return 0;
+										}
+										return WeakThis->SelectionManager->GetSoftSelectionRadius();
+									}
+								)
+								.OnValueChanged_Lambda(
+									[WeakThis](double InValue)
+									{
+										if (WeakThis.IsValid())
+										{
+											WeakThis->SelectionManager->SetSoftSelectionRadius(InValue);
+										}
+									}
+								)
+						]
+						+ SHorizontalBox::Slot()
+						.FillWidth(0.1f);
+
+						SoftSelectionSection.AddEntry(FToolMenuEntry::InitWidget(
+							"SoftSelectionRadius", SoftSelectionRadiusWidget, LOCTEXT("SoftSelectionRadiusLabel", "Radius")));
+				}
 			}
 		})
 	);
@@ -1142,6 +1214,7 @@ void USkeletalMeshModelingToolsEditorMode::ToggleGeometrySelectionViewportIntera
 		BehaviorSource->FindInteraction(USkeletalMeshModelingToolsGeometryFrustumSelection::StaticClass())->SetEnabled(bEnable);
 		BehaviorSource->FindInteraction(USkeletalMeshModelingToolsGeometryClickSelection::StaticClass())->SetEnabled(bEnable);
 		BehaviorSource->FindInteraction(USkeletalMeshModelingToolsGeometryHoverInteraction::StaticClass())->SetEnabled(bEnable);
+		BehaviorSource->FindInteraction(USkeletalMeshModelingToolsGeometrySoftSelectRadiusDragInteraction::StaticClass())->SetEnabled(bEnable);
 				
 		BehaviorSource->FindInteraction(USkeletalMeshModelingToolsClickSelection::StaticClass())->SetEnabled(!bEnable);
 		BehaviorSource->FindInteraction(UViewportMoveYawInteraction::StaticClass())->SetEnabled(!bEnable);	
diff --git a/Engine/Plugins/Animation/SkeletalMeshModelingTools/Source/SkeletalMeshModelingTools/Private/SkeletalMeshModelingToolsEditorModeToolkit.cpp b/Engine/Plugins/Animation/SkeletalMeshModelingTools/Source/SkeletalMeshModelingTools/Private/SkeletalMeshModelingToolsEditorModeToolkit.cpp
--- a/Engine/Plugins/Animation/SkeletalMeshModelingTools/Source/SkeletalMeshModelingTools/Private/SkeletalMeshModelingToolsEditorModeToolkit.cpp
+++ b/Engine/Plugins/Animation/SkeletalMeshModelingTools/Source/SkeletalMeshModelingTools/Private/SkeletalMeshModelingToolsEditorModeToolkit.cpp
@@ -11,6 +11,7 @@
 #include "InteractiveTool.h"
 #include "InteractiveToolsContext.h"
 #include "ISkeletalMeshEditor.h"
+#include "ModelingToolsEditorModeSettings.h"
 #include "ModelingToolsEditorModeStyle.h"
 #include "Modules/ModuleManager.h"
 #include "PropertyEditorModule.h"
@@ -38,6 +39,7 @@
 #include "Widgets/Layout/SExpandableArea.h"
 #include "Widgets/Docking/SDockTab.h"
 #include "ScopedTransaction.h"
+#include "Selection/GeometrySelectionManager.h"
 
 
 #define LOCTEXT_NAMESPACE "SkeletalMeshModelingToolsEditorModeToolkit"
@@ -49,6 +51,11 @@ FSkeletalMeshModelingToolsEditorModeToolkit::~FSkeletalMeshModelingToolsEditorMo
 		Context->OnToolNotificationMessage.RemoveAll(this);
 		Context->OnToolWarningMessage.RemoveAll(this);
 	}
+
+	if (UModelingToolsModeCustomizationSettings* ModelingCustomizationSettings = GetMutableDefault<UModelingToolsModeCustomizationSettings>())
+	{
+		ModelingCustomizationSettings->OnSettingChanged().RemoveAll(this);
+	}
 }
 
 void FSkeletalMeshModelingToolsEditorModeToolkit::Init(
@@ -115,6 +122,11 @@ void FSkeletalMeshModelingToolsEditorModeToolkit::Init(
 	
 	EditorMode = CastChecked<USkeletalMeshModelingToolsEditorMode>(InOwningMode.Get());
 	EditorMode->OnInitialized().AddSP(this, &FSkeletalMeshModelingToolsEditorModeToolkit::OnEditorModeInitialized);
+
+	UModelingToolsModeCustomizationSettings* ModelingCustomizationSettings = GetMutableDefault<UModelingToolsModeCustomizationSettings>();
+	ModelingCustomizationSettings->OnSettingChanged().AddSP(SharedThis(this),  &FSkeletalMeshModelingToolsEditorModeToolkit::UpdateSelectionColors);
+
+
 }
 
 FName FSkeletalMeshModelingToolsEditorModeToolkit::GetToolkitFName() const
@@ -602,6 +614,15 @@ void FSkeletalMeshModelingToolsEditorModeToolkit::OnEditorModeInitialized()
 
 }
 
+void FSkeletalMeshModelingToolsEditorModeToolkit::UpdateSelectionColors(UObject* Object, FPropertyChangedEvent& PropertyChangedEvent) const
+{
+	if (EditorMode->GetSelectionManager())
+	{
+		const UModelingToolsModeCustomizationSettings* UISettings = GetMutableDefault<UModelingToolsModeCustomizationSettings>();
+		EditorMode->GetSelectionManager()->SetSelectionColors(UISettings->UnselectedColor, UISettings->HoverOverSelectedColor, UISettings->HoverOverUnselectedColor, UISettings->GeometrySelectedColor, UISettings->GeometrySoftSelectedColor);
+	}
+}
+
 void FSkeletalMeshModelingToolsEditorModeToolkit::ShowDynamicMeshSkeletonTree()
 {
 	if (ModeUILayer.IsValid())
diff --git a/Engine/Plugins/Animation/SkeletalMeshModelingTools/Source/SkeletalMeshModelingTools/Private/SkeletalMeshModelingToolsEditorModeToolkit.h b/Engine/Plugins/Animation/SkeletalMeshModelingTools/Source/SkeletalMeshModelingTools/Private/SkeletalMeshModelingToolsEditorModeToolkit.h
--- a/Engine/Plugins/Animation/SkeletalMeshModelingTools/Source/SkeletalMeshModelingTools/Private/SkeletalMeshModelingToolsEditorModeToolkit.h
+++ b/Engine/Plugins/Animation/SkeletalMeshModelingTools/Source/SkeletalMeshModelingTools/Private/SkeletalMeshModelingToolsEditorModeToolkit.h
@@ -21,7 +21,8 @@ class FSkeletalMeshModelingToolsEditorModeToolkit : public FModeToolkit
 public:
 
 	virtual ~FSkeletalMeshModelingToolsEditorModeToolkit() override;
-	
+
+
 	// IToolkit overrides
 	virtual void Init(const TSharedPtr<IToolkitHost>& InToolkitHost, TWeakObjectPtr<UEdMode> InOwningMode) override;
 	virtual FName GetToolkitFName() const override;
@@ -60,7 +61,7 @@ private:
 
 
 	void OnEditorModeInitialized();
-
+	void UpdateSelectionColors(UObject* Object, FPropertyChangedEvent& PropertyChangedEvent) const;
 
 	TSharedPtr<SWidget> MakeFooterWidget();
 	bool IsApplyButtonEnabled() const;
diff --git a/Engine/Plugins/Animation/SkeletalMeshModelingTools/Source/SkeletalMeshModelingTools/Private/ViewportInteractions/SKMModelingToolsGeometrySoftSelectRadiusDrag.cpp b/Engine/Plugins/Animation/SkeletalMeshModelingTools/Source/SkeletalMeshModelingTools/Private/ViewportInteractions/SKMModelingToolsGeometrySoftSelectRadiusDrag.cpp
new file mode 100644
--- /dev/null
+++ b/Engine/Plugins/Animation/SkeletalMeshModelingTools/Source/SkeletalMeshModelingTools/Private/ViewportInteractions/SKMModelingToolsGeometrySoftSelectRadiusDrag.cpp
@@ -0,0 +1,169 @@
+// Copyright Epic Games, Inc. All Rights Reserved.
+
+#include "ViewportInteractions/SKMModelingToolsGeometrySoftSelectRadiusDrag.h"
+
+#include "CanvasItem.h"
+#include "CanvasTypes.h"
+#include "EditorViewportClient.h"
+#include "MeshAdapterTransforms.h"
+#include "SceneView.h"
+#include "EditorDragTools/EditorViewportClientProxy.h"
+#include "Selection/GeometrySelectionManager.h"
+#include "ViewportInteractions/ViewportInteractionsBehaviorSource.h"
+
+
+USkeletalMeshModelingToolsGeometrySoftSelectRadiusDragInteraction::USkeletalMeshModelingToolsGeometrySoftSelectRadiusDragInteraction()
+{
+	InteractionName = TEXT("Skeletal Mesh Modeling Tools Geometry Soft Select Radius Drag Interaction");
+	
+	UViewportClickDragBehavior* MouseBehavior = NewObject<UViewportClickDragBehavior>();
+	MouseBehavior->Initialize(this);
+	MouseBehavior->SetBindings({
+		UE::Editor::ViewportInteractions::FButtonBinding(EKeys::LeftMouseButton).TriggersStart(),
+		UE::Editor::ViewportInteractions::FButtonBinding(EKeys::B),
+	});
+
+	RegisterInputBehavior(MouseBehavior);
+}
+
+void USkeletalMeshModelingToolsGeometrySoftSelectRadiusDragInteraction::BindSelectionManager(UGeometrySelectionManager* InSelectionManager)
+{
+	SelectionManager = InSelectionManager;
+}
+
+void USkeletalMeshModelingToolsGeometrySoftSelectRadiusDragInteraction::Render(
+	const FSceneView* View,
+	FCanvas* Canvas,
+    EViewInteractionState InteractionState
+    )
+{
+	if (!IsDragging() || InteractionState != EViewInteractionState::Focused)
+	{
+		return;
+	}
+
+	if (!Canvas || !View)
+	{
+		return;
+	}
+
+	if (!SelectionManager->HasSelection())
+	{
+		return;
+	}
+	
+	UE::Geometry::FFrame3d SelectionFrame; 
+	SelectionManager->GetSelectionWorldFrame(SelectionFrame);
+
+	double Radius = SelectionManager->GetSoftSelectionRadius();
+	
+	FVector2D OriginPixel; 
+	View->WorldToPixel(SelectionFrame.Origin, OriginPixel);
+	OriginPixel /= Canvas->GetDPIScale();
+
+	FVector RadiusSamplePoint = SelectionFrame.Origin + View->GetViewUp() * Radius;
+	FVector2D RadiusSamplePointPixel;
+	View->WorldToPixel(RadiusSamplePoint, RadiusSamplePointPixel);
+	RadiusSamplePointPixel /= Canvas->GetDPIScale();
+	
+	double PixelRadius = FMath::Abs(RadiusSamplePointPixel.Y - OriginPixel.Y);
+
+
+	const int32 NumSides = FMath::Max(20, (int)PixelRadius / 25);
+	const double AngleDelta = 2.0 * UE_PI / NumSides;
+	const FVector2D AxisX(1., 0.);
+	const FVector2D AxisY(0., -1.);
+	FVector2D LastVertex = OriginPixel + AxisX * PixelRadius;
+
+	for (int32 SideIndex = 0; SideIndex < NumSides; SideIndex++)
+	{
+		const double CurAngle = AngleDelta * (SideIndex + 1);
+		const FVector2D Vertex = OriginPixel + (AxisX * FMath::Cos(CurAngle) + AxisY * FMath::Sin(CurAngle)) * PixelRadius;
+
+		FCanvasLineItem LineItem{ FVector2D(LastVertex), FVector2D(Vertex) };
+		LineItem.SetColor(FLinearColor::Green);
+		Canvas->DrawItem(LineItem);
+
+		LastVertex = Vertex;
+	}
+
+}
+
+FInputRayHit USkeletalMeshModelingToolsGeometrySoftSelectRadiusDragInteraction::CanBeginClickDragSequence(const FInputDeviceRay& PressPos)
+{
+	if (IsEnabled())
+	{
+		if (SelectionManager->HasSelection())
+		{
+			return FInputRayHit(TNumericLimits<float>::Max()); // bHit is true.
+		}
+	}
+
+	return FInputRayHit();
+}
+
+void USkeletalMeshModelingToolsGeometrySoftSelectRadiusDragInteraction::OnDragStart(const FInputDeviceRay& InPressPos)
+{
+	if (UViewportInteractionsBehaviorSource* BehaviorSource = GetViewportInteractionsBehaviorSource())
+	{
+		BehaviorSource->SetMouseCursorOverride(EMouseCursor::Default);
+	}
+
+	Start = FVector(InPressPos.ScreenPosition.X, InPressPos.ScreenPosition.Y, 0);
+	End = Start;
+	
+	SelectionManager->SetSoftSelectionEnabled(true);
+	StartingRadius = SelectionManager->GetSoftSelectionRadius();
+
+	bIsDragging = true;
+	
+	// Signal that this tool is now active
+	OnActivateTool().Broadcast();
+
+}
+
+void USkeletalMeshModelingToolsGeometrySoftSelectRadiusDragInteraction::OnDrag(const FInputDeviceRay& InDragPos)
+{
+	End = FVector(InDragPos.ScreenPosition.X, InDragPos.ScreenPosition.Y, 0);
+	static double Multiplier = 1;
+	FEditorViewportClient* EditorViewportClient = GetEditorViewportClient();
+
+	FSceneViewFamilyContext ViewFamily(
+		FSceneViewFamily::ConstructionValues(EditorViewportClient->Viewport, EditorViewportClient->GetScene(), EditorViewportClient->EngineShowFlags)
+		.SetRealtimeUpdate(EditorViewportClient->IsRealtime()));
+
+	FSceneView* View = EditorViewportClient->CalcSceneView(&ViewFamily);
+
+	UE::Geometry::FFrame3d SelectionFrame; 
+	SelectionManager->GetSelectionWorldFrame(SelectionFrame);
+
+	double SampleRadius = 10;
+	
+	FVector2D OriginPixel; 
+	View->WorldToPixel(SelectionFrame.Origin, OriginPixel);
+
+	FVector RadiusSamplePoint = SelectionFrame.Origin + View->GetViewUp() * SampleRadius;
+	FVector2D RadiusSamplePointPixel;
+	View->WorldToPixel(RadiusSamplePoint, RadiusSamplePointPixel);
+	
+	const double PixelRadius = FMath::Abs(RadiusSamplePointPixel.Y - OriginPixel.Y);
+
+	double CentimeterPerPixel = SampleRadius / PixelRadius;
+	
+	SelectionManager->SetSoftSelectionRadius(StartingRadius + (End.X - Start.X) * CentimeterPerPixel );
+}
+
+void USkeletalMeshModelingToolsGeometrySoftSelectRadiusDragInteraction::OnDragEnd(const FInputDeviceRay& InReleasePos)
+{
+	if (SelectionManager->GetSoftSelectionRadius() == 0)
+	{
+		SelectionManager->SetSoftSelectionEnabled(false);
+	}
+	
+	Super::OnDragEnd(InReleasePos);
+}
+
+void USkeletalMeshModelingToolsGeometrySoftSelectRadiusDragInteraction::OnTerminateDragSequence()
+{
+	Super::OnTerminateDragSequence();
+}
diff --git a/Engine/Plugins/Animation/SkeletalMeshModelingTools/Source/SkeletalMeshModelingTools/Private/ViewportInteractions/SKMModelingToolsGeometrySoftSelectRadiusDrag.h b/Engine/Plugins/Animation/SkeletalMeshModelingTools/Source/SkeletalMeshModelingTools/Private/ViewportInteractions/SKMModelingToolsGeometrySoftSelectRadiusDrag.h
new file mode 100644
--- /dev/null
+++ b/Engine/Plugins/Animation/SkeletalMeshModelingTools/Source/SkeletalMeshModelingTools/Private/ViewportInteractions/SKMModelingToolsGeometrySoftSelectRadiusDrag.h
@@ -0,0 +1,40 @@
+// Copyright Epic Games, Inc. All Rights Reserved.
+
+#pragma once
+
+#include "EditorDragTools/DragToolInteraction.h"
+
+#include "SKMModelingToolsGeometrySoftSelectRadiusDrag.generated.h"
+
+class UGeometrySelectionManager;
+
+UCLASS(MinimalAPI, Transient)
+class USkeletalMeshModelingToolsGeometrySoftSelectRadiusDragInteraction
+	: public UDragToolInteraction
+{
+	GENERATED_BODY()
+
+public:
+	USkeletalMeshModelingToolsGeometrySoftSelectRadiusDragInteraction();
+
+	void BindSelectionManager(UGeometrySelectionManager* InSelectionManager);
+	
+	virtual void Render(const FSceneView* View, FCanvas* Canvas, EViewInteractionState InteractionState) override;
+	
+	//~ Begin IViewportClickDragBehaviorTarget
+	virtual FInputRayHit CanBeginClickDragSequence(const FInputDeviceRay& InPressPos) override;
+	virtual void OnDragStart(const FInputDeviceRay& InPressPos) override;
+	virtual void OnDrag(const FInputDeviceRay& InDragPos) override;
+	virtual void OnDragEnd(const FInputDeviceRay& InReleasePos) override;
+	virtual void OnTerminateDragSequence() override;
+	//~ End IViewportClickDragBehaviorTarget
+
+protected:
+	virtual TArray<FEditorModeID> GetUnsupportedModes() const override { return {}; }
+	
+protected:
+	double StartingRadius = 0;
+	
+	TWeakObjectPtr<UGeometrySelectionManager> SelectionManager;
+};
+
diff --git a/Engine/Plugins/Animation/SkeletalMeshModelingTools/Source/SkeletalMeshModelingTools/SkeletalMeshModelingTools.Build.cs b/Engine/Plugins/Animation/SkeletalMeshModelingTools/Source/SkeletalMeshModelingTools/SkeletalMeshModelingTools.Build.cs
--- a/Engine/Plugins/Animation/SkeletalMeshModelingTools/Source/SkeletalMeshModelingTools/SkeletalMeshModelingTools.Build.cs
+++ b/Engine/Plugins/Animation/SkeletalMeshModelingTools/Source/SkeletalMeshModelingTools/SkeletalMeshModelingTools.Build.cs
@@ -44,7 +44,8 @@ public class SkeletalMeshModelingTools : ModuleRules
 				"MeshConversion",
 				"InterchangeEngine",
 				"GeometryFramework",
-				"EditorWidgets" 
+				"EditorWidgets",
+				"DeveloperSettings"
 			}
 		);
 
diff --git a/Engine/Plugins/Editor/ModelingToolsEditorMode/Source/ModelingToolsEditorMode/Private/ModelingSelectionInteraction.cpp b/Engine/Plugins/Editor/ModelingToolsEditorMode/Source/ModelingToolsEditorMode/Private/ModelingSelectionInteraction.cpp
--- a/Engine/Plugins/Editor/ModelingToolsEditorMode/Source/ModelingToolsEditorMode/Private/ModelingSelectionInteraction.cpp
+++ b/Engine/Plugins/Editor/ModelingToolsEditorMode/Source/ModelingToolsEditorMode/Private/ModelingSelectionInteraction.cpp
@@ -442,39 +442,9 @@ void UModelingSelectionInteraction::ApplyPendingTransformInteractions()
 		LastRotateDelta = (FVector4d)RotateDelta;
 		LastScaleDelta = CurScaleDelta;
 
-		if (bTransformInWorldFrame)
-		{
-			// For a world frame gizmo, the scaling needs to happen in world aligned gizmo space, but the 
-			// rotation is still encoded in the local gizmo frame change.
-			FQuaterniond RotationToApply = CurFrame.Rotation * InitialGizmoFrame.Rotation.Inverse();
-			SelectionManager->UpdateTransformation( [&](int32 VertexID, const FVector3d& InitialPosition, const FTransform3d& WorldTransform)
-			{
-				FVector3d LocalTranslation = WorldTransform.InverseTransformVector(TranslationDelta);
+		FQuat RotationDelta = FQuat(CurFrame.Rotation * InitialGizmoFrame.Rotation.Inverse());
 
-				FVector3d PosLocal = InitialPosition;
-				FVector3d PosWorld = WorldTransform.TransformPosition(PosLocal);
-				FVector3d PosWorldGizmo = PosWorld - InitialGizmoFrame.Origin;
-
-				FVector3d NewPosWorld = RotationToApply * (PosWorldGizmo * CurScale) + CurFrame.Origin;
-				FVector3d NewPosLocal = WorldTransform.InverseTransformPosition(NewPosWorld);
-				return NewPosLocal;
-			});
-		}
-		else
-		{
-			SelectionManager->UpdateTransformation( [&](int32 VertexID, const FVector3d& InitialPosition, const FTransform3d& WorldTransform)
-			{
-				// For a local gizmo, we just get the coordinates in the original frame, scale in that frame,
-				// then interpret them as coordinates in the new frame.
-				FVector3d PosLocal = InitialPosition;
-				FVector3d PosWorld = WorldTransform.TransformPosition(PosLocal);
-				FVector3d PosGizmo = InitialGizmoFrame.ToFramePoint(PosWorld);
-				PosGizmo = CurScale * PosGizmo;
-				FVector3d NewPosWorld = CurFrame.FromFramePoint(PosGizmo);
-				FVector3d NewPosLocal = WorldTransform.InverseTransformPosition(NewPosWorld);
-				return NewPosLocal;
-			});
-		}
+		SelectionManager->UpdateTransformation(InitialGizmoFrame.ToFTransform(), TranslationDelta, RotationDelta, CurScaleDelta, bTransformInWorldFrame);
 	}
 
 	bGizmoUpdatePending = false;
diff --git a/Engine/Plugins/Editor/ModelingToolsEditorMode/Source/ModelingToolsEditorMode/Private/ModelingToolsEditorMode.cpp b/Engine/Plugins/Editor/ModelingToolsEditorMode/Source/ModelingToolsEditorMode/Private/ModelingToolsEditorMode.cpp
--- a/Engine/Plugins/Editor/ModelingToolsEditorMode/Source/ModelingToolsEditorMode/Private/ModelingToolsEditorMode.cpp
+++ b/Engine/Plugins/Editor/ModelingToolsEditorMode/Source/ModelingToolsEditorMode/Private/ModelingToolsEditorMode.cpp
@@ -145,6 +145,7 @@
 #include "ViewportToolbar/UnrealEdViewportToolbar.h"
 #include "ToolMenus.h"
 #include "Framework/Docking/LayoutExtender.h"
+#include "Widgets/Input/SSpinBox.h"
 
 #include UE_INLINE_GENERATED_CPP_BY_NAME(ModelingToolsEditorMode)
 
@@ -948,7 +949,7 @@ void UModelingToolsEditorMode::Enter()
 		ToggleFromGeometryAction( EModelingSelectionInteraction_LocalFrameMode::FromObject));
 
 	// Colors initialized here any time Modeling mode is entered
-	SelectionManager->SetSelectionColors(ModelingEditorSettings->UnselectedColor, ModelingEditorSettings->HoverOverSelectedColor, ModelingEditorSettings->HoverOverUnselectedColor, ModelingEditorSettings->GeometrySelectedColor);
+	SelectionManager->SetSelectionColors(ModelingEditorSettings->UnselectedColor, ModelingEditorSettings->HoverOverSelectedColor, ModelingEditorSettings->HoverOverUnselectedColor, ModelingEditorSettings->GeometrySelectedColor, ModelingEditorSettings->GeometrySoftSelectedColor);
 
 	// this function registers and tracks an active UGeometrySelectionEditCommand and it's associated UICommand
 	auto RegisterSelectionCommand = [this, &CommandList](UGeometrySelectionEditCommand* Command, TSharedPtr<FUICommandInfo> UICommand, bool bAlwaysVisible)
@@ -1010,6 +1011,20 @@ void UModelingToolsEditorMode::Enter()
 		}));
 	CommandList->MapAction(ToolManagerCommands.SelectionHitBackFaces, ToggleHitBackFacesAction);
 
+	// Register Soft Selection Toggle
+	FUIAction ToggleSoftSelectionAction = FUIAction(
+		FExecuteAction::CreateLambda([this]
+		{
+			bool bIsSoftSelectionEnabled = SelectionManager->IsSoftSelectionEnabled();
+			SelectionManager->SetSoftSelectionEnabled(!bIsSoftSelectionEnabled);
+		}),
+		FCanExecuteAction(),
+		FIsActionChecked::CreateLambda([this]()
+		{
+			return SelectionManager->IsSoftSelectionEnabled();
+		}));
+	CommandList->MapAction(ToolManagerCommands.SelectionEnableSoftSelection, ToggleSoftSelectionAction);
+	
 	// register extensions
 	TArray<IModelingModeToolExtension*> Extensions = IModularFeatures::Get().GetModularFeatureImplementations<IModelingModeToolExtension>(
 		IModelingModeToolExtension::GetModularFeatureName());
@@ -1445,14 +1460,15 @@ void UModelingToolsEditorMode::OnBlueprintPreCompile(UBlueprint* Blueprint)
 
 	FToolMenuOwnerScoped ScopeOwner(ModelingModeOwnerName);
 	UToolMenu* const Menu = UToolMenus::Get()->ExtendMenu(InMenuName);
-	
+
+	TWeakObjectPtr<UModelingToolsEditorMode> WeakThis(this);
 	
 	FToolMenuSection& MeshElementSelectionMenu = Menu->FindOrAddSection("Left");
 	FToolMenuEntry Entry = FToolMenuEntry::InitSubMenu(
 		"Mesh Element Selection",
 		LOCTEXT("MeshElementSelectionSubmenuLabel", "Mesh Element Selection"),
 		LOCTEXT("MeshElementSelectionSubmenuTooltip", "Mesh Element Selection settings in the viewport"),
-		FNewToolMenuDelegate::CreateLambda([ WeakCommandList = InCommandList.ToWeakPtr()](UToolMenu* Submenu) -> void
+		FNewToolMenuDelegate::CreateLambda([WeakThis, WeakCommandList = InCommandList.ToWeakPtr()](UToolMenu* Submenu) -> void
 		{
 			FToolMenuSection& MeshElementSelectionOptionsSection =
 			Submenu->FindOrAddSection("Element Selection", LOCTEXT("ElementSelectionLabel", "Element Selection"));
@@ -1546,6 +1562,50 @@ void UModelingToolsEditorMode::OnBlueprintPreCompile(UBlueprint* Blueprint)
 								FModelingToolsManagerCommands::Get().SelectionHitBackFaces, CommandList));
 				}
 
+				// Section for SoftSelection
+				{
+					FToolMenuSection& SoftSelectionSection =
+						Submenu->FindOrAddSection("Soft Selection", LOCTEXT("SoftSelectionLabel", "Soft Selection"));
+					SoftSelectionSection.AddEntry(
+							FToolMenuEntry::InitMenuEntryWithCommandList(
+								FModelingToolsManagerCommands::Get().SelectionEnableSoftSelection, CommandList));
+
+
+					
+					TSharedRef<SWidget> SoftSelectionRadiusWidget =
+						SNew(SHorizontalBox)
+						+ SHorizontalBox::Slot()
+						.FillWidth(0.9f)
+						[
+							SNew(SSpinBox<double>)
+								.MinValue(0.0f)
+								.Value_Lambda(
+									[WeakThis]() -> double 
+									{
+										if (!WeakThis.IsValid())
+										{
+											return 0;
+										}
+										return WeakThis->SelectionManager->GetSoftSelectionRadius();
+									}
+								)
+								.OnValueChanged_Lambda(
+									[WeakThis](double InValue)
+									{
+										if (WeakThis.IsValid())
+										{
+											WeakThis->SelectionManager->SetSoftSelectionRadius(InValue);
+										}
+									}
+								)
+						]
+						+ SHorizontalBox::Slot()
+						.FillWidth(0.1f);
+
+						SoftSelectionSection.AddEntry(FToolMenuEntry::InitWidget(
+							"SoftSelectionRadius", SoftSelectionRadiusWidget, LOCTEXT("SoftSelectionRadiusLabel", "Radius")));
+				}
+
 				// Section for locking/unlocking target
 				{
 					FToolMenuSection& LockingSection =
diff --git a/Engine/Plugins/Editor/ModelingToolsEditorMode/Source/ModelingToolsEditorMode/Private/ModelingToolsEditorModeToolkit.cpp b/Engine/Plugins/Editor/ModelingToolsEditorMode/Source/ModelingToolsEditorMode/Private/ModelingToolsEditorModeToolkit.cpp
--- a/Engine/Plugins/Editor/ModelingToolsEditorMode/Source/ModelingToolsEditorMode/Private/ModelingToolsEditorModeToolkit.cpp
+++ b/Engine/Plugins/Editor/ModelingToolsEditorMode/Source/ModelingToolsEditorMode/Private/ModelingToolsEditorModeToolkit.cpp
@@ -2585,7 +2585,7 @@ void FModelingToolsEditorModeToolkit::UpdateSelectionColors(UObject* Obj, FPrope
 	if (CachedSelectionManager)
 	{
 		const UModelingToolsModeCustomizationSettings* UISettings = GetMutableDefault<UModelingToolsModeCustomizationSettings>();
-		CachedSelectionManager->SetSelectionColors(UISettings->UnselectedColor, UISettings->HoverOverSelectedColor, UISettings->HoverOverUnselectedColor, UISettings->GeometrySelectedColor);
+		CachedSelectionManager->SetSelectionColors(UISettings->UnselectedColor, UISettings->HoverOverSelectedColor, UISettings->HoverOverUnselectedColor, UISettings->GeometrySelectedColor, UISettings->GeometrySoftSelectedColor);
 	}
 }
 
diff --git a/Engine/Plugins/Editor/ModelingToolsEditorMode/Source/ModelingToolsEditorMode/Private/ModelingToolsManagerActions.cpp b/Engine/Plugins/Editor/ModelingToolsEditorMode/Source/ModelingToolsEditorMode/Private/ModelingToolsManagerActions.cpp
--- a/Engine/Plugins/Editor/ModelingToolsEditorMode/Source/ModelingToolsEditorMode/Private/ModelingToolsManagerActions.cpp
+++ b/Engine/Plugins/Editor/ModelingToolsEditorMode/Source/ModelingToolsEditorMode/Private/ModelingToolsManagerActions.cpp
@@ -267,6 +267,8 @@ void FModelingToolsManagerCommands::RegisterCommands()
 	REGISTER_MODELING_TOOL_COMMAND(SelectionHitBackFaces, "Hit Back Faces", "Hit triangles that are facing away from the camera");
 	REGISTER_MODELING_TOOL_COMMAND(SelectionLocking, "Lock Target", "Click to Toggle locking on the Selected Object and disallow/allow Mesh Selections");
 	
+	REGISTER_MODELING_TOOL_COMMAND(SelectionEnableSoftSelection, "Soft Selection", "Apply a falloff effect to nearby vertices when transforming the selected elements");
+	
 	UI_COMMAND(AddToFavorites, "Add to Favorites", "Add to Favorites", EUserInterfaceActionType::Button, FInputChord());
 	UI_COMMAND(RemoveFromFavorites, "Remove from Favorites", "Remove from Favorites", EUserInterfaceActionType::Button, FInputChord());
 	UI_COMMAND(LoadFavoritesTools, "Faves", "Favorites", EUserInterfaceActionType::RadioButton, FInputChord());
diff --git a/Engine/Plugins/Editor/ModelingToolsEditorMode/Source/ModelingToolsEditorMode/Public/ModelingToolsEditorMode.h b/Engine/Plugins/Editor/ModelingToolsEditorMode/Source/ModelingToolsEditorMode/Public/ModelingToolsEditorMode.h
--- a/Engine/Plugins/Editor/ModelingToolsEditorMode/Source/ModelingToolsEditorMode/Public/ModelingToolsEditorMode.h
+++ b/Engine/Plugins/Editor/ModelingToolsEditorMode/Source/ModelingToolsEditorMode/Public/ModelingToolsEditorMode.h
@@ -174,7 +174,7 @@ protected:
 	static FDelegateHandle GlobalModelingWorldTeardownEventHandle;
 private:
 	// add modeling-mode-specific portion of new viewport toolbar
-	static void PopulateModelingModeViewportToolbar(const FName InMenuName, const TSharedPtr<const FUICommandList>& InCommandList);
+	void PopulateModelingModeViewportToolbar(const FName InMenuName, const TSharedPtr<const FUICommandList>& InCommandList);
 	// remove modeling-mode-specific portion of new viewport toolbar
 	static void RemoveModelingModeViewportToolbarExtensions();
 
diff --git a/Engine/Plugins/Editor/ModelingToolsEditorMode/Source/ModelingToolsEditorMode/Public/ModelingToolsEditorModeSettings.h b/Engine/Plugins/Editor/ModelingToolsEditorMode/Source/ModelingToolsEditorMode/Public/ModelingToolsEditorModeSettings.h
--- a/Engine/Plugins/Editor/ModelingToolsEditorMode/Source/ModelingToolsEditorMode/Public/ModelingToolsEditorModeSettings.h
+++ b/Engine/Plugins/Editor/ModelingToolsEditorMode/Source/ModelingToolsEditorMode/Public/ModelingToolsEditorModeSettings.h
@@ -381,6 +381,12 @@ public:
 	UPROPERTY(Config, EditAnywhere, Category = "Modeling Mode|Mesh Element Selection Mode")
 	FLinearColor GeometrySelectedColor = FLinearColor(0.982251f, 0.672443f, 0.029557f, .5f); // default yellow
 
+	/**
+	* Custom color for when geometry is a part of the soft selection when using Mesh Element Selection
+	*/
+	UPROPERTY(Config, EditAnywhere, Category = "Mesh Terrain Mode|Mesh Element Selection Mode")
+	FLinearColor GeometrySoftSelectedColor = FLinearColor(0.063f, 0.161f, 0.71f, .5f); // default blue
+
 public:
 
 	// saved-state for various mode settings that are configured via UI toggles/etc, and not exposed in settings dialog
diff --git a/Engine/Plugins/Editor/ModelingToolsEditorMode/Source/ModelingToolsEditorMode/Public/ModelingToolsManagerActions.h b/Engine/Plugins/Editor/ModelingToolsEditorMode/Source/ModelingToolsEditorMode/Public/ModelingToolsManagerActions.h
--- a/Engine/Plugins/Editor/ModelingToolsEditorMode/Source/ModelingToolsEditorMode/Public/ModelingToolsManagerActions.h
+++ b/Engine/Plugins/Editor/ModelingToolsEditorMode/Source/ModelingToolsEditorMode/Public/ModelingToolsManagerActions.h
@@ -203,6 +203,7 @@ public:
 	TSharedPtr<FUICommandInfo> SelectionMeshTypes_StaticMeshes;
 	TSharedPtr<FUICommandInfo> SelectionLocking;
 	TSharedPtr<FUICommandInfo> SelectionHitBackFaces;
+	TSharedPtr<FUICommandInfo> SelectionEnableSoftSelection;
 
 	TSharedPtr<FUICommandInfo> AddToFavorites;
 	TSharedPtr<FUICommandInfo> RemoveFromFavorites;
diff --git a/Engine/Plugins/Runtime/GeometryProcessing/Source/DynamicMesh/Private/Operations/GroupTopologyDeformer.cpp b/Engine/Plugins/Runtime/GeometryProcessing/Source/DynamicMesh/Private/Operations/GroupTopologyDeformer.cpp
--- a/Engine/Plugins/Runtime/GeometryProcessing/Source/DynamicMesh/Private/Operations/GroupTopologyDeformer.cpp
+++ b/Engine/Plugins/Runtime/GeometryProcessing/Source/DynamicMesh/Private/Operations/GroupTopologyDeformer.cpp
@@ -129,7 +129,9 @@ void FGroupTopologyDeformer::SetActiveHandleFromSelection(const FGeometrySelecti
 		return;
 	}
 
-	TArray<int> Elements;
+	TArray<int> Faces;
+	TArray<int> Edges;
+	TArray<int> Verts;
 	if (Selection.ElementType == EGeometryElementType::Face)
 	{
 		for (uint64 Element : Selection.Selection)
@@ -137,10 +139,9 @@ void FGroupTopologyDeformer::SetActiveHandleFromSelection(const FGeometrySelecti
 			int32 GroupID = FGeoSelectionID(Element).TopologyID;
 			if (Topology->FindGroupByID(GroupID) != nullptr)
 			{
-				Elements.Add(GroupID);
+				Faces.Add(GroupID);
 			}
 		}
-		SetActiveHandleFaces(Elements);
 	}
 	else if (Selection.ElementType == EGeometryElementType::Edge)
 	{
@@ -149,10 +150,9 @@ void FGroupTopologyDeformer::SetActiveHandleFromSelection(const FGeometrySelecti
 			int32 GroupEdgeID = FGeoSelectionID(Element).TopologyID;
 			if (GroupEdgeID >= 0 && GroupEdgeID < Topology->Edges.Num())
 			{
-				Elements.Add(GroupEdgeID);
+				Edges.Add(GroupEdgeID);
 			}
 		}
-		SetActiveHandleEdges(Elements);
 	}
 	else if (Selection.ElementType == EGeometryElementType::Vertex)
 	{
@@ -161,15 +161,109 @@ void FGroupTopologyDeformer::SetActiveHandleFromSelection(const FGeometrySelecti
 			int32 CornerID = FGeoSelectionID(Element).TopologyID;
 			if (CornerID >= 0 && CornerID < Topology->Corners.Num())
 			{
-				Elements.Add(CornerID);
+				Verts.Add(CornerID);
 			}
 		}
-		SetActiveHandleCorners(Elements);
+	}
+	
+	for (const TPair<uint64, double>& WeightedVert : Selection.SoftSelectedVertices)
+	{
+		int32 CornerID = FGeoSelectionID(WeightedVert.Key).TopologyID;
+		if (CornerID >= 0 && CornerID < Topology->Corners.Num())
+		{
+			Verts.Add(CornerID);
+		}
 	}
 
+	SetActiveHandle(Faces, Edges, Verts);
+
 }
 
 
+void FGroupTopologyDeformer::SetActiveHandle(
+	const TArray<int>& FaceGroupIDs,
+	const TArray<int>& TopologyEdgeIDs,
+	const TArray<int>& TopologyCornerIDs
+	)
+{
+	Reset();
+
+	// Face
+	
+	// is this right? some HandleBoundaryVertices may also be interior if we do it per-face...
+	for (int GroupID : FaceGroupIDs)
+	{
+		Topology->CollectGroupVertices(GroupID, HandleVertices);
+		Topology->CollectGroupBoundaryVertices(GroupID, HandleBoundaryVertices);
+	}
+
+	// find neighbour group set
+	TArray<int> HandleGroups = FaceGroupIDs;
+	TArray<int> NbrGroups;
+	for (int GroupID : FaceGroupIDs)
+	{
+		const TArray<int>& CurGroupNbrGroups = Topology->GetGroupNbrGroups(GroupID);
+		for (int NbrGroupID : CurGroupNbrGroups)
+		{
+			if (HandleGroups.Contains(NbrGroupID) == false)
+			{
+				NbrGroups.AddUnique(NbrGroupID);
+			}
+		}
+	}
+
+	// Edge
+	for (int EdgeID : TopologyEdgeIDs)
+	{
+		const TArray<int>& EdgeVerts = Topology->GetGroupEdgeVertices(EdgeID);
+		for (int VertID : EdgeVerts)
+		{
+			HandleVertices.Add(VertID);
+			HandleBoundaryVertices.Add(VertID);
+		}
+	}
+
+	TArray<int> EdgeNbrGroups;
+	Topology->FindEdgeNbrGroups(TopologyEdgeIDs, EdgeNbrGroups);
+	for (int NbrGroupID : EdgeNbrGroups)
+	{
+		if (HandleGroups.Contains(NbrGroupID) == false)
+		{
+			NbrGroups.AddUnique(NbrGroupID);
+		}
+	}	
+
+	// Corner
+	for (int CornerID : TopologyCornerIDs)
+	{
+		int VertID = Topology->GetCornerVertexID(CornerID);
+		if (VertID >= 0)
+		{
+			HandleVertices.Add(VertID);
+			HandleBoundaryVertices.Add(VertID);
+		}
+	}
+
+
+	TArray<int> CornerNbrGroups;
+	Topology->FindCornerNbrGroups(TopologyCornerIDs, CornerNbrGroups);
+	for (int NbrGroupID : CornerNbrGroups)
+	{
+		if (HandleGroups.Contains(NbrGroupID) == false)
+		{
+			NbrGroups.AddUnique(NbrGroupID);
+		}
+	}	
+	
+	ModifiedVertices = HandleVertices;
+	
+	CalculateROI(HandleGroups, NbrGroups);
+
+	SaveInitialPositions();
+
+	ComputeEncoding();
+}
+
 
 void FGroupTopologyDeformer::SaveInitialPositions()
 {
diff --git a/Engine/Plugins/Runtime/GeometryProcessing/Source/DynamicMesh/Public/Operations/GroupTopologyDeformer.h b/Engine/Plugins/Runtime/GeometryProcessing/Source/DynamicMesh/Public/Operations/GroupTopologyDeformer.h
--- a/Engine/Plugins/Runtime/GeometryProcessing/Source/DynamicMesh/Public/Operations/GroupTopologyDeformer.h
+++ b/Engine/Plugins/Runtime/GeometryProcessing/Source/DynamicMesh/Public/Operations/GroupTopologyDeformer.h
@@ -15,6 +15,7 @@ namespace Geometry
 class FDynamicMesh3;
 class FGroupTopology;
 struct FGeometrySelection;
+struct FGeometrySoftSelection;
 
 
 /**
@@ -103,6 +104,12 @@ public:
 	 */
 	UE_API virtual void SetActiveHandleFromSelection(const FGeometrySelection& Selection);
 
+	
+	/**
+	 * Set the active handle to the given Faces, Edges, Corners
+	 */
+	UE_API virtual void SetActiveHandle(const TArray<int>& FaceGroupIDs, const TArray<int>& TopologyEdgeIDs, const TArray<int>& TopologyCornerIDs);
+	
 	//
 	// Solving
 	// 
diff --git a/Engine/Plugins/Runtime/GeometryProcessing/Source/DynamicMesh/Public/Selections/GeometrySelection.h b/Engine/Plugins/Runtime/GeometryProcessing/Source/DynamicMesh/Public/Selections/GeometrySelection.h
--- a/Engine/Plugins/Runtime/GeometryProcessing/Source/DynamicMesh/Public/Selections/GeometrySelection.h
+++ b/Engine/Plugins/Runtime/GeometryProcessing/Source/DynamicMesh/Public/Selections/GeometrySelection.h
@@ -149,6 +149,9 @@ struct FGeometrySelection
 	/** Set of selected items/elements */
 	TSet<uint64> Selection;
 
+	/** Map of soft selected vertices and their weights*/
+	TMap<uint64, double> SoftSelectedVertices;
+	
 	/**
 	 * @return true if Selection is empty
 	 */
@@ -198,6 +201,20 @@ struct FGeometrySelection
 	{
 		return ElementType == OtherSelection.ElementType && TopologyType == OtherSelection.TopologyType;
 	}
+
+	/**
+	 * @return a FGeometrySelection containing the soft selected vertices
+	 */
+	FGeometrySelection GetSoftSelectedVertices() const
+	{
+		FGeometrySelection Result;
+		Result.TopologyType = TopologyType;
+		Result.ElementType = EGeometryElementType::Vertex;
+
+		SoftSelectedVertices.GetKeys(Result.Selection);
+	
+		return Result;
+	}
 };
 
 
@@ -309,8 +326,6 @@ struct FGeometrySelectionUpdateResult
 	FGeometrySelectionDelta SelectionDelta;
 };
 
-
-
 /**
  * FGeometrySelectionEditor is a utility/helper class used for modifying
  * a FGeometrySelection. The various functions can be used to add or remove
diff --git a/Engine/Plugins/Runtime/MeshModelingToolset/Source/MeshModelingTools/Private/DeformMeshPolygonsTool.cpp b/Engine/Plugins/Runtime/MeshModelingToolset/Source/MeshModelingTools/Private/DeformMeshPolygonsTool.cpp
--- a/Engine/Plugins/Runtime/MeshModelingToolset/Source/MeshModelingTools/Private/DeformMeshPolygonsTool.cpp
+++ b/Engine/Plugins/Runtime/MeshModelingToolset/Source/MeshModelingTools/Private/DeformMeshPolygonsTool.cpp
@@ -186,6 +186,7 @@ public:
 	void SetActiveHandleFaces(const TArray<int>& FaceGroupIDs) override;
 	void SetActiveHandleEdges(const TArray<int>& TopologyEdgeIDs) override;
 	void SetActiveHandleCorners(const TArray<int>& TopologyCornerIDs) override;
+	void SetActiveHandle(const TArray<int>& FaceGroupIDs, const TArray<int>& TopologyEdgeIDs, const TArray<int>& TopologyCornerIDs) override;
 
 
 	/** Allocates shared storage for use in task synchronization */
@@ -700,6 +701,33 @@ void FGroupTopologyLaplacianDeformer::SetActiveHandleCorners(const TArray<int>&
 	SaveInitialPositions();
 }
 
+void FGroupTopologyLaplacianDeformer::SetActiveHandle(const TArray<int>& FaceGroupIDs, const TArray<int>& TopologyEdgeIDs,
+	const TArray<int>& TopologyCornerIDs)
+{
+	// Currently mixing different types of handles is not yet supported.
+	
+	if (!FaceGroupIDs.IsEmpty())
+	{
+		check(TopologyEdgeIDs.IsEmpty() && TopologyCornerIDs.IsEmpty());
+		SetActiveHandleFaces(FaceGroupIDs);
+		return;
+	}
+	
+	if (!TopologyEdgeIDs.IsEmpty())
+	{
+		check(FaceGroupIDs.IsEmpty() && TopologyCornerIDs.IsEmpty());
+		SetActiveHandleEdges(TopologyEdgeIDs);
+		return;
+	}
+
+	if (!TopologyCornerIDs.IsEmpty())
+	{
+		check(FaceGroupIDs.IsEmpty() && TopologyEdgeIDs.IsEmpty());
+		SetActiveHandleCorners(TopologyCornerIDs);
+		return;
+	}
+}
+
 
 void FGroupTopologyLaplacianDeformer::UpdateSelection(const FDynamicMesh3* TargetMesh, const TArray<int>& Groups,
                                                       bool bLocalizeDeformation)
diff --git a/Engine/Plugins/Runtime/MeshModelingToolset/Source/ModelingComponents/Private/Selection/DynamicMeshPolygroupTransformer.cpp b/Engine/Plugins/Runtime/MeshModelingToolset/Source/ModelingComponents/Private/Selection/DynamicMeshPolygroupTransformer.cpp
--- a/Engine/Plugins/Runtime/MeshModelingToolset/Source/ModelingComponents/Private/Selection/DynamicMeshPolygroupTransformer.cpp
+++ b/Engine/Plugins/Runtime/MeshModelingToolset/Source/ModelingComponents/Private/Selection/DynamicMeshPolygroupTransformer.cpp
@@ -16,7 +16,7 @@ using namespace UE::Geometry;
 
 #define LOCTEXT_NAMESPACE "FDynamicMeshPolygroupTransformer"
 
-void FDynamicMeshPolygroupTransformer::BeginTransform(const FGeometrySelection& Selection)
+void FDynamicMeshPolygroupTransformer::BeginTransform(const UE::Geometry::FGeometrySelection& Selection)
 {
 	if (!ensure(Selection.TopologyType == EGeometryTopologyType::Polygroup))
 	{
@@ -42,7 +42,20 @@ void FDynamicMeshPolygroupTransformer::BeginTransform(const FGeometrySelection&
 		UE::Geometry::EnumeratePolygroupSelectionVertices(Selection, SourceMesh, UseTopology, FTransform::Identity,
 			[&](uint32 VertexID, const FVector3d& Position) { VertexIDs.Add((int32)VertexID); }
 		);
+
+		NumSelectedMeshVertices = VertexIDs.Num();
 		MeshVertices = VertexIDs.Array();
+
+		int32 NumWeightedMeshVertices = NumSelectedMeshVertices + Selection.SoftSelectedVertices.Num();
+		MeshVertices.Reserve(NumWeightedMeshVertices);
+		SoftSelectedVertexWeights.Reserve(Selection.SoftSelectedVertices.Num());
+		
+		for (const TPair<uint64, double>& WeightedVertex : Selection.SoftSelectedVertices)
+		{
+			MeshVertices.Add(FGeoSelectionID(WeightedVertex.Key).GeometryID);
+			SoftSelectedVertexWeights.Add(WeightedVertex.Value);
+		}
+
 		int32 NumVertices = MeshVertices.Num();
 		InitialPositions.SetNum(NumVertices);
 		for (int32 k = 0; k < NumVertices; ++k)
@@ -83,7 +96,7 @@ void FDynamicMeshPolygroupTransformer::BeginTransform(const FGeometrySelection&
 }
 
 void FDynamicMeshPolygroupTransformer::UpdateTransform( 
-	TFunctionRef<FVector3d(int32 VertexID, const FVector3d& InitialPosition, const FTransform& WorldTransform)> PositionTransformFunc  )
+	TFunctionRef<FVector3d(int32 VertexID, double Weight,const FVector3d& InitialPosition, const FTransform& WorldTransform)> PositionTransformFunc  )
 {
 	if (bFallbackToSimpleTransform)
 	{
@@ -97,7 +110,12 @@ void FDynamicMeshPolygroupTransformer::UpdateTransform(
 
 	for (int32 k = 0; k < N; ++k)
 	{
-		UpdatedPositions[k] = PositionTransformFunc(MeshVertices[k], InitialPositions[k], WorldTransform);
+		double Weight = 1.0f;
+		if (k >= NumSelectedMeshVertices)
+		{
+			Weight = SoftSelectedVertexWeights[(k - NumSelectedMeshVertices)];
+		}
+		UpdatedPositions[k] = PositionTransformFunc(MeshVertices[k], Weight, InitialPositions[k], WorldTransform);
 	}
 
 	Selector->GetDynamicMesh()->EditMesh([&](FDynamicMesh3& EditMesh)
@@ -109,7 +127,12 @@ void FDynamicMeshPolygroupTransformer::UpdateTransform(
 		{
 			if ( int32* ArrayIndex = ROIMap.Find(VertIdx) )
 			{
-				return PositionTransformFunc( VertIdx, InitialPositions[*ArrayIndex], WorldTransform );
+				double Weight = 1.0f;
+				if (*ArrayIndex >= NumSelectedMeshVertices)
+				{
+					Weight = SoftSelectedVertexWeights[(*ArrayIndex - NumSelectedMeshVertices)];
+				}
+				return PositionTransformFunc( VertIdx, Weight, InitialPositions[*ArrayIndex], WorldTransform );
 			}
 			else
 			{
diff --git a/Engine/Plugins/Runtime/MeshModelingToolset/Source/ModelingComponents/Private/Selection/DynamicMeshSelector.cpp b/Engine/Plugins/Runtime/MeshModelingToolset/Source/ModelingComponents/Private/Selection/DynamicMeshSelector.cpp
--- a/Engine/Plugins/Runtime/MeshModelingToolset/Source/ModelingComponents/Private/Selection/DynamicMeshSelector.cpp
+++ b/Engine/Plugins/Runtime/MeshModelingToolset/Source/ModelingComponents/Private/Selection/DynamicMeshSelector.cpp
@@ -16,6 +16,7 @@
 #include "Selection/DynamicMeshPolygroupTransformer.h"
 #include "Selection/ToolSelectionUtil.h"
 #include "ToolDataVisualizer.h"
+#include "Parameterization/MeshDijkstra.h"
 
 using namespace UE::Geometry;
 
@@ -58,6 +59,7 @@ void FBaseDynamicMeshSelector::InvalidateOnMeshChange(FDynamicMeshChangeInfo Cha
 {
 	ColliderMesh.Reset();
 	GroupEdgeSegmentTree.Reset();
+	PathFinder.Reset();
 
 	if (ChangeInfo.Type != EDynamicMeshChangeType::MeshVertexChange)
 	{
@@ -603,6 +605,104 @@ void FBaseDynamicMeshSelector::UpdateSelectionViaShape(
 	}
 }
 
+void FBaseDynamicMeshSelector::UpdateSoftSelection(
+	FGeometrySelection& Selection,
+	const FUpdateSoftSelectionConfig& InSoftSelectionConfig
+	)
+{
+	Selection.SoftSelectedVertices.Reset();
+	
+	if (!InSoftSelectionConfig.bEnabled)
+	{
+		return;
+	}
+
+	if (InSoftSelectionConfig.bBaseSelectionChanged || !PathFinder.IsValid())
+	{
+		// seed a Dijkstra path finder with the selected vertices
+		PathFinder = MakePimpl<TMeshDijkstra<FDynamicMesh3>>(GetDynamicMesh()->GetMeshPtr());
+		TArray<TMeshDijkstra<FDynamicMesh3>::FSeedPoint> SeedPoints;
+
+		TSet<int32> VertexIDs;
+		GetDynamicMesh()->ProcessMesh([&Selection, &VertexIDs, this](const FDynamicMesh3& SourceMesh)
+			{
+				// get set of selected vertex IDs, and then vector (only converting to vector for VertexToTriangleOneRing...)
+				if (Selection.TopologyType == EGeometryTopologyType::Polygroup)
+				{
+					UE::Geometry::EnumeratePolygroupSelectionVertices(Selection, SourceMesh, GetGroupTopology(), FTransform::Identity,
+						[&VertexIDs](uint32 VertexID, const FVector3d& Position) { VertexIDs.Add((int32)VertexID); }
+					);
+				}
+				else
+				{
+					UE::Geometry::EnumerateTriangleSelectionVertices(Selection, SourceMesh, nullptr,
+						[&VertexIDs](uint32 VertexID, const FVector3d& Position) { VertexIDs.Add((int32)VertexID); }
+					);
+				}
+			});
+	
+		for (const int32 SelectedVertex : VertexIDs)
+		{
+			SeedPoints.Add({ SelectedVertex, SelectedVertex, 0 });
+		}
+	
+	
+		PathFinder->ComputeToMaxDistance(SeedPoints, TNumericLimits<double>::Max());	
+	}
+
+
+	PathFinder->EnumerateComputedPoint([&Selection, &InSoftSelectionConfig, this](TMeshDijkstra<FDynamicMesh3>::FComputedPoint InPoint)
+		{
+			// Skip points already in the selection
+			if (InPoint.bSeedPoint)
+			{
+				return;
+			}
+
+			const double& Radius = InSoftSelectionConfig.Radius;
+			
+			if (InPoint.GraphDistance > Radius)
+			{
+				return;
+			}
+			
+			FGeoSelectionID SelectionID;
+			bool bAddPoint = false;
+			if (Selection.TopologyType == EGeometryTopologyType::Triangle)
+			{
+				SelectionID = FGeoSelectionID((uint32)InPoint.PointID);
+				bAddPoint = true;
+			}
+			else if (Selection.TopologyType == EGeometryTopologyType::Polygroup)
+			{
+				const int32 CornerID = GetGroupTopology()->GetCornerIDFromVertexID(InPoint.PointID);
+				if (CornerID != IndexConstants::InvalidID)
+				{
+					bAddPoint = true;
+					SelectionID = FGeoSelectionID(InPoint.PointID, CornerID);
+				}
+			}
+			
+			if (bAddPoint)
+			{
+				uint64 SelectionIDEncoded = SelectionID.Encoded();
+
+				double NormalizedDistance = Radius > UE_SMALL_NUMBER ? InPoint.GraphDistance / Radius : 0.0 ;
+
+				double Weight = FMath::CubicInterp(1.0, 0.0, 0.0, 0.0, NormalizedDistance);
+				
+				Selection.SoftSelectedVertices.Add(SelectionIDEncoded, Weight);
+			}
+		});
+}
+
+bool FBaseDynamicMeshSelector::ConvertSelection(
+	const FGeometrySelection& FromSelectionIn,
+	FGeometrySelection& ToSelectionOut,
+	const EEnumerateSelectionConversionParams ConversionParams)
+{
+	return UE::Geometry::ConvertSelection(GetDynamicMesh()->GetMeshRef(), GetGroupTopology(), FromSelectionIn, ToSelectionOut, ConversionParams);	
+}
 
 
 void FBaseDynamicMeshSelector::GetSelectionFrame(const FGeometrySelection& Selection, FFrame3d& SelectionFrame, bool bTransformToWorld)
@@ -914,8 +1014,20 @@ void FBasicDynamicMeshSelectionTransformer::BeginTransform(const FGeometrySelect
 				[&](uint32 VertexID, const FVector3d& Position) { VertexIDs.Add((int32)VertexID); }
 			);
 		}
+			
+		NumSelectedMeshVertices = VertexIDs.Num();
 		MeshVertices = VertexIDs.Array();
 
+		int32 NumWeightedMeshVertices = NumSelectedMeshVertices + Selection.SoftSelectedVertices.Num();
+		MeshVertices.Reserve(NumWeightedMeshVertices);
+		SoftSelectedVertexWeights.Reserve(Selection.SoftSelectedVertices.Num());
+			
+		for (const TPair<uint64, double>& WeightedVertex : Selection.SoftSelectedVertices)
+		{
+			MeshVertices.Add(FGeoSelectionID(WeightedVertex.Key).GeometryID);
+			SoftSelectedVertexWeights.Add(WeightedVertex.Value);
+		}
+
 		// save initial positions
 		InitialPositions.Reserve(MeshVertices.Num());
 		for (int32 vid : MeshVertices)
@@ -967,7 +1079,7 @@ void FBasicDynamicMeshSelectionTransformer::BeginTransform(const FGeometrySelect
 }
 
 void FBasicDynamicMeshSelectionTransformer::UpdateTransform( 
-	TFunctionRef<FVector3d(int32 VertexID, const FVector3d& InitialPosition, const FTransform& WorldTransform)> PositionTransformFunc 
+	TFunctionRef<FVector3d(int32 VertexID, double Weight, const FVector3d& InitialPosition, const FTransform& WorldTransform)> PositionTransformFunc 
 )
 {
 	int32 N = MeshVertices.Num();
@@ -975,7 +1087,12 @@ void FBasicDynamicMeshSelectionTransformer::UpdateTransform(
 
 	for (int32 k = 0; k < N; ++k)
 	{
-		UpdatedPositions[k] = PositionTransformFunc(MeshVertices[k], InitialPositions[k], WorldTransform);
+		double Weight = 1.0f;
+		if (k >= NumSelectedMeshVertices)
+		{
+			Weight = SoftSelectedVertexWeights[(k - NumSelectedMeshVertices)];
+		}
+		UpdatedPositions[k] = PositionTransformFunc(MeshVertices[k], Weight, InitialPositions[k], WorldTransform);
 	}
 
 	Selector->GetDynamicMesh()->EditMesh([&](FDynamicMesh3& EditMesh)
diff --git a/Engine/Plugins/Runtime/MeshModelingToolset/Source/ModelingComponents/Private/Selection/GeometrySelectionManager.cpp b/Engine/Plugins/Runtime/MeshModelingToolset/Source/ModelingComponents/Private/Selection/GeometrySelectionManager.cpp
--- a/Engine/Plugins/Runtime/MeshModelingToolset/Source/ModelingComponents/Private/Selection/GeometrySelectionManager.cpp
+++ b/Engine/Plugins/Runtime/MeshModelingToolset/Source/ModelingComponents/Private/Selection/GeometrySelectionManager.cpp
@@ -37,6 +37,8 @@ void UGeometrySelectionManager::Initialize( UInteractiveToolsContext* ToolsConte
 	ToolsContext = ToolsContextIn;
 	TransactionsAPI = TransactionsAPIIn;
 	PreviewGeometry = NewObject<UPreviewGeometry>(this);
+
+	OnSelectionModified.AddUObject(this, &UGeometrySelectionManager::UpdateSoftSelectionOnSelectionModified);
 }
 
 void UGeometrySelectionManager::RegisterSelectorFactory(TUniquePtr<IGeometrySelectorFactory> Factory)
@@ -300,6 +302,26 @@ void UGeometrySelectionManager::RebuildSelectable() const
 	}
 }
 
+void UGeometrySelectionManager::UpdateSoftSelectionOnSelectionModified()
+{
+	if (!HasActiveTargets())
+	{
+		return;
+	}
+	// currently only going to support one object, not sure how to support more yet...
+	FGeometrySelectionTarget* Target = ActiveTargetReferences[0].Get();
+
+	IGeometrySelector::FUpdateSoftSelectionConfig SoftSelectionConfig;
+	SoftSelectionConfig.bBaseSelectionChanged = true;
+	SoftSelectionConfig.bEnabled = bEnableSoftSelection;
+	SoftSelectionConfig.Radius = SoftSelectionRadius;
+	Target->Selector->UpdateSoftSelection(Target->Selection, SoftSelectionConfig);
+	
+	MarkRenderCachesDirty(false);
+}
+
+
+
 
 void UGeometrySelectionManager::SetMeshSelectionTypeAndMode(EGeometryElementType NewElementType, EMeshTopologyMode NewTopologyMode, bool bConvertSelection)
 {
@@ -419,6 +441,52 @@ void UGeometrySelectionManager::SetHitBackFaces(const bool bInHitBackFaces)
 	}
 }
 
+void UGeometrySelectionManager::SetSoftSelectionEnabled(bool bEnable)
+{
+	if (bEnableSoftSelection != bEnable)
+	{
+		bEnableSoftSelection = bEnable;
+		if (!HasActiveTargets())
+		{
+			return;
+		}
+		// currently only going to support one object, not sure how to support more yet...
+		FGeometrySelectionTarget* Target = ActiveTargetReferences[0].Get();
+
+		IGeometrySelector::FUpdateSoftSelectionConfig Config;
+		Config.bBaseSelectionChanged = false;
+		Config.bEnabled = bEnableSoftSelection;
+		Config.Radius = SoftSelectionRadius;
+		Target->Selector->UpdateSoftSelection(Target->Selection, Config);
+
+		MarkRenderCachesDirty(false);
+	}
+}
+
+void UGeometrySelectionManager::SetSoftSelectionRadius(double Radius)
+{
+	double ClampedRadius = FMath::Max(Radius, 0);
+	if (SoftSelectionRadius != ClampedRadius)
+	{
+		SoftSelectionRadius = ClampedRadius;
+		
+		if (!HasActiveTargets())
+		{
+			return;
+		}
+		// currently only going to support one object, not sure how to support more yet...
+		FGeometrySelectionTarget* Target = ActiveTargetReferences[0].Get();
+
+		IGeometrySelector::FUpdateSoftSelectionConfig Config;
+		Config.bBaseSelectionChanged = false;
+		Config.bEnabled = bEnableSoftSelection;
+		Config.Radius = SoftSelectionRadius;
+		Target->Selector->UpdateSoftSelection(Target->Selection, Config);
+
+		MarkRenderCachesDirty(false);
+	}
+}
+
 
 bool UGeometrySelectionManager::HasSelection() const
 {
@@ -1386,8 +1454,18 @@ bool UGeometrySelectionManager::BeginTransformation()
 	return true;
 }
 
+void UGeometrySelectionManager::UpdateTransformation(
+	TFunctionRef<FVector3d(int32 VertexID, const FVector3d&, const FTransform&)> PositionTransformFunc)
+{
+	auto PositionTransformFuncWithUnusedWeight = [&PositionTransformFunc](int32 VertexID, double Weight, const FVector3d& Position, const FTransform& Transform)
+		{
+			return PositionTransformFunc(VertexID, Position, Transform);
+		};
+	UpdateTransformation(PositionTransformFuncWithUnusedWeight);
+}
+
 void UGeometrySelectionManager::UpdateTransformation( 
-	TFunctionRef<FVector3d(int32 VertexID, const FVector3d&, const FTransform&)> PositionTransformFunc )
+	TFunctionRef<FVector3d(int32 VertexID, double Weight, const FVector3d&, const FTransform&)> PositionTransformFunc )
 {
 	if ( ! ensure(IsInActiveTransformation()) ) return;
 
@@ -1400,6 +1478,7 @@ void UGeometrySelectionManager::UpdateTransformation(
 	}
 
 	RenderCachesDirtyFlags |= EEnumerateRenderCachesDirtyFlags::SelectionCachesDirty;
+	RenderCachesDirtyFlags |= EEnumerateRenderCachesDirtyFlags::SoftSelectionCachesDirty;
 }
 
 void UGeometrySelectionManager::EndTransformation()
@@ -1422,9 +1501,48 @@ void UGeometrySelectionManager::EndTransformation()
 	GetTransactionsAPI()->EndUndoTransaction();
 
 	RenderCachesDirtyFlags |= EEnumerateRenderCachesDirtyFlags::SelectionCachesDirty;
+	RenderCachesDirtyFlags |= EEnumerateRenderCachesDirtyFlags::SoftSelectionCachesDirty;
 }
 
+void UGeometrySelectionManager::UpdateTransformation(
+	const FTransform& StartTransform,
+	const FVector& TranslationToApply,
+	const FQuat& RotationToApply,
+	const FVector& ScaleToApply,
+	bool bWorldSpaceScale
+	)
+{
+	if (bWorldSpaceScale)
+	{
+		// For a world frame gizmo, the scaling needs to happen in world aligned gizmo space, but the 
+		// rotation is still encoded in the local gizmo frame change.
+		UpdateTransformation([&StartTransform, &TranslationToApply, &RotationToApply, &ScaleToApply](int32 VertexID, double Weight, const FVector3d& InitialPosition, const FTransform3d& WorldTransform)
+			{
+				FVector3d PosLocal = InitialPosition;
+				FVector3d PosWorld = WorldTransform.TransformPosition(PosLocal);
+				FVector3d PosWorldGizmo = PosWorld - StartTransform.GetLocation();
 
+				FVector3d NewPosWorld = FQuat::Slerp(FQuat::Identity, RotationToApply, Weight) * (((ScaleToApply * Weight ) + 1.0f) * PosWorldGizmo) + (StartTransform.GetLocation() + TranslationToApply * Weight);
+				FVector3d NewPosLocal = WorldTransform.InverseTransformPosition(NewPosWorld);
+				return NewPosLocal;
+			});	
+	}
+	else
+	{
+		UpdateTransformation([&StartTransform, &TranslationToApply, &RotationToApply, &ScaleToApply](int32 VertexID, double Weight, const FVector3d& InitialPosition, const FTransform3d& WorldTransform)
+			{
+				// For a local gizmo, we just get the coordinates in the original frame, scale in that frame,
+				// then interpret them as coordinates in the new frame.
+				FVector3d PosLocal = InitialPosition;
+				FVector3d PosWorld = WorldTransform.TransformPosition(PosLocal);
+				FVector3d PosGizmo = StartTransform.InverseTransformPosition(PosWorld);
+				PosGizmo = ((ScaleToApply * Weight ) + 1.0f) * PosGizmo;
+				FVector3d NewPosWorld = FQuat::Slerp(FQuat::Identity, RotationToApply, Weight) * StartTransform.GetRotation() * PosGizmo + (StartTransform.GetLocation() + TranslationToApply * Weight);
+				FVector3d NewPosLocal = WorldTransform.InverseTransformPosition(NewPosWorld);
+				return NewPosLocal;
+			});
+	}	
+}
 
 
 
@@ -1608,10 +1726,16 @@ void UGeometrySelectionManager::RevertChange(IGeometrySelectionChange* Change)
 void UGeometrySelectionManager::OnTargetGeometryModified(IGeometrySelector* Selector)
 {
 	CachedSelectableRenderElements.Reset();
+	CachedSelectionRenderElements.Reset();
+	CachedSoftSelectionRenderElements.Reset();
 	
 	CachedSelectableRenderElements.SetNum(ActiveTargetReferences.Num());
+	CachedSelectionRenderElements.SetNum(ActiveTargetReferences.Num());
+	CachedSoftSelectionRenderElements.SetNum(ActiveTargetReferences.Num());
 
 	RemoveSets(UnselectedParams.Identifiers);
+	RemoveSets(SelectedParams.Identifiers);
+	RemoveSets(SoftSelectedParams.Identifiers);
 	
 	MarkRenderCachesDirty();
 	ClearActivePreview();
@@ -1622,12 +1746,15 @@ void UGeometrySelectionManager::UpdateSelectionRenderCacheOnTargetChange()
 {
 	CachedSelectionRenderElements.Reset();
 	CachedSelectableRenderElements.Reset();
+	CachedSoftSelectionRenderElements.Reset();
 	
 	CachedSelectionRenderElements.SetNum(ActiveTargetReferences.Num());
 	CachedSelectableRenderElements.SetNum(ActiveTargetReferences.Num());
+	CachedSoftSelectionRenderElements.SetNum(ActiveTargetReferences.Num());
 
 	RemoveSets(SelectedParams.Identifiers);
 	RemoveSets(UnselectedParams.Identifiers);
+	RemoveSets(SoftSelectedParams.Identifiers);
 	
 	MarkRenderCachesDirty();
 	ClearActivePreview();
@@ -1639,7 +1766,7 @@ void UGeometrySelectionManager::MarkRenderCachesDirty(bool bMarkSelectableDirty)
 	{
 		RenderCachesDirtyFlags |= EEnumerateRenderCachesDirtyFlags::UnselectedCachesDirty;
 	}
-	RenderCachesDirtyFlags |= EEnumerateRenderCachesDirtyFlags::PreviewCachesDirty | EEnumerateRenderCachesDirtyFlags::SelectionCachesDirty;
+	RenderCachesDirtyFlags |= EEnumerateRenderCachesDirtyFlags::PreviewCachesDirty | EEnumerateRenderCachesDirtyFlags::SelectionCachesDirty | EEnumerateRenderCachesDirtyFlags::SoftSelectionCachesDirty;
 }
 
 void UGeometrySelectionManager::RebuildSelectionRenderCaches()
@@ -1647,6 +1774,7 @@ void UGeometrySelectionManager::RebuildSelectionRenderCaches()
 	RebuildSelectionRenderCache();
 	RebuildSelectableRenderCache();
 	RebuildPreviewRenderCache();
+	RebuildSoftSelectionRenderCache();
 }
 
 void UGeometrySelectionManager::ClearActivePreview()
@@ -1744,7 +1872,7 @@ void UGeometrySelectionManager::RebuildSelectableRenderCache()
 		return;
 	}
 
-	check(ActiveTargetReferences.Num() == CachedSelectionRenderElements.Num());
+	check(ActiveTargetReferences.Num() == CachedSelectableRenderElements.Num());
 	for (int32 k = 0; k < ActiveTargetReferences.Num(); ++k)
 	{
 		TSharedPtr<FGeometrySelectionTarget> Target = ActiveTargetReferences[k];
@@ -1793,6 +1921,34 @@ void UGeometrySelectionManager::RebuildPreviewRenderCache()
 	RenderCachesDirtyFlags &= ~EEnumerateRenderCachesDirtyFlags::PreviewCachesDirty;
 }
 
+void UGeometrySelectionManager::RebuildSoftSelectionRenderCache()
+{
+	if ((RenderCachesDirtyFlags & EEnumerateRenderCachesDirtyFlags::SoftSelectionCachesDirty) == EEnumerateRenderCachesDirtyFlags::None)
+	{
+		return;
+	}
+
+	check(ActiveTargetReferences.Num() == CachedSoftSelectionRenderElements.Num());
+	for (int32 k = 0; k < ActiveTargetReferences.Num(); ++k)
+	{
+		TSharedPtr<FGeometrySelectionTarget> Target = ActiveTargetReferences[k];
+
+		FGeometrySelection SoftSelectionEdges;
+		SoftSelectionEdges.TopologyType = GetSelectionTopologyType();
+		SoftSelectionEdges.ElementType = EGeometryElementType::Edge;
+		
+		Target->Selector->ConvertSelection(Target->Selection.GetSoftSelectedVertices(), SoftSelectionEdges, EEnumerateSelectionConversionParams::ExpandSelection);
+		
+		FGeometrySelectionElements& SoftSelectionElements = CachedSoftSelectionRenderElements[k];
+		SoftSelectionElements.Reset();
+
+		Target->Selector->AccumulateSelectionElements(SoftSelectionEdges, SoftSelectionElements, true, EEnumerateSelectionMapping::Default);
+
+		CreateOrUpdateAllSets(SoftSelectionElements, SoftSelectedParams);
+	}
+
+	RenderCachesDirtyFlags &= ~EEnumerateRenderCachesDirtyFlags::SoftSelectionCachesDirty;
+}
 
 void UGeometrySelectionManager::DebugPrintSelection()
 {
@@ -1838,12 +1994,18 @@ void UGeometrySelectionManager::DebugRender(IToolsContextRenderAPI* RenderAPI)
 	}
 }
 
-void UGeometrySelectionManager::SetSelectionColors(const FLinearColor UnselectedCol, const FLinearColor HoverOverSelectedCol, const FLinearColor HoverOverUnselectedCol, const FLinearColor GeometrySelectedCol)
+void UGeometrySelectionManager::SetSelectionColors(
+	FLinearColor UnselectedCol,
+	FLinearColor HoverOverSelectedCol,
+	FLinearColor HoverOverUnselectedCol,
+	FLinearColor GeometrySelectedCol,
+	FLinearColor SoftSelectedCol)
 {
 	UnselectedParams.Color = UnselectedCol.ToFColor(true);
 	HoverOverSelectedParams.Color = HoverOverSelectedCol.ToFColor(true);
 	HoverOverUnselectedParams.Color = HoverOverUnselectedCol.ToFColor(true);
 	SelectedParams.Color = GeometrySelectedCol.ToFColor(true);
+	SoftSelectedParams.Color = SoftSelectedCol.ToFColor(true);
 
 	// on initial set up of Materials used for selection colors material (typically when entering modeling mode)
 	if (UnselectedParams.SelectionFillColor == nullptr) // if one is null, they all will be
@@ -1857,6 +2019,7 @@ void UGeometrySelectionManager::SetSelectionColors(const FLinearColor Unselected
 		SetMaterial(UnselectedParams, UnselectedCol);
 		SetMaterial(HoverOverUnselectedParams, HoverOverUnselectedCol);
 		SetMaterial(SelectedParams, GeometrySelectedCol);
+		SetMaterial(SoftSelectedParams, SoftSelectedCol);
 
 		// to avoid flickering, this version of GetCustomTwoSidedDeptOffsetMaterial (without opacity parameter) must be called for HoverOverSelected
 		HoverOverSelectedParams.SelectionFillColor =
@@ -1875,6 +2038,7 @@ void UGeometrySelectionManager::SetSelectionColors(const FLinearColor Unselected
 		SetColorAndOpacity(HoverOverSelectedParams, HoverOverSelectedCol);
 		SetColorAndOpacity(HoverOverUnselectedParams, HoverOverUnselectedCol);
 		SetColorAndOpacity(SelectedParams, GeometrySelectedCol);
+		SetColorAndOpacity(SoftSelectedParams, SoftSelectedCol);
 	}
 
 	// ensures that when color is changed in Editor Preferences, colors are immediately updated in the UI
@@ -1892,6 +2056,7 @@ void UGeometrySelectionManager::SetSelectionColors(const FLinearColor Unselected
 	UpdateAllSetsColor(HoverOverSelectedParams);
 	UpdateAllSetsColor(HoverOverUnselectedParams);
 	UpdateAllSetsColor(SelectedParams);
+	UpdateAllSetsColor(SoftSelectedParams);
 }
 
 void UGeometrySelectionManager::DisconnectPreviewGeometry()
diff --git a/Engine/Plugins/Runtime/MeshModelingToolset/Source/ModelingComponents/Public/Selection/DynamicMeshPolygroupTransformer.h b/Engine/Plugins/Runtime/MeshModelingToolset/Source/ModelingComponents/Public/Selection/DynamicMeshPolygroupTransformer.h
--- a/Engine/Plugins/Runtime/MeshModelingToolset/Source/ModelingComponents/Public/Selection/DynamicMeshPolygroupTransformer.h
+++ b/Engine/Plugins/Runtime/MeshModelingToolset/Source/ModelingComponents/Public/Selection/DynamicMeshPolygroupTransformer.h
@@ -20,7 +20,7 @@ class FDynamicMeshPolygroupTransformer : public FBasicDynamicMeshSelectionTransf
 public:
 
 	UE_API virtual void BeginTransform(const UE::Geometry::FGeometrySelection& Selection) override;
-	UE_API virtual void UpdateTransform( TFunctionRef<FVector3d(int32 VertexID, const FVector3d& InitialPosition, const FTransform& WorldTransform)> PositionTransformFunc ) override;
+	UE_API virtual void UpdateTransform( TFunctionRef<FVector3d(int32 VertexID, double Weight, const FVector3d& InitialPosition, const FTransform& WorldTransform)> PositionTransformFunc ) override;
 	UE_API virtual void EndTransform(IToolsContextTransactionsAPI* TransactionsAPI) override;
 
 	UE_API virtual void PreviewRender(IToolsContextRenderAPI* RenderAPI) override;
diff --git a/Engine/Plugins/Runtime/MeshModelingToolset/Source/ModelingComponents/Public/Selection/DynamicMeshSelector.h b/Engine/Plugins/Runtime/MeshModelingToolset/Source/ModelingComponents/Public/Selection/DynamicMeshSelector.h
--- a/Engine/Plugins/Runtime/MeshModelingToolset/Source/ModelingComponents/Public/Selection/DynamicMeshSelector.h
+++ b/Engine/Plugins/Runtime/MeshModelingToolset/Source/ModelingComponents/Public/Selection/DynamicMeshSelector.h
@@ -11,6 +11,7 @@
 #include "DynamicMesh/DynamicMeshAABBTree3.h"
 #include "DynamicMesh/DynamicMeshChangeTracker.h"
 #include "TransformTypes.h"
+#include "Parameterization/MeshDijkstra.h"
 
 #define UE_API MODELINGCOMPONENTS_API
 
@@ -111,6 +112,9 @@ public:
 		FGeometrySelectionUpdateResult& ResultOut 
 	) override;
 
+	UE_API virtual void UpdateSoftSelection(FGeometrySelection& Selection, const FUpdateSoftSelectionConfig& InSoftSelectionConfig) override;
+
+	UE_API virtual bool ConvertSelection(const FGeometrySelection& FromSelectionIn, FGeometrySelection& ToSelectionOut, const UE::Geometry::EEnumerateSelectionConversionParams ConversionParams) override;
 
 	virtual FTransform GetLocalToWorldTransform() const
 	{
@@ -194,6 +198,12 @@ protected:
 	UE_API void UpdateGroupEdgeSegmentTree();
 	UE_API const UE::Geometry::FSegmentTree3* GetGroupEdgeSpatial();
 
+	//
+	// PathFinder will be built on-demand if surface-based soft selection is enabled
+	//
+	TPimplPtr<UE::Geometry::TMeshDijkstra<FDynamicMesh3>> PathFinder;
+
+	
 	// support for sleep/restore
 	TWeakObjectPtr<UDynamicMesh> SleepingTargetMesh = nullptr;
 
@@ -263,7 +273,7 @@ public:
 
 	UE_API virtual void BeginTransform(const UE::Geometry::FGeometrySelection& Selection) override;
 
-	UE_API virtual void UpdateTransform( TFunctionRef<FVector3d(int32 VertexID, const FVector3d& InitialPosition, const FTransform& WorldTransform)> PositionTransformFunc ) override;
+	UE_API virtual void UpdateTransform( TFunctionRef<FVector3d(int32 VertexID, double Weight, const FVector3d& InitialPosition, const FTransform& WorldTransform)> PositionTransformFunc ) override;
 	UE_API void UpdatePendingVertexChange(bool bFinal);
 
 	/** Enable line drawing of selection during transform, this is necessary in some contexts where live mesh update is too slow */
@@ -279,6 +289,12 @@ protected:
 	FBaseDynamicMeshSelector* Selector;
 
 	TArray<int32> MeshVertices;
+	
+	TArray<double> SoftSelectedVertexWeights;
+	
+	// The number of vertices in MeshVertices = NumSelectedMeshVertices + SoftSelectedVertexWeights.Num()
+	int32 NumSelectedMeshVertices = 0;
+	
 	TArray<FVector3d> InitialPositions;
 	TSet<int32> TriangleROI;
 	TSet<int32> OverlayNormalsSet;
diff --git a/Engine/Plugins/Runtime/MeshModelingToolset/Source/ModelingComponents/Public/Selection/GeometrySelectionManager.h b/Engine/Plugins/Runtime/MeshModelingToolset/Source/ModelingComponents/Public/Selection/GeometrySelectionManager.h
--- a/Engine/Plugins/Runtime/MeshModelingToolset/Source/ModelingComponents/Public/Selection/GeometrySelectionManager.h
+++ b/Engine/Plugins/Runtime/MeshModelingToolset/Source/ModelingComponents/Public/Selection/GeometrySelectionManager.h
@@ -40,8 +40,9 @@ enum class EEnumerateRenderCachesDirtyFlags: uint8
 	SelectionCachesDirty = 1 << 0,
 	UnselectedCachesDirty = 1 << 1, //SelectableRenderCaches
 	PreviewCachesDirty = 1 << 2,
+	SoftSelectionCachesDirty = 1 << 3,
 
-	Default = SelectionCachesDirty | UnselectedCachesDirty | PreviewCachesDirty,
+	Default = SelectionCachesDirty | UnselectedCachesDirty | PreviewCachesDirty | SoftSelectionCachesDirty,
 };
 ENUM_CLASS_FLAGS(EEnumerateRenderCachesDirtyFlags);
 
@@ -387,13 +388,32 @@ public:
 	 * Update the active transformations with the given PositionTransformFunc.
 	 * See IGeometrySelectionTransformer::UpdateTransform for details on this callback
 	 */
+	UE_DEPRECATED(5.8, "Use the UpdateTransformation function that takes a lambda that includes per-vertex weight instead")
 	UE_API virtual void UpdateTransformation( TFunctionRef<FVector3d(int32 VertexID, const FVector3d&, const FTransform&)> PositionTransformFunc );
+	
+	/**
+	 * Update the active transformations with the given PositionTransformFunc.
+	 * See IGeometrySelectionTransformer::UpdateTransform for details on this callback
+	 */
+	UE_EXPERIMENTAL(5.8, "Per-vertex weighted transformation is experiemental while the API is under active development.")
+	UE_API virtual void UpdateTransformation( TFunctionRef<FVector3d(int32 VertexID, double Weight, const FVector3d&, const FTransform&)> PositionTransformFunc );
 
 	/**
 	 * End the current active transformation, and emit changes/transactions
 	 */
 	UE_API virtual void EndTransformation();
 
+	/**
+	 * Update the active transformations with the given transform deltas
+	 */
+	UE_EXPERIMENTAL(5.8, "This function is experiemental while the API is under active development.")
+	UE_API void UpdateTransformation(
+		const FTransform& StartTransform,
+		const FVector& TranslationToApply,
+		const FQuat& RotationToApply,
+		const FVector& ScaleToApply,
+		bool bWorldSpaceScale
+		);
 
 	//
 	// Command Execution
@@ -439,9 +459,9 @@ public:
 	UE_API virtual void DebugRender(IToolsContextRenderAPI* RenderAPI);
 	
 	/** Set the colors to be used during mesh element selection for:
-	 * Unselected elements, Hover over selection, Hover over non-selection, and Selected elements
+	 * Unselected elements, Hover over selection, Hover over non-selection, Selected elements, and Soft selected element
 	 */
-	UE_API void SetSelectionColors(FLinearColor UnselectedCol, FLinearColor HoverOverSelectedCol, FLinearColor HoverOverUnselectedCol, FLinearColor GeometrySelectedCol);
+	UE_API void SetSelectionColors(FLinearColor UnselectedCol, FLinearColor HoverOverSelectedCol, FLinearColor HoverOverUnselectedCol, FLinearColor GeometrySelectedCol, FLinearColor SoftSelectedCol);
 
 	/** Disconnect and cleanup for PreviewGeometry object. */
 	UE_API void DisconnectPreviewGeometry();
@@ -462,6 +482,19 @@ public:
 	// Selection Filters
 	UE_API void SetHitBackFaces(bool bInHitBackFaces);
 	bool GetHitBackFaces() const { return bHitBackFaces; }
+
+	// Soft Selection
+	UE_EXPERIMENTAL(5.8, "Soft Selection is experimental while the API is under active development.")
+	UE_API void SetSoftSelectionEnabled(bool bEnable);
+
+	UE_EXPERIMENTAL(5.8, "Soft Selection is experimental while the API is under active development.")
+	bool IsSoftSelectionEnabled() const { return bEnableSoftSelection; };
+
+	UE_EXPERIMENTAL(5.8, "Soft Selection is experimental while the API is under active development.")
+	UE_API void SetSoftSelectionRadius(double Radius);
+
+	UE_EXPERIMENTAL(5.8, "Soft Selection is experimental while the API is under active development.")
+	double GetSoftSelectionRadius() const {return SoftSelectionRadius; }
 	
 protected:
 	bool bHitBackFaces = true;
@@ -548,6 +581,7 @@ protected:
 	inline static TStaticArray<FString, 3> HoverOverSelectedSetIds = {  "HoverSelected_Vertices", "HoverSelected_Lines", "HoverSelected_Triangles"};
 	inline static TStaticArray<FString, 3> HoverOverUnselectedSetIds = { "HoverUnselected_Vertices", "HoverUnselected_Lines", "HoverUnselected_Triangles" };
 	inline static TStaticArray<FString, 3> SelectedSetIds = {"Selected_Vertices", "Selected_Lines", "Selected_Triangles"};
+	inline static TStaticArray<FString, 3> SoftSelectedSetIds = {"SoftSelected_Vertices", "SoftSelected_Lines", "SoftSelected_Triangles"};
 
 	UPROPERTY()
 	FMeshElementSelectionParams UnselectedParams {UnselectedSetIds, 5.f, 2.f, 8.f};
@@ -557,7 +591,8 @@ protected:
 	FMeshElementSelectionParams HoverOverUnselectedParams {HoverOverUnselectedSetIds, 10.f, 6.f, 10.f};
 	UPROPERTY()
 	FMeshElementSelectionParams SelectedParams {SelectedSetIds, 6.f, 6.f, 10.f};
-	
+	UPROPERTY()
+	FMeshElementSelectionParams SoftSelectedParams {SoftSelectedSetIds, 5.5f, 2.f, 10.f};
 
 	
 	// various change types need internal access
@@ -591,6 +626,9 @@ private:
 	
 	UE_API void RebuildPreviewRenderCache();
 
+	TArray<FGeometrySelectionElements> CachedSoftSelectionRenderElements;
+	UE_API void RebuildSoftSelectionRenderCache();
+
 	EEnumerateRenderCachesDirtyFlags RenderCachesDirtyFlags = EEnumerateRenderCachesDirtyFlags::None;
 
 	UE_API void RemoveAllSets() const;
@@ -615,6 +653,13 @@ private:
 	};
 	FSavedSelection SavedSelection;
 
+
+	void UpdateSoftSelectionOnSelectionModified();
+
+	// Soft Selection
+	bool bEnableSoftSelection = false;
+
+	double SoftSelectionRadius = 20.0;
 };
 
 
diff --git a/Engine/Plugins/Runtime/MeshModelingToolset/Source/ModelingComponents/Public/Selection/GeometrySelector.h b/Engine/Plugins/Runtime/MeshModelingToolset/Source/ModelingComponents/Public/Selection/GeometrySelector.h
--- a/Engine/Plugins/Runtime/MeshModelingToolset/Source/ModelingComponents/Public/Selection/GeometrySelector.h
+++ b/Engine/Plugins/Runtime/MeshModelingToolset/Source/ModelingComponents/Public/Selection/GeometrySelector.h
@@ -181,8 +181,24 @@ public:
 	 * The caller provides a VertexTransformFunc callback, which will only be called during
 	 * the call to UpdateTransform to compute new 3D positions for 3D points. 
 	 */
-	virtual void UpdateTransform( TFunctionRef<FVector3d(int VertexID, const FVector3d& InitialPosition, const FTransform& WorldTransform)> VertexTransformFunc ) = 0;
+	UE_DEPRECATED(5.8, "Use the UpdateTransformation function that takes a lambda that includes per-vertex weight instead")
+	virtual void UpdateTransform( TFunctionRef<FVector3d(int VertexID, const FVector3d& InitialPosition, const FTransform& WorldTransform)> VertexTransformFunc )
+	{
+		auto VertexTransformFuncWithUnusedWeight = [&](int32 VertexID, double Weight, const FVector3d& Position, const FTransform& Transform)
+			{
+				return VertexTransformFunc(VertexID, Position, Transform);
+			};
+		UpdateTransform(VertexTransformFuncWithUnusedWeight);
+	};
 
+	/**
+	 * Update the active transform (eg called on each mouse update during 3D gizmo usage).
+	 * The caller provides a VertexTransformFunc callback, which will only be called during
+	 * the call to UpdateTransform to compute new 3D positions for 3D points. 
+	 */
+	UE_EXPERIMENTAL(5.8, "Per-vertex weighted transformation is experiemental while the API is under active development.")
+	virtual void UpdateTransform( TFunctionRef<FVector3d(int VertexID, double Weight, const FVector3d& InitialPosition, const FTransform& WorldTransform)> VertexTransformFunc ) = 0;
+	
 	/**
 	 * Optional function called per-frame to visualize the active Transformation
 	 */
@@ -402,6 +418,32 @@ public:
 	) = 0;
 
 
+	UE_EXPERIMENTAL(5.8, "Soft Selection is experimental while the API is under active development.")
+	struct FUpdateSoftSelectionConfig
+	{
+		// When false, this serves as a hint to the call that cached data used to build the selection might be reusable,
+		// since the base selection didn't change
+		bool bBaseSelectionChanged = true;
+		
+		// When false, the update call just clears soft selection.
+		bool bEnabled = false;
+
+		// Geodesic-like distance to expand (currently actually not very even due to reliance on mesh topology)
+		double Radius = 20.0;
+	};
+	
+	/**
+	 * Generate the soft selection based on the input selection.
+	 */
+	UE_EXPERIMENTAL(5.8, "Soft Selection is experimental while the API is under active development.")
+	virtual void UpdateSoftSelection(FGeometrySelection& InSelection, const FUpdateSoftSelectionConfig& InSoftSelectionConfig) = 0;
+
+	/**
+	 * Converts selection from one type to another type
+	 */
+	UE_EXPERIMENTAL(5.8, "This function is experimental while the API is under active development.")
+	virtual bool ConvertSelection(const FGeometrySelection& FromSelectionIn, FGeometrySelection& ToSelectionOut, const UE::Geometry::EEnumerateSelectionConversionParams ConversionParams) = 0;
+
 	/**
 	 * @return the World transform for the Selector's target object
 	 */
diff --git a/Engine/Source/Editor/Experimental/EditorInteractiveToolsFramework/Public/EditorDragTools/DragToolInteraction.h b/Engine/Source/Editor/Experimental/EditorInteractiveToolsFramework/Public/EditorDragTools/DragToolInteraction.h
--- a/Engine/Source/Editor/Experimental/EditorInteractiveToolsFramework/Public/EditorDragTools/DragToolInteraction.h
+++ b/Engine/Source/Editor/Experimental/EditorInteractiveToolsFramework/Public/EditorDragTools/DragToolInteraction.h
@@ -53,7 +53,7 @@ class UDragToolInteraction : public UViewportInteraction, public IViewportClickD
 	GENERATED_BODY()
 
 public:
-	UDragToolInteraction();
+	UE_API UDragToolInteraction();
 
 	UE_API virtual void Initialize(UViewportInteractionsBehaviorSource* const InViewportInteractionsBehaviorSource) override;
 
@@ -72,8 +72,8 @@ public:
 		return FInputRayHit();
 	}
 
-	UE_API virtual void OnDragStart(const FInputDeviceRay& InDragStartPos) override {}
-	UE_API virtual void OnDrag(const FInputDeviceRay& InDragPos) override {}
+	virtual void OnDragStart(const FInputDeviceRay& InDragStartPos) override {}
+	virtual void OnDrag(const FInputDeviceRay& InDragPos) override {}
 	UE_API virtual void OnDragEnd(const FInputDeviceRay& InDragEndPos) override;
 	UE_API virtual void OnForceEndCapture() override;
 	UE_API virtual void OnTerminateDragSequence() override;
diff --git a/Engine/Source/Editor/Experimental/EditorInteractiveToolsFramework/Public/ViewportInteractions/ViewportInteractionsBehaviorSource.h b/Engine/Source/Editor/Experimental/EditorInteractiveToolsFramework/Public/ViewportInteractions/ViewportInteractionsBehaviorSource.h
--- a/Engine/Source/Editor/Experimental/EditorInteractiveToolsFramework/Public/ViewportInteractions/ViewportInteractionsBehaviorSource.h
+++ b/Engine/Source/Editor/Experimental/EditorInteractiveToolsFramework/Public/ViewportInteractions/ViewportInteractionsBehaviorSource.h
@@ -208,7 +208,7 @@ public:
 	/**
 	 * Sets current mouse cursor override
 	 */
-	void SetMouseCursorOverride(EMouseCursor::Type InMouseCursor);
+	EDITORINTERACTIVETOOLSFRAMEWORK_API void SetMouseCursorOverride(EMouseCursor::Type InMouseCursor);
 
 	void ClearMouseCursorOverride();
 
diff --git a/Engine/Source/Runtime/GeometryCore/Public/Parameterization/MeshDijkstra.h b/Engine/Source/Runtime/GeometryCore/Public/Parameterization/MeshDijkstra.h
--- a/Engine/Source/Runtime/GeometryCore/Public/Parameterization/MeshDijkstra.h
+++ b/Engine/Source/Runtime/GeometryCore/Public/Parameterization/MeshDijkstra.h
@@ -311,6 +311,32 @@ public:
 		return true;
 	}
 
+	struct FComputedPoint
+	{
+		int32 PointID = INDEX_NONE;
+		int32 SeedIndex = INDEX_NONE;
+		double GraphDistance = 0;
+		bool bSeedPoint = false;
+	};
+
+	
+	void EnumerateComputedPoint(TFunction<void(FComputedPoint)> InFunction)
+	{
+		for (const FGraphNode& Node : AllocatedNodes)
+		{
+			if (Node.bFrozen)
+			{
+				FComputedPoint PointInfo;
+				PointInfo.PointID = Node.PointID;
+				PointInfo.SeedIndex = Node.SeedPointID;
+				PointInfo.GraphDistance = Node.GraphDistance;
+				PointInfo.bSeedPoint = Node.ParentPointID == INDEX_NONE;
+				
+				InFunction(MoveTemp(PointInfo));
+			}
+		}	
+	}
+	
 private:
 
 	// information about each active/computed point
