[GeometryCore] added experiemtal warning guard

#Synced-CL 50430604

diff --git a/Engine/Plugins/Runtime/MeshModelingToolset/Source/ModelingComponents/Public/Selection/GeometrySelector.h b/Engine/Plugins/Runtime/MeshModelingToolset/Source/ModelingComponents/Public/Selection/GeometrySelector.h
--- a/Engine/Plugins/Runtime/MeshModelingToolset/Source/ModelingComponents/Public/Selection/GeometrySelector.h
+++ b/Engine/Plugins/Runtime/MeshModelingToolset/Source/ModelingComponents/Public/Selection/GeometrySelector.h
@@ -188,7 +188,9 @@ public:
 			{
 				return VertexTransformFunc(VertexID, Position, Transform);
 			};
+		PRAGMA_DISABLE_EXPERIMENTAL_WARNINGS
 		UpdateTransform(VertexTransformFuncWithUnusedWeight);
+		PRAGMA_ENABLE_EXPERIMENTAL_WARNINGS
 	};
 
 	/**
@@ -196,7 +198,7 @@ public:
 	 * The caller provides a VertexTransformFunc callback, which will only be called during
 	 * the call to UpdateTransform to compute new 3D positions for 3D points. 
 	 */
-	UE_EXPERIMENTAL(5.8, "Per-vertex weighted transformation is experiemental while the API is under active development.")
+	UE_EXPERIMENTAL(5.8, "Per-vertex weighted transformation is experimental while the API is under active development.")
 	virtual void UpdateTransform( TFunctionRef<FVector3d(int VertexID, double Weight, const FVector3d& InitialPosition, const FTransform& WorldTransform)> VertexTransformFunc ) = 0;
 	
 	/**
