[Geometry] Fixed incorect use of experimental tag

#Synced-CL 50430604

diff --git a/Engine/Plugins/Runtime/MeshModelingToolset/Source/ModelingComponents/Public/Selection/GeometrySelector.h b/Engine/Plugins/Runtime/MeshModelingToolset/Source/ModelingComponents/Public/Selection/GeometrySelector.h
--- a/Engine/Plugins/Runtime/MeshModelingToolset/Source/ModelingComponents/Public/Selection/GeometrySelector.h
+++ b/Engine/Plugins/Runtime/MeshModelingToolset/Source/ModelingComponents/Public/Selection/GeometrySelector.h
@@ -420,7 +420,7 @@ public:
 	) = 0;
 
 
-	UE_EXPERIMENTAL(5.8, "Soft Selection is experimental while the API is under active development.")
+	// Experimental 5.8
 	struct FUpdateSoftSelectionConfig
 	{
 		// When false, this serves as a hint to the call that cached data used to build the selection might be reusable,
