Dataflow:
- Added AddPoints(const TArray<FRenderablePoint>& InPoints) to the PointSetComponent

#Synced-CL 50527124

diff --git a/Engine/Plugins/Runtime/MeshModelingToolset/Source/ModelingComponents/Private/Drawing/PointSetComponent.cpp b/Engine/Plugins/Runtime/MeshModelingToolset/Source/ModelingComponents/Private/Drawing/PointSetComponent.cpp
--- a/Engine/Plugins/Runtime/MeshModelingToolset/Source/ModelingComponents/Private/Drawing/PointSetComponent.cpp
+++ b/Engine/Plugins/Runtime/MeshModelingToolset/Source/ModelingComponents/Private/Drawing/PointSetComponent.cpp
@@ -300,6 +300,24 @@ int32 UPointSetComponent::AddPoints(
 	return NumPoints;
 }
 
+int32 UPointSetComponent::AddPoints(const TArray<FRenderablePoint>& InPoints)
+{
+	const int NumPoints = InPoints.Num();
+	if (NumPoints == 0)
+	{
+		return 0;
+	}
+
+	MarkRenderStateDirty();
+	bBoundsDirty = true;
+
+	for (int32 Idx = 0; Idx < NumPoints; ++Idx)
+	{
+		Points.Add(InPoints[Idx]);
+	}
+
+	return NumPoints;
+}
 
 void UPointSetComponent::InsertPoint(const int32 ID, const FRenderablePoint& OverlayPoint)
 {
diff --git a/Engine/Plugins/Runtime/MeshModelingToolset/Source/ModelingComponents/Public/Drawing/PointSetComponent.h b/Engine/Plugins/Runtime/MeshModelingToolset/Source/ModelingComponents/Public/Drawing/PointSetComponent.h
--- a/Engine/Plugins/Runtime/MeshModelingToolset/Source/ModelingComponents/Public/Drawing/PointSetComponent.h
+++ b/Engine/Plugins/Runtime/MeshModelingToolset/Source/ModelingComponents/Public/Drawing/PointSetComponent.h
@@ -75,6 +75,8 @@ public:
 		return AddPoint(FRenderablePoint(InPosition, InColor, InSize, InDepthBias));
 	}
 
+	UE_API int32 AddPoints(const TArray<FRenderablePoint>& InPoints);
+
 	/** 
 	 * Add points to be rendered using the component. 
 	 * @return Number of added points
