[SkeletalMeshModellingTools] Fix switching between the isolated and the full mesh now retains selection when in vertex or edge selection mode

#Synced-CL 50647805

diff --git a/Engine/Plugins/Runtime/MeshModelingToolset/Source/MeshModelingToolsEditorOnly/Private/SkeletalMesh/SkinWeightsPaintTool.cpp b/Engine/Plugins/Runtime/MeshModelingToolset/Source/MeshModelingToolsEditorOnly/Private/SkeletalMesh/SkinWeightsPaintTool.cpp
--- a/Engine/Plugins/Runtime/MeshModelingToolset/Source/MeshModelingToolsEditorOnly/Private/SkeletalMesh/SkinWeightsPaintTool.cpp
+++ b/Engine/Plugins/Runtime/MeshModelingToolset/Source/MeshModelingToolsEditorOnly/Private/SkeletalMesh/SkinWeightsPaintTool.cpp
@@ -4604,8 +4604,141 @@ void UWeightToolSelectionIsolator::SetTrianglesToIsolate(const TArray<int32>& Tr
 	bIsolatedMeshNeedsUpdated = true;
 }
 
+namespace UE::MeshModelingToolsEditorOnly
+{
+	enum class EWeightToolScopedRetainSelectionDirection : uint8
+	{
+		RestoreInPartialMesh,
+		RestoreInFullMesh
+	};
+
+	/** Util to retain the selection when isolating or unisolating the selection */
+	struct FWeightToolScopedRetainSelection
+		: FNoncopyable
+	{
+		FWeightToolScopedRetainSelection(
+			USkinWeightsPaintTool& InWeightTool,
+			const UE::Geometry::FDynamicSubmesh3& Mesh,
+			const EWeightToolScopedRetainSelectionDirection Direction)
+		{
+			using namespace UE::Geometry;
+
+			UToolMeshSelector* MeshSelector = InWeightTool.GetMainMeshSelector();
+			SelectionMechanic = MeshSelector ? MeshSelector->GetSelectionMechanic() : nullptr;
+			const USkinWeightsPaintToolProperties* ToolProperties = InWeightTool.GetWeightToolProperties();
+
+			if (ensure(ToolProperties && SelectionMechanic))
+			{
+				const EComponentSelectionMode SelectionMode = ToolProperties->ComponentSelectionMode;
+
+				const EGeometryElementType ElementType = [SelectionMode]()
+					{
+						if (SelectionMode == EComponentSelectionMode::Vertices)
+						{
+							return EGeometryElementType::Vertex;
+						}
+						else if (SelectionMode == EComponentSelectionMode::Edges)
+						{
+							return EGeometryElementType::Edge;
+						}
+						else if (SelectionMode == EComponentSelectionMode::Faces)
+						{
+							return EGeometryElementType::Face;
+						}
+						else
+						{
+							ensureMsgf(0, TEXT("Unhandled selection mode"));
+							return EGeometryElementType::Vertex;
+						}
+					}();
+
+				UE::Geometry::FGeometrySelection OldSelection;
+
+				OldSelection.TopologyType = EGeometryTopologyType::Triangle;
+				NewSelection.TopologyType = EGeometryTopologyType::Triangle;
+
+				OldSelection.ElementType = ElementType;
+				NewSelection.ElementType = ElementType;
+
+				SelectionMechanic->GetSelection_AsTriangleTopology(OldSelection);
+
+				if (SelectionMode == EComponentSelectionMode::Vertices)
+				{
+					Algo::Transform(OldSelection.Selection, NewSelection.Selection,
+						[&Mesh, Direction](const int32 ElementID)
+						{
+							if (Direction == EWeightToolScopedRetainSelectionDirection::RestoreInFullMesh)
+							{
+								return Mesh.MapVertexToBaseMesh(ElementID);
+							}
+							else
+							{
+								return Mesh.MapVertexToSubmesh(ElementID);
+							}
+						});
+				}
+				else if (SelectionMode == EComponentSelectionMode::Edges)
+				{
+					for (const int32 EdgeID : OldSelection.Selection)
+					{
+						if (Direction == EWeightToolScopedRetainSelectionDirection::RestoreInFullMesh)
+						{
+							UE::Geometry::FMeshTriEdgeID TriEdgeID(EdgeID);							
+							TriEdgeID.TriangleID = Mesh.MapTriangleToBaseMesh(TriEdgeID.TriangleID);
+
+							NewSelection.Selection.Add(TriEdgeID.Encoded());
+						}
+						else
+						{
+							UE::Geometry::FMeshTriEdgeID TriEdgeID(EdgeID);
+							TriEdgeID.TriangleID = Mesh.MapTriangleToSubmesh(TriEdgeID.TriangleID);
+
+							NewSelection.Selection.Add(TriEdgeID.Encoded());
+						}
+					}
+				}
+				else if (SelectionMode == EComponentSelectionMode::Faces)
+				{
+					Algo::Transform(OldSelection.Selection, NewSelection.Selection,
+						[&Mesh, Direction](const int32 Index)
+						{
+							if (Direction == EWeightToolScopedRetainSelectionDirection::RestoreInFullMesh)
+							{
+								return Mesh.MapTriangleToBaseMesh(Index);
+							}
+							else
+							{
+								return Mesh.MapTriangleToSubmesh(Index);
+							}
+						});
+				}
+				else
+				{
+					ensureMsgf(0, TEXT("Unhandled selection mode"));
+				}
+			}
+		}
+
+		~FWeightToolScopedRetainSelection()
+		{
+			if (SelectionMechanic)
+			{					
+				// Restore selection
+				SelectionMechanic->SetSelection_AsTriangleTopology(NewSelection);
+			}
+		}
+
+	private:
+		UE::Geometry::FGeometrySelection NewSelection;
+
+		UPolygonSelectionMechanic* SelectionMechanic = nullptr;
+	};
+}
+
 void UWeightToolSelectionIsolator::RestoreFullMesh()
 {
+	using namespace UE::MeshModelingToolsEditorOnly;
+
 	if (!IsSelectionIsolated())
 	{
 		// nothing hidden
@@ -4625,42 +4758,8 @@ void UWeightToolSelectionIsolator::RestoreFullMesh()
 	}
 
 	// retain selection (a convenience for the user that allows for easily adjusting the isolation by going back/forth as needed)
-	struct FScopedRetainSelection
-	{
-		FScopedRetainSelection(USkinWeightsPaintTool& InWeightTool, const UE::Geometry::FDynamicSubmesh3& PartialSubMesh)
-		{
-			UToolMeshSelector* MeshSelector = InWeightTool.GetMainMeshSelector();
-			SelectionMechanic = MeshSelector ? MeshSelector->GetSelectionMechanic() : nullptr;
-			if (ensure(SelectionMechanic))
-			{
-				UE::Geometry::FGeometrySelection SelectionInParitalMesh;
-				SelectionMechanic->GetSelection_AsTriangleTopology(SelectionInParitalMesh);
-
-				SelectionInFullMesh = SelectionInParitalMesh;
-				SelectionInFullMesh.Selection.Reset();
-
-				Algo::Transform(SelectionInParitalMesh.Selection, SelectionInFullMesh.Selection,
-					[&PartialSubMesh](const int32 PartialIndex)
-					{
-						return PartialSubMesh.MapTriangleToBaseMesh(PartialIndex);
-					});
-			}
-		}
-
-		~FScopedRetainSelection()
-		{
-			if (SelectionMechanic)
-			{
-				SelectionMechanic->SetSelection_AsTriangleTopology(SelectionInFullMesh);
-			}
-		}
-
-		private:			
-			UE::Geometry::FGeometrySelection SelectionInFullMesh;
-			UPolygonSelectionMechanic* SelectionMechanic = nullptr;
-	};
-	const FScopedRetainSelection RetainSeletionScope(*WeightTool, PartialSubMesh);
-
+	const FWeightToolScopedRetainSelection RetainSeletionScope(*WeightTool, PartialSubMesh, EWeightToolScopedRetainSelectionDirection::RestoreInFullMesh);
+ 
 	// reinitialize with full mesh
 	// this resizes Weights to the full mesh size
 	WeightTool->UpdateCurrentlyEditedMesh(FullMesh);
@@ -4697,6 +4796,8 @@ const FDynamicMesh3& UWeightToolSelectionIsolator::GetPartialMesh() const
 
 void UWeightToolSelectionIsolator::CreatePartialMesh()
 {
+	using namespace UE::MeshModelingToolsEditorOnly;
+
 	UPolygonSelectionMechanic* SelectionMechanic = WeightTool->GetMainMeshSelector()->GetSelectionMechanic();
 	if (!ensure(SelectionMechanic))
 	{
@@ -4721,19 +4822,11 @@ void UWeightToolSelectionIsolator::CreatePartialMesh()
 	// create a partial sub-mesh from a subset of triangles on the dynamic mesh
 	PartialSubMesh = UE::Geometry::FDynamicSubmesh3(&FullMesh, CurrentlyIsolatedTriangles);
 
-	// reinitialize all mesh data structures
-	WeightTool->UpdateCurrentlyEditedMesh(PartialSubMesh.GetSubmesh());
-	
 	// retain selection (a convenience for the user that allows for easily adjusting the isolation by going back/forth as needed)
-	UE::Geometry::FGeometrySelection SelectionInParitalMesh;
-	SelectionInParitalMesh.ElementType = UE::Geometry::EGeometryElementType::Face;
-	SelectionInParitalMesh.TopologyType = UE::Geometry::EGeometryTopologyType::Triangle;
-	for (int32 TriangleID : PartialSubMesh.GetSubmesh().TriangleIndicesItr())
-	{
-		SelectionInParitalMesh.Selection.Add(TriangleID);
-	}	
+	const FWeightToolScopedRetainSelection RetainSeletionScope(*WeightTool, PartialSubMesh, EWeightToolScopedRetainSelectionDirection::RestoreInPartialMesh);
 
-	SelectionMechanic->SetSelection_AsTriangleTopology(SelectionInParitalMesh);
+	// reinitialize all mesh data structures
+	WeightTool->UpdateCurrentlyEditedMesh(PartialSubMesh.GetSubmesh());
 }
 
 
