[SkeletalMeshEditingTools] Fix skeleton tree didn't reflect which bones are weighted after changing weights

#Synced-CL 50647805

diff --git a/Engine/Plugins/Runtime/MeshModelingToolset/Source/MeshModelingToolsEditorOnly/Private/SkeletalMesh/SkinWeightsPaintTool.cpp b/Engine/Plugins/Runtime/MeshModelingToolset/Source/MeshModelingToolsEditorOnly/Private/SkeletalMesh/SkinWeightsPaintTool.cpp
--- a/Engine/Plugins/Runtime/MeshModelingToolset/Source/MeshModelingToolsEditorOnly/Private/SkeletalMesh/SkinWeightsPaintTool.cpp
+++ b/Engine/Plugins/Runtime/MeshModelingToolset/Source/MeshModelingToolsEditorOnly/Private/SkeletalMesh/SkinWeightsPaintTool.cpp
@@ -3101,6 +3101,9 @@ void USkinWeightsPaintTool::OnShutdown(EToolShutdownType ShutdownType)
 		DynamicMeshCommitter->CommitDynamicMesh(EditedMesh);
 		
 		GetToolManager()->EndUndoTransaction();
+	
+		// Refresh the skeleton tree
+		GetNotifier()->Notify({}, ESkeletalMeshNotifyType::HierarchyChanged);
 	}
 
 	// restore viewport show flags and preview settings
