[SkeletalMeshEditingTools] Orient bones now orients bones with multiple children in parent orientation

#Synced-CL 50647805

diff --git a/Engine/Plugins/Runtime/MeshModelingToolset/Source/SkeletalMeshModifiers/Private/SkeletonModifier.cpp b/Engine/Plugins/Runtime/MeshModelingToolset/Source/SkeletalMeshModifiers/Private/SkeletonModifier.cpp
--- a/Engine/Plugins/Runtime/MeshModelingToolset/Source/SkeletalMeshModifiers/Private/SkeletonModifier.cpp
+++ b/Engine/Plugins/Runtime/MeshModelingToolset/Source/SkeletalMeshModifiers/Private/SkeletonModifier.cpp
@@ -1923,17 +1923,14 @@ bool USkeletonModifier::OrientBones(const TArray<FName>& InBoneNames, const FOri
 		
 		const int32 ParentIndex = BoneIndex != INDEX_NONE ? ReferenceSkeleton->GetRawParentIndex(BoneIndex) : INDEX_NONE;
 
-		TArray<int32> Children; ReferenceSkeleton->GetRawDirectChildBones(BoneIndex, Children);
-		const int32 NumChildren = Children.Num();
-		if (NumChildren > 1)
-		{ // we can't align if there are more than one children
-			return BoneGlobal;
-		}
-
 		const FTransform& ParentGlobal = TransformComposer->GetGlobalTransform(ParentIndex);
 		FVector Direction = (BoneGlobal.GetLocation() - ParentGlobal.GetLocation()).GetSafeNormal();
-
-		if (NumChildren > 0)
+			
+		TArray<int32> Children; 
+		ReferenceSkeleton->GetRawDirectChildBones(BoneIndex, Children);
+	
+		if (InOptions.bOrientChildren &&
+			Children.Num() == 1)
 		{
 			const FTransform& ChildGlobal = TransformComposer->GetGlobalTransform(Children[0]);
 			Direction = (ChildGlobal.GetLocation() - BoneGlobal.GetLocation()).GetSafeNormal();
@@ -1952,7 +1949,8 @@ bool USkeletonModifier::OrientBones(const TArray<FName>& InBoneNames, const FOri
 			{
 				const bool bComputePlane = (FMath::Abs(FVector::DotProduct(Direction, SecondaryDirection)) - 1.0f) < KINDA_SMALL_NUMBER;
 				if (bComputePlane)
-				{ // use the plane normal as the secondary target, otherwise use InOptions' SecondaryTarget
+				{ 
+					// use the plane normal as the secondary target, otherwise use InOptions' SecondaryTarget
 					FOrientOptions Options = InOptions;
 					Options.SecondaryTarget = FVector::CrossProduct(Direction, SecondaryDirection);
 					return Options.OrientTransform(Direction, BoneGlobal);
