fix flatten brush speed having extra rad factor

#Synced-CL 50981087

diff --git a/Engine/Plugins/Runtime/MeshModelingToolset/Source/MeshModelingTools/Public/Sculpting/MeshSmoothingBrushOps.h b/Engine/Plugins/Runtime/MeshModelingToolset/Source/MeshModelingTools/Public/Sculpting/MeshSmoothingBrushOps.h
--- a/Engine/Plugins/Runtime/MeshModelingToolset/Source/MeshModelingTools/Public/Sculpting/MeshSmoothingBrushOps.h
+++ b/Engine/Plugins/Runtime/MeshModelingToolset/Source/MeshModelingTools/Public/Sculpting/MeshSmoothingBrushOps.h
@@ -215,7 +215,7 @@ class FFlattenBrushOp : public FMeshSculptBrushOp
 {
 
 public:
-	double BrushSpeedTuning = 1.0;
+	double BrushSpeedTuning = 5.0;
 
 	virtual void ApplyStamp(const FDynamicMesh3* Mesh, const FSculptBrushStamp& Stamp, const TArray<int32>& Vertices, TArray<FVector3d>& NewPositionsOut) override
 	{
@@ -224,7 +224,10 @@ public:
 
 		const FVector3d& StampPos = Stamp.LocalFrame.Origin;
 
-		double UseSpeed = Stamp.Power * Stamp.Radius * Stamp.DeltaTime * BrushSpeedTuning;
+		// Note that (unlike w/ other brushes) we don't multiply by a Stamp.Radius factor here, 
+		// because we make speed proportional to distance from target plane, which already gives an implicit ~Stamp.Radius factor
+		// (as distance from plane tends to increase proportionally to distance from brush center, for a surface that isn't already plane-aligned)
+		double UseSpeed = Stamp.Power * Stamp.DeltaTime * BrushSpeedTuning;
 		const FFrame3d& FlattenPlane = Stamp.RegionPlane;
 		FVector3d PlaneZ = FlattenPlane.Z();
 
