World Render Capture: Fix washed out baked textures from Approximate HLOD.

* CL 47967648 added display gamma correction to the buffer visualization pipeline. WorldRenderCapture was reading from the display output, receiving gamma-corrected values instead of linear. The baking code then applied its own sRGB encoding on top, resulting in double-encoded textures.
* Switched all non-depth captures to use BufferVisualizationPipes, which tap the material output before display gamma is applied. This is the same mechanism CombinedMRS already used.
* Unified the capture paths into CaptureUsingPostProcessMaterialPipe() and removed the old render-target-readback code path.
* CL <pending!> corrected the output of the Roughness & Subsurface visualization materials, allowing the capture code to be simplified even further.

* Cleanup: removed GammaRenderTexture (no longer needed), removed deprecated CaptureEmissiveFromPosition(), deprecated CaptureMRSFromPosition(), deduplicated MRS material path.

#Synced-CL 51051335

diff --git a/Engine/Plugins/Runtime/MeshModelingToolset/Source/ModelingComponents/Private/Scene/WorldRenderCapture.cpp b/Engine/Plugins/Runtime/MeshModelingToolset/Source/ModelingComponents/Private/Scene/WorldRenderCapture.cpp
--- a/Engine/Plugins/Runtime/MeshModelingToolset/Source/ModelingComponents/Private/Scene/WorldRenderCapture.cpp
+++ b/Engine/Plugins/Runtime/MeshModelingToolset/Source/ModelingComponents/Private/Scene/WorldRenderCapture.cpp
@@ -2,6 +2,7 @@
 
 #include "Scene/WorldRenderCapture.h"
 #include "AssetUtils/Texture2DUtil.h"
+#include "BufferVisualizationData.h"
 
 #include "Components/PrimitiveComponent.h"
 #include "Engine/Canvas.h"
@@ -282,47 +283,38 @@ FSphere FWorldRenderCapture::ComputeContainingRenderSphere(float HorzFOVDegrees,
 	return FSphere(VisibleBounds.Origin, TargetDistance);
 }
 
-
-
-
-
-
-
-UTextureRenderTarget2D* FWorldRenderCapture::GetRenderTexture(bool bLinear)
+UTextureRenderTarget2D* FWorldRenderCapture::GetRenderTexture()
 {
 	if (RenderTextureDimensions != this->Dimensions)
 	{
 		LinearRenderTexture = nullptr;
-		GammaRenderTexture = nullptr;
 	}
 
-	TObjectPtr<UTextureRenderTarget2D>* WhichTexture = (bLinear) ? &LinearRenderTexture : &GammaRenderTexture;
-
-	if ( *WhichTexture != nullptr )
+	if (LinearRenderTexture != nullptr)
 	{
-		(*WhichTexture)->UpdateResourceImmediate(true);
+		LinearRenderTexture->UpdateResourceImmediate(true);
 
 		UE::Internal::FlushRHIThreadToUpdateTextureRenderTargetReference();
 
-		return *WhichTexture;
+		return LinearRenderTexture;
 	}
 
 	RenderTextureDimensions = Dimensions;
 	int32 Width = Dimensions.GetWidth();
 	int32 Height = Dimensions.GetHeight();
 
-	*WhichTexture = NewObject<UTextureRenderTarget2D>();
-	if (ensure(*WhichTexture))
+	LinearRenderTexture = NewObject<UTextureRenderTarget2D>();
+	if (ensure(LinearRenderTexture))
 	{
-		(*WhichTexture)->ClearColor = FLinearColor::Transparent;
-		(*WhichTexture)->TargetGamma = (bLinear) ? 1.0f : 2.2f;
-		(*WhichTexture)->InitCustomFormat(Width, Height, PF_FloatRGBA, false);
-		(*WhichTexture)->UpdateResourceImmediate(true);
-		
+		LinearRenderTexture->ClearColor = FLinearColor::Transparent;
+		LinearRenderTexture->TargetGamma = 1.0f;
+		LinearRenderTexture->InitCustomFormat(Width, Height, PF_FloatRGBA, false);
+		LinearRenderTexture->UpdateResourceImmediate(true);
+
 		UE::Internal::FlushRHIThreadToUpdateTextureRenderTargetReference();
 	}
 
-	return *WhichTexture;
+	return LinearRenderTexture;
 }
 
 
@@ -539,19 +531,20 @@ FName GetBufferVisualizationModeName(ERenderCaptureType CaptureType)
 } // end namespace UE
 
 
-bool FWorldRenderCapture::CaptureMRSFromPosition(
+static const TCHAR* MRSPostProcessMaterialAssetPath = TEXT("/MeshModelingToolsetExp/Materials/PostProcess_PackedMRSA.PostProcess_PackedMRSA");
+
+bool FWorldRenderCapture::CaptureUsingPostProcessMaterialPipe(
+	UMaterialInterface* PostProcessMaterial,
 	const FFrame3d& ViewFrame,
 	double HorzFOVDegrees,
 	double NearPlaneDist,
 	FImageAdapter& ResultImageOut,
-	const FRenderCaptureConfig& Config)
+	const FRenderCaptureConfig& Config,
+	bool bEnableLighting,
+	const TCHAR* DebugImageName)
 {
-	TRACE_CPUPROFILER_EVENT_SCOPE(FWorldRenderCapture_CombinedMRS);
+	TRACE_CPUPROFILER_EVENT_SCOPE_TEXT(*FString::Printf(TEXT("FWorldRenderCapture_%s"), DebugImageName));
 
-	// this post-process material renders an image with R=Metallic, G=Roughness, B=Specular, A=AmbientOcclusion
-	FString MRSPostProcessMaterialAssetPath = TEXT("/MeshModelingToolsetExp/Materials/PostProcess_PackedMRSA.PostProcess_PackedMRSA");
-	TSoftObjectPtr<UMaterialInterface> PostProcessMaterialPtr = TSoftObjectPtr<UMaterialInterface>(FSoftObjectPath(MRSPostProcessMaterialAssetPath));
-	UMaterialInterface* PostProcessMaterial = PostProcessMaterialPtr.LoadSynchronous();
 	if (!ensure(PostProcessMaterial))
 	{
 		return false;
@@ -560,7 +553,7 @@ bool FWorldRenderCapture::CaptureMRSFromPosition(
 	check(this->World);
 	FSceneInterface* Scene = this->World->Scene;
 
-	UTextureRenderTarget2D* RenderTargetTexture = GetRenderTexture(true);
+	UTextureRenderTarget2D* RenderTargetTexture = GetRenderTexture();
 	if (!ensure(RenderTargetTexture))
 	{
 		return false;
@@ -570,6 +563,13 @@ bool FWorldRenderCapture::CaptureMRSFromPosition(
 	FEngineShowFlags ShowFlags(ESFIM_Game);
 	ApplyViewMode(VMI_Unlit, true, ShowFlags);
 	Internal::SetCommonShowFlags(ShowFlags, Config.bAntiAliasing);
+	ShowFlags.SetLighting(bEnableLighting);
+	if (bEnableLighting)
+	{
+		// Emissive capture requires the post-processing pipeline to be active.
+		ShowFlags.SetPostProcessing(true);
+		ShowFlags.SetPostProcessMaterial(true);
+	}
 
 	FSceneViewFamilyContext ViewFamily(FSceneViewFamily::ConstructionValues(
 		RenderTargetResource, Scene, ShowFlags)
@@ -577,12 +577,8 @@ bool FWorldRenderCapture::CaptureMRSFromPosition(
 		.SetRealtimeUpdate(false)
 	);
 
-	//ViewFamily.EngineShowFlags.SetScreenPercentage(false);
 	ViewFamily.SetScreenPercentageInterface(new FLegacyScreenPercentageDriver(ViewFamily, 1.f));
 
-	// This is called in various other places, unclear if we should be doing this too
-	//EngineShowFlagOverride(ESFIM_Game, VMI_Unlit, ViewFamily.EngineShowFlags, true);
-
 	FSceneViewInitOptions ViewInitOptions = Internal::MakeSceneViewInitOptions(
 		ViewFrame,
 		Dimensions,
@@ -642,25 +638,29 @@ bool FWorldRenderCapture::CaptureMRSFromPosition(
 
 	if (bWriteDebugImage)
 	{
-		WriteDebugImage(ResultImageOut, TEXT("CombinedMRS"));
+		WriteDebugImage(ResultImageOut, DebugImageName);
 	}
 
 	return true;
 }
 
 
-
-
-
-bool FWorldRenderCapture::CaptureEmissiveFromPosition(
-	const FFrame3d& Frame,
+bool FWorldRenderCapture::CaptureMRSFromPosition(
+	const FFrame3d& ViewFrame,
 	double HorzFOVDegrees,
 	double NearPlaneDist,
 	FImageAdapter& ResultImageOut,
 	const FRenderCaptureConfig& Config)
 {
-	const FName BufferVisualizationMode = UE::Internal::GetBufferVisualizationModeName(ERenderCaptureType::Emissive);
-	return CaptureBufferVisualizationFromPosition(BufferVisualizationMode, Frame, HorzFOVDegrees, NearPlaneDist, ResultImageOut, Config);
+	TSoftObjectPtr<UMaterialInterface> PostProcessMaterialPtr = TSoftObjectPtr<UMaterialInterface>(FSoftObjectPath(MRSPostProcessMaterialAssetPath));
+	UMaterialInterface* PostProcessMaterial = PostProcessMaterialPtr.LoadSynchronous();
+	if (!ensure(PostProcessMaterial))
+	{
+		return false;
+	}
+
+	const bool bEnableLighting = false;
+	return CaptureUsingPostProcessMaterialPipe(PostProcessMaterial, ViewFrame, HorzFOVDegrees, NearPlaneDist, ResultImageOut, Config, bEnableLighting, TEXT("CombinedMRS"));
 }
 
 
@@ -780,81 +780,19 @@ bool FWorldRenderCapture::CaptureBufferVisualizationFromPosition(
 	FImageAdapter& ResultImageOut,
 	const FRenderCaptureConfig& Config)
 {
-	TRACE_CPUPROFILER_EVENT_SCOPE_TEXT(*FString::Printf(TEXT("FWorldRenderCapture_%s"), *VisualizationMode.ToString()));
-
-	const bool bEmissive = VisualizationMode == Internal::GetBufferVisualizationModeName(ERenderCaptureType::Emissive);
-	const bool bRoughness = VisualizationMode == Internal::GetBufferVisualizationModeName(ERenderCaptureType::Roughness);
-
-	// The following handles buffer visualization materials found in /Engine/BufferVisualization/<VisualizationMode>.
-	// Sometimes these postprocess materials change the raw GBuffer data, in particular:
-	// - The Roughness postprocess material output is:           GBufferRoughness^Gamma,   with Gamma=2.2
-	// - The SubsurfaceColor postprocess material output is:     GBufferSSColor^(1/Gamma), with Gamma=2.2
-	// Note the relationship between linear and gamma encoded values: EncodedValue = LinearValue^(1/Gamma)
-
-	// Roughness visualization is rendered with gamma correction (unclear why)
-	UTextureRenderTarget2D* RenderTargetTexture = GetRenderTexture(!bRoughness);
-	if (ensure(RenderTargetTexture) == false)
+	if (!GetBufferVisualizationData().IsInitialized())
 	{
 		return false;
 	}
-	FTextureRenderTargetResource* RenderTargetResource = RenderTargetTexture->GameThread_GetRenderTargetResource();
-
-	FEngineShowFlags ShowFlags(ESFIM_Game);
-	ApplyViewMode(VMI_VisualizeBuffer, true, ShowFlags);
-	Internal::SetCommonShowFlags(ShowFlags, Config.bAntiAliasing);
-	ShowFlags.SetPostProcessMaterial(true);
-	ShowFlags.SetPostProcessing(true);
-	ShowFlags.SetVisualizeBuffer(true);
-	ShowFlags.SetLighting(bEmissive); // Need this because Emissive is lighting I guess
-
-	FSceneViewFamilyContext ViewFamily(
-		FSceneViewFamily::ConstructionValues(RenderTargetResource, World->Scene, ShowFlags)
-		.SetTime(FGameTime())
-		.SetRealtimeUpdate(false)
-	);
-
-	FSceneViewInitOptions ViewInitOptions = Internal::MakeSceneViewInitOptions(
-		Frame,
-		Dimensions,
-		HorzFOVDegrees,
-		NearPlaneDist,
-		VisiblePrimitives);
-	ViewInitOptions.ViewFamily = &ViewFamily;
-
-	FSceneView* NewView = new FSceneView(ViewInitOptions);
-	ViewFamily.Views.Add(NewView);
-
-	NewView->CurrentBufferVisualizationMode = VisualizationMode;
-
-	ViewFamily.SetScreenPercentageInterface(new FLegacyScreenPercentageDriver(ViewFamily, 1.0f));
-
-	// should we cache the FCanvas?
-	FCanvas Canvas(RenderTargetResource, nullptr, FGameTime(), NewView->GetFeatureLevel());
-	Canvas.Clear(FLinearColor::Transparent);
-
-	UE::Internal::PerformSceneRender(&Canvas, &ViewFamily);
-
-	// Cache the view/projection matricies we used to render the scene
-	PRAGMA_DISABLE_DEPRECATION_WARNINGS
-	LastCaptureViewMatrices = NewView->ViewMatrices;
-	PRAGMA_ENABLE_DEPRECATION_WARNINGS
 
-	// Copy the contents of the remote texture to system memory. For Emissive we dont normalize buffer values, we want full HDR ranges.
-	ReadImageBuffer.Reset();
-	ReadImageBuffer.SetNumUninitialized(Dimensions.Num());
-	FReadSurfaceDataFlags ReadSurfaceDataFlags(bEmissive ? RCM_MinMax : RCM_UNorm); // Should we always use MinMax to avoid normalization?
-	ReadSurfaceDataFlags.SetLinearToGamma(false);
-	RenderTargetResource->ReadLinearColorPixels(ReadImageBuffer, ReadSurfaceDataFlags, FIntRect(0, 0, Dimensions.GetWidth(), Dimensions.GetHeight()));
-
-	FImageView CaptureData(ReadImageBuffer.GetData(), Dimensions.GetWidth(), Dimensions.GetHeight());
-	UE::Internal::CopyCaptureDataToOutputImageFormat(CaptureData, ResultImageOut);
-
-	if (bWriteDebugImage)
+	UMaterialInterface* VisMaterial = GetBufferVisualizationData().GetMaterial(VisualizationMode);
+	if (!ensure(VisMaterial))
 	{
-		WriteDebugImage(ResultImageOut, *VisualizationMode.ToString());
+		return false;
 	}
 
-	return true;
+	const bool bEnableLighting = (VisualizationMode == Internal::GetBufferVisualizationModeName(ERenderCaptureType::Emissive));
+	return CaptureUsingPostProcessMaterialPipe(VisMaterial, Frame, HorzFOVDegrees, NearPlaneDist, ResultImageOut, Config, bEnableLighting, *VisualizationMode.ToString());
 }
 
 
@@ -877,23 +815,37 @@ bool FWorldRenderCapture::CaptureFromPosition(
 	RenderCaptureInterface::FScopedCapture RenderCapture(RenderCaptureDraws > 0, TEXT("RenderCaptureFromPosition"));
 	RenderCaptureDraws--;
 
-	bool bCaptured = false;
-
-	if (CaptureType == ERenderCaptureType::CombinedMRS)
+	if (CaptureType == ERenderCaptureType::DeviceDepth)
 	{
-		bCaptured = CaptureMRSFromPosition(ViewFrame, HorzFOVDegrees, NearPlaneDist, ResultImageOut, Config);
+		return CaptureDeviceDepthFromPosition(ViewFrame, HorzFOVDegrees, NearPlaneDist, ResultImageOut, Config);
 	}
-	else if (CaptureType == ERenderCaptureType::DeviceDepth)
+
+	UMaterialInterface* Material = nullptr;
+	FString DebugName;
+
+	if (CaptureType == ERenderCaptureType::CombinedMRS)
 	{
-		bCaptured = CaptureDeviceDepthFromPosition(ViewFrame, HorzFOVDegrees, NearPlaneDist, ResultImageOut, Config);
+		TSoftObjectPtr<UMaterialInterface> PostProcessMaterialPtr = TSoftObjectPtr<UMaterialInterface>(FSoftObjectPath(MRSPostProcessMaterialAssetPath));
+		Material = PostProcessMaterialPtr.LoadSynchronous();
+		DebugName = TEXT("CombinedMRS");
 	}
 	else
 	{
 		const FName BufferVisualizationMode = UE::Internal::GetBufferVisualizationModeName(CaptureType);
-		bCaptured = CaptureBufferVisualizationFromPosition(BufferVisualizationMode, ViewFrame, HorzFOVDegrees, NearPlaneDist, ResultImageOut, Config);
+		if (GetBufferVisualizationData().IsInitialized())
+		{
+			Material = GetBufferVisualizationData().GetMaterial(BufferVisualizationMode);
+		}
+		DebugName = BufferVisualizationMode.ToString();
+	}
+
+	if (!ensure(Material))
+	{
+		return false;
 	}
 
-	return bCaptured;
+	const bool bEnableLighting = (CaptureType == ERenderCaptureType::Emissive);
+	return CaptureUsingPostProcessMaterialPipe(Material, ViewFrame, HorzFOVDegrees, NearPlaneDist, ResultImageOut, Config, bEnableLighting, *DebugName);
 }
 
 
@@ -934,7 +886,6 @@ void FWorldRenderCapture::AddReferencedObjects(FReferenceCollector& Collector)
 {
 	Collector.AddReferencedObjects(ComponentsRestoreMaxDrawDistance);
 	Collector.AddReferencedObject(LinearRenderTexture);
-	Collector.AddReferencedObject(GammaRenderTexture);
 	Collector.AddReferencedObject(DepthRenderTexture);
 }
 
diff --git a/Engine/Plugins/Runtime/MeshModelingToolset/Source/ModelingComponents/Public/Scene/WorldRenderCapture.h b/Engine/Plugins/Runtime/MeshModelingToolset/Source/ModelingComponents/Public/Scene/WorldRenderCapture.h
--- a/Engine/Plugins/Runtime/MeshModelingToolset/Source/ModelingComponents/Public/Scene/WorldRenderCapture.h
+++ b/Engine/Plugins/Runtime/MeshModelingToolset/Source/ModelingComponents/Public/Scene/WorldRenderCapture.h
@@ -300,8 +300,8 @@ UE_API PRAGMA_ENABLE_DEPRECATION_WARNINGS
 	UE_API void SetEnableWriteDebugImage(bool bEnable, int32 ImageCounter = -1, FString FolderName = FString());
 
 	// ~Begin FGCObject Interface
-	virtual void AddReferencedObjects(FReferenceCollector& Collector) override;
-	virtual FString GetReferencerName() const override;
+	UE_API virtual void AddReferencedObjects(FReferenceCollector& Collector) override;
+	UE_API virtual FString GetReferencerName() const override;
 	// ~End FGCObject Interface
 
 protected:
@@ -314,11 +314,10 @@ protected:
 
 	// Temporary textures used as render targets. We explicitly prevent this from being GC'd internally
 	TObjectPtr<UTextureRenderTarget2D> LinearRenderTexture = nullptr;
-	TObjectPtr<UTextureRenderTarget2D> GammaRenderTexture = nullptr;
 	TObjectPtr<UTextureRenderTarget2D> DepthRenderTexture = nullptr;
 
 	FImageDimensions RenderTextureDimensions;
-	UE_API UTextureRenderTarget2D* GetRenderTexture(bool bLinear);
+	UE_API UTextureRenderTarget2D* GetRenderTexture();
 	UE_API UTextureRenderTarget2D* GetDepthRenderTexture();
 
 	UE_DEPRECATED(5.4, "Please use GetRenderCaptureViewMatrices instead")
@@ -327,16 +326,8 @@ protected:
 	// temporary buffer used to read from texture
 	TArray<FLinearColor> ReadImageBuffer;
 
-	/** Emissive is a special case and uses different code than capture of color/property channels */
-	UE_DEPRECATED(5.4, "Please use CaptureBufferVisualizationFromPosition with VisualizationMode=PreTonemapHDRColor instead")
-	UE_API bool CaptureEmissiveFromPosition(
-		const FFrame3d& Frame,
-		double HorzFOVDegrees,
-		double NearPlaneDist,
-		FImageAdapter& ResultImageOut,
-		const FRenderCaptureConfig& Config = {});
-
 	/** Combined Metallic/Roughness/Specular uses a custom postprocess material */
+	UE_DEPRECATED(5.8, "Please use CaptureFromPosition with ERenderCaptureType::CombinedMRS instead")
 	UE_API bool CaptureMRSFromPosition(
 		const FFrame3d& Frame,
 		double HorzFOVDegrees,
@@ -367,6 +358,20 @@ protected:
 		FImageAdapter& ResultImageOut,
 		const FRenderCaptureConfig& Config = {});
 
+	/**
+	 * Renders the given PostProcessMaterial and captures its output using BufferVisualizationPipes, which grabs the material output before
+	 * any gamma correction is applied
+	 */
+	UE_API bool CaptureUsingPostProcessMaterialPipe(
+		UMaterialInterface* PostProcessMaterial,
+		const FFrame3d& Frame,
+		double HorzFOVDegrees,
+		double NearPlaneDist,
+		FImageAdapter& ResultImageOut,
+		const FRenderCaptureConfig& Config,
+		bool bEnableLighting,
+		const TCHAR* DebugImageName);
+
 	bool bWriteDebugImage = false;
 	int32 DebugImageCounter = -1;
 	FString DebugImageFolderName = TEXT("WorldRenderCapture");
