[Skeleton Editing Tools]
- Convert Select operation to occur on mouse up. Allow for capture stealing in a selection operation so that drag-based camera controls can interrupt.
- Fix an issue where using any LMB ITF interactions including those with a higher priority would cause selection to change or be cleared.
- Fix an issue where pending clicks would occur on Tick() even when the tool did not win capture.

#rnx
#Synced-CL 51094013

diff --git a/Engine/Plugins/Runtime/MeshModelingToolset/Source/MeshModelingToolsEditorOnly/Private/SkeletalMesh/SkeletonEditingTool.cpp b/Engine/Plugins/Runtime/MeshModelingToolset/Source/MeshModelingToolsEditorOnly/Private/SkeletalMesh/SkeletonEditingTool.cpp
--- a/Engine/Plugins/Runtime/MeshModelingToolset/Source/MeshModelingToolsEditorOnly/Private/SkeletalMesh/SkeletonEditingTool.cpp
+++ b/Engine/Plugins/Runtime/MeshModelingToolset/Source/MeshModelingToolsEditorOnly/Private/SkeletalMesh/SkeletonEditingTool.cpp
@@ -237,6 +237,11 @@ void USkeletonEditingTool::SetupBehaviors()
 	ClickDragBehavior->Initialize(this);
 	ClickDragBehavior->Modifiers.RegisterModifier(AddToSelectionModifier, FInputDeviceState::IsShiftKeyDown);
 	ClickDragBehavior->Modifiers.RegisterModifier(ToggleSelectionModifier, FInputDeviceState::IsCtrlKeyDown);
+	
+	ClickDragBehavior->SetAllowsCaptureStealing(TAttribute<bool>::CreateLambda([this]
+	{
+		return Operation == EEditingOperation::Select;
+	}));
 	AddInputBehavior(ClickDragBehavior);
 }
 
@@ -1357,16 +1362,21 @@ FTransform USkeletonEditingTool::ComputeTransformFromComponents(const TArray<int
 
 void USkeletonEditingTool::OnClickPress(const FInputDeviceRay& InPressPos)
 {
-	if (PendingFunction)
+	if (PendingClickFunction)
 	{
-		if (!bDeferUntilFocused)
+		if (bDeferUntilFocused)
+		{
+			DeferredFocusFunction = PendingClickFunction;
+			PendingClickFunction.Reset();
+		}
+		else
 		{
-			PendingFunction();
-			PendingFunction.Reset();
+			PendingClickFunction();
+			PendingClickFunction.Reset();
 		}
 	}
 	else
-	{ // make sure that the PendingFunction handles BeginChange if it needs to
+	{ // make sure that the PendingClickFunction handles BeginChange if it needs to
 		BeginChange();
 	}
 }
@@ -1604,6 +1614,13 @@ void USkeletonEditingTool::NormalizeSelection()
 
 void USkeletonEditingTool::OnClickRelease(const FInputDeviceRay& InReleasePos)
 {
+	if (PendingReleaseFunction)
+	{
+		PendingReleaseFunction();
+		PendingReleaseFunction.Reset();
+		return;
+	}
+
 	if (Operation == EEditingOperation::Create)
 	{
 		return EndChange();
@@ -1620,10 +1637,10 @@ void USkeletonEditingTool::OnTick(float DeltaTime)
 	const FViewport* Viewport = GetToolManager()->GetContextQueriesAPI()->GetFocusedViewport();
 	bDeferUntilFocused = Viewport && !Viewport->HasFocus();
 	
-	if (PendingFunction)
+	if (!bDeferUntilFocused && DeferredFocusFunction)
 	{
-		PendingFunction();
-		PendingFunction.Reset();
+		DeferredFocusFunction();
+		DeferredFocusFunction.Reset();
 	}
 	
 	SelectionWatcher.CheckAndUpdate();
@@ -1643,7 +1660,9 @@ FInputRayHit USkeletonEditingTool::CanBeginClickDragSequence(const FInputDeviceR
 		return InvalidRayHit;
 	}
 	
-	PendingFunction.Reset();
+	PendingClickFunction.Reset();
+	PendingReleaseFunction.Reset();
+	DeferredFocusFunction.Reset();
 	ParentIndex = INDEX_NONE;
 	
 	FViewport* Viewport = GetToolManager()->GetContextQueriesAPI()->GetFocusedViewport();
@@ -1692,7 +1711,7 @@ FInputRayHit USkeletonEditingTool::CanBeginClickDragSequence(const FInputDeviceR
 		}
 		
 		// otherwise, update current selection
-		PendingFunction = [this, BoneIndex]
+		TFunction<void()> SelectFunction = [this, BoneIndex]
 		{
 			const FReferenceSkeleton& ReferenceSkeleton = Modifier->GetReferenceSkeleton();
 			SelectBone(ReferenceSkeleton.GetBoneName(BoneIndex));
@@ -1700,7 +1719,18 @@ FInputRayHit USkeletonEditingTool::CanBeginClickDragSequence(const FInputDeviceR
 			Properties->Name = GetCurrentBone();
 			Properties->Transform = ReferenceSkeleton.GetRefBonePose()[BoneIndex];
 			ParentIndex = ReferenceSkeleton.GetParentIndex(BoneIndex);
-		};
+		}; 
+		
+		if (Operation == EEditingOperation::Select)
+		{
+			// When just doing a select operation, select occurs on mouse up.
+			PendingClickFunction = []{};
+			PendingReleaseFunction = SelectFunction;
+		}
+		else
+		{
+			PendingClickFunction = SelectFunction;
+		}
 	
 		return FInputRayHit(0.0);
 	}
@@ -1708,9 +1738,14 @@ FInputRayHit USkeletonEditingTool::CanBeginClickDragSequence(const FInputDeviceR
 	// if we didn't pick anything
 	if (Operation == EEditingOperation::Select)
 	{
-		Selection.Empty();
-		Properties->Name = GetCurrentBone();
-		return InvalidRayHit;
+		PendingClickFunction = []{ };
+		PendingReleaseFunction = [this]
+		{
+			Selection.Empty();
+			Properties->Name = GetCurrentBone();
+		};
+		
+		return FInputRayHit(0.0);
 	}
 
 	// if we're in creation mode then create a new bone
@@ -1720,7 +1755,7 @@ FInputRayHit USkeletonEditingTool::CanBeginClickDragSequence(const FInputDeviceR
 		if (ProjectionProperties->GetProjectionPoint(InPressPos, HitPoint))
 		{
 			// CurrentBone is gonna be the parent
-			PendingFunction = [this, HitPoint]
+			PendingClickFunction = [this, HitPoint]
 			{
 				const FReferenceSkeleton& ReferenceSkeleton = Modifier->GetReferenceSkeleton();
 				ParentIndex = ReferenceSkeleton.FindRawBoneIndex(GetCurrentBone());
diff --git a/Engine/Plugins/Runtime/MeshModelingToolset/Source/MeshModelingToolsEditorOnly/Public/SkeletalMesh/SkeletonEditingTool.h b/Engine/Plugins/Runtime/MeshModelingToolset/Source/MeshModelingToolsEditorOnly/Public/SkeletalMesh/SkeletonEditingTool.h
--- a/Engine/Plugins/Runtime/MeshModelingToolset/Source/MeshModelingToolsEditorOnly/Public/SkeletalMesh/SkeletonEditingTool.h
+++ b/Engine/Plugins/Runtime/MeshModelingToolset/Source/MeshModelingToolsEditorOnly/Public/SkeletalMesh/SkeletonEditingTool.h
@@ -316,7 +316,9 @@ private:
 	
 	int32 ParentIndex = INDEX_NONE;
 
-	TFunction<void()> PendingFunction;
+	TFunction<void()> PendingClickFunction;
+	TFunction<void()> PendingReleaseFunction;
+	TFunction<void()> DeferredFocusFunction;
 
 	TUniquePtr<UE::Geometry::FTriangleGroupTopology> Topology = nullptr;
 
