fix mesh sculpt layer properties doing a full mesh edit every time it reads a single layer name

#Synced-CL 51103486

diff --git a/Engine/Plugins/Runtime/MeshModelingToolset/Source/MeshModelingTools/Private/Properties/MeshSculptLayerProperties.cpp b/Engine/Plugins/Runtime/MeshModelingToolset/Source/MeshModelingTools/Private/Properties/MeshSculptLayerProperties.cpp
--- a/Engine/Plugins/Runtime/MeshModelingToolset/Source/MeshModelingTools/Private/Properties/MeshSculptLayerProperties.cpp
+++ b/Engine/Plugins/Runtime/MeshModelingToolset/Source/MeshModelingTools/Private/Properties/MeshSculptLayerProperties.cpp
@@ -63,6 +63,32 @@ void UMeshSculptLayerProperties::SetLayerWeights(UE::Geometry::FDynamicMeshSculp
 	SculptLayers->UpdateLayerWeights(FullLayerWeights);
 }
 
+bool UMeshSculptLayerProperties::ReadSculptLayers(TFunctionRef<void(const UE::Geometry::FDynamicMesh3& Mesh, const UE::Geometry::FDynamicMeshSculptLayers* SculptLayers)> ReadFn)
+{
+	bool bFoundLayers = false;
+	if (Tool)
+	{
+		Tool->ProcessToolMeshes([this, &ReadFn, &bFoundLayers](const UE::Geometry::FDynamicMesh3& Mesh, int32 Idx) -> void
+			{
+				// Sculpt layer UI only supports a single mesh for now
+				if (Idx > 0)
+				{
+					return;
+				}
+				if (const UE::Geometry::FDynamicMeshAttributeSet* Attributes = Mesh.Attributes())
+				{
+					if (const UE::Geometry::FDynamicMeshSculptLayers* SculptLayers = Attributes->GetSculptLayers())
+					{
+						ReadFn(Mesh, SculptLayers);
+						bFoundLayers = true;
+					}
+				}
+			}
+		);
+	}
+	return bFoundLayers;
+}
+
 void UMeshSculptLayerProperties::EditSculptLayers(TFunctionRef<void(UE::Geometry::FDynamicMesh3& Mesh, UE::Geometry::FDynamicMeshSculptLayers* SculptLayers)> EditFn, bool bEmitChange)
 {
 	if (Tool && Tool->AllowToolMeshUpdates())
@@ -250,11 +276,14 @@ FName UMeshSculptLayerProperties::GetLayerName(const int32 InIndex)
 	}
 	
 	FName LayerName;
-	EditSculptLayers([this, &InIndex, &LayerName](UE::Geometry::FDynamicMesh3& Mesh, UE::Geometry::FDynamicMeshSculptLayers* SculptLayers)
+	ReadSculptLayers([this, &InIndex, &LayerName](const UE::Geometry::FDynamicMesh3& Mesh, const UE::Geometry::FDynamicMeshSculptLayers* SculptLayers)->void
 		{
-			LayerName = SculptLayers->GetLayer(InIndex + NumLockedBaseLayers)->GetName();
+			if (const UE::Geometry::FDynamicMeshSculptLayerAttribute* Layer = SculptLayers->GetLayer(InIndex + NumLockedBaseLayers))
+			{
+				LayerName = Layer->GetName();
+			}
 			UpdateSettingsFromMesh(SculptLayers);
-		}, false);
+		});
 	return LayerName;
 }
 
diff --git a/Engine/Plugins/Runtime/MeshModelingToolset/Source/MeshModelingTools/Public/Properties/MeshSculptLayerProperties.h b/Engine/Plugins/Runtime/MeshModelingToolset/Source/MeshModelingTools/Public/Properties/MeshSculptLayerProperties.h
--- a/Engine/Plugins/Runtime/MeshModelingToolset/Source/MeshModelingTools/Public/Properties/MeshSculptLayerProperties.h
+++ b/Engine/Plugins/Runtime/MeshModelingToolset/Source/MeshModelingTools/Public/Properties/MeshSculptLayerProperties.h
@@ -87,6 +87,9 @@ private:
 	// @param EditFn The edit to apply if possible
 	// @param bEmitChange Whether to emit a change object along with the edit
 	void EditSculptLayers(TFunctionRef<void(UE::Geometry::FDynamicMesh3& Mesh, UE::Geometry::FDynamicMeshSculptLayers* SculptLayers)> EditFn, bool bEmitChange);
+
+	// Helper to read current sculpt layers if possible
+	bool ReadSculptLayers(TFunctionRef<void(const UE::Geometry::FDynamicMesh3& Mesh, const UE::Geometry::FDynamicMeshSculptLayers* SculptLayers)> ReadFn);
 };
 
 DECLARE_LOG_CATEGORY_EXTERN(LogSculptLayerProps, Warning, All);
\ No newline at end of file
