[SkeletalMeshModellingTools] Mirroring now replaces strings anywhere in the name and supports regex expressions

#Synced-CL 51094013

diff --git a/Engine/Plugins/Runtime/MeshModelingToolset/Source/SkeletalMeshModifiers/Private/SkeletonModifier.cpp b/Engine/Plugins/Runtime/MeshModelingToolset/Source/SkeletalMeshModifiers/Private/SkeletonModifier.cpp
--- a/Engine/Plugins/Runtime/MeshModelingToolset/Source/SkeletalMeshModifiers/Private/SkeletonModifier.cpp
+++ b/Engine/Plugins/Runtime/MeshModelingToolset/Source/SkeletalMeshModifiers/Private/SkeletonModifier.cpp
@@ -1413,12 +1413,23 @@ void USkeletonModifier::GetMirroredNames(
 		return;
 	}
 
-	const FName Left(InOptions.LeftString), Right(InOptions.RightString);
-	
+	const FName LeftName(InOptions.LeftString);
+	const FName RightName(InOptions.RightString);
+
+	const FName LeftFindRegex = *FString::Printf(TEXT("([^}]*)%s([^}]*)"), *InOptions.LeftString);
+	const FName RightReplaceRegex = *FString::Printf(TEXT("$1%s$2"), *InOptions.RightString);
+
+	const FName RightFindRegex = *FString::Printf(TEXT("([^}]*)%s([^}]*)"), *InOptions.RightString);
+	const FName LeftReplaceRegex = *FString::Printf(TEXT("$1%s$2"), *InOptions.LeftString);
+
 	const TArray MirrorFindReplaceExpressions({
-		FMirrorFindReplaceExpression(Left, Right, EMirrorFindReplaceMethod::Suffix),
-		FMirrorFindReplaceExpression(Right, Left, EMirrorFindReplaceMethod::Suffix)
-		});
+		FMirrorFindReplaceExpression(LeftName, RightName, EMirrorFindReplaceMethod::Prefix),
+		FMirrorFindReplaceExpression(RightName, LeftName, EMirrorFindReplaceMethod::Prefix),
+		FMirrorFindReplaceExpression(LeftName, RightName, EMirrorFindReplaceMethod::Suffix),
+		FMirrorFindReplaceExpression(RightName, LeftName, EMirrorFindReplaceMethod::Suffix),
+		FMirrorFindReplaceExpression(LeftFindRegex, RightReplaceRegex, EMirrorFindReplaceMethod::RegularExpression),
+		FMirrorFindReplaceExpression(RightFindRegex, LeftReplaceRegex, EMirrorFindReplaceMethod::RegularExpression)
+	});
 		
 	const TArray<FMeshBoneInfo>& BoneInfos = ReferenceSkeleton->GetRawRefBoneInfo();
 	
