[SkeletalMeshEditingTools] Improve local rotation axis settings to use a default thickness of 0.1

#rnx

#Synced-CL 51156678

diff --git a/Engine/Plugins/Runtime/MeshModelingToolset/Source/MeshModelingToolsEditorOnly/Public/SkeletalMesh/SkeletalMeshModelingToolsEditorSettings.h b/Engine/Plugins/Runtime/MeshModelingToolset/Source/MeshModelingToolsEditorOnly/Public/SkeletalMesh/SkeletalMeshModelingToolsEditorSettings.h
--- a/Engine/Plugins/Runtime/MeshModelingToolset/Source/MeshModelingToolsEditorOnly/Public/SkeletalMesh/SkeletalMeshModelingToolsEditorSettings.h
+++ b/Engine/Plugins/Runtime/MeshModelingToolset/Source/MeshModelingToolsEditorOnly/Public/SkeletalMesh/SkeletalMeshModelingToolsEditorSettings.h
@@ -38,7 +38,7 @@ private:
 	float LocalRotationAxisLength = 0.f;
 	
 	UPROPERTY(config)
-	float LocalRotationAxisThickness = 0.f;
+	float LocalRotationAxisThickness = 0.1f;
 
 	UPROPERTY(config)
 	bool bShowAllLocalRotationAxes = false;
