World Render Capture: Request 32-bit float output for pipe captures.

When OutputFormat is not explicitly set, the post-process material output inherits the scene color texture format, which can vary depending on the project configuration. ReadbackPixelData() does not handle this format, causing the pipe to receive invalid data.

Set bIsExpecting32BitPixelData on the pipe to force PF_A32B32G32R32F output, which is always handled by ReadbackPixelData() and provides full precision for baking.

#Synced-CL 51038270

diff --git a/Engine/Plugins/Runtime/MeshModelingToolset/Source/ModelingComponents/Private/Scene/WorldRenderCapture.cpp b/Engine/Plugins/Runtime/MeshModelingToolset/Source/ModelingComponents/Private/Scene/WorldRenderCapture.cpp
--- a/Engine/Plugins/Runtime/MeshModelingToolset/Source/ModelingComponents/Private/Scene/WorldRenderCapture.cpp
+++ b/Engine/Plugins/Runtime/MeshModelingToolset/Source/ModelingComponents/Private/Scene/WorldRenderCapture.cpp
@@ -606,6 +606,7 @@ bool FWorldRenderCapture::CaptureUsingPostProcessMaterialPipe(
 	// create new FImagePixelPipe and add output lambda that takes ownership of the data
 	NewView->FinalPostProcessSettings.BufferVisualizationPipes.Empty();
 	TSharedPtr<FImagePixelPipe, ESPMode::ThreadSafe> ImagePixelPipe = MakeShared<FImagePixelPipe, ESPMode::ThreadSafe>();
+	ImagePixelPipe->bIsExpecting32BitPixelData = true;
 	ImagePixelPipe->AddEndpoint([&PostProcessPassPixelData](TUniquePtr<FImagePixelData>&& Data) {
 		PostProcessPassPixelData = MoveTemp(Data);
 	});
