[SkeletalMeshEditingTools] Improve copy-pasting and duplicating bones

* Add new options to Paste Local and Paste Global
* Pasting bones now retains the hierarchy of pasted bones
* Pasting bones now pastes copied bones and their children
* Duplicate now retains the hierarchy of duplicated bones

#Synced-CL 51275426

diff --git a/Engine/Plugins/Runtime/MeshModelingToolset/Source/MeshModelingToolsEditorOnly/Private/SkeletalMesh/SReferenceSkeletonTree.cpp b/Engine/Plugins/Runtime/MeshModelingToolset/Source/MeshModelingToolsEditorOnly/Private/SkeletalMesh/SReferenceSkeletonTree.cpp
--- a/Engine/Plugins/Runtime/MeshModelingToolset/Source/MeshModelingToolsEditorOnly/Private/SkeletalMesh/SReferenceSkeletonTree.cpp
+++ b/Engine/Plugins/Runtime/MeshModelingToolset/Source/MeshModelingToolsEditorOnly/Private/SkeletalMesh/SReferenceSkeletonTree.cpp
@@ -502,8 +502,13 @@ void SReferenceSkeletonTree::BindCommands()
 			FExecuteAction::CreateSP(this, &SReferenceSkeletonTree::HandleCopyBones),
 			FCanExecuteAction::CreateSP(this, &SReferenceSkeletonTree::CanCopyBones));
 		
-		CommandList->MapAction(Commands.PasteBones,
-			FExecuteAction::CreateSP(this, &SReferenceSkeletonTree::HandlePasteBones),
+		constexpr bool bPasteUsingGlobalTransforms = true;
+		CommandList->MapAction(Commands.PasteBonesLocal,
+			FExecuteAction::CreateSP(this, &SReferenceSkeletonTree::HandlePasteBones, !bPasteUsingGlobalTransforms),
+			FCanExecuteAction::CreateSP(this, &SReferenceSkeletonTree::CanPasteBones));
+		
+		CommandList->MapAction(Commands.PasteBonesGlobal,
+			FExecuteAction::CreateSP(this, &SReferenceSkeletonTree::HandlePasteBones, bPasteUsingGlobalTransforms),
 			FCanExecuteAction::CreateSP(this, &SReferenceSkeletonTree::CanPasteBones));
 		
 		CommandList->MapAction(Commands.DuplicateBones,
@@ -645,15 +650,15 @@ void SReferenceSkeletonTree::HandleCopyBones() const
 		return;
 	}
 
-	TArray<FName> BoneNames;
-	GetSelectedBoneNames(BoneNames);
+	TArray<FName> SelectedBones;
+	GetSelectedBoneNames(SelectedBones);
 
-	if (BoneNames.IsEmpty())
+	if (SelectedBones.IsEmpty())
 	{
 		return;
 	}
 
-	SkeletonClipboard::CopyToClipboard(*Modifier.Get(), BoneNames);
+	SkeletonClipboard::CopyToClipboard(*Modifier.Get(), SelectedBones);
 }
 
 bool SReferenceSkeletonTree::CanCopyBones() const
@@ -666,8 +671,13 @@ bool SReferenceSkeletonTree::CanCopyBones() const
 	return Modifier.IsValid() && HasSelectedItems();
 }
 
-void SReferenceSkeletonTree::HandlePasteBones()
+void SReferenceSkeletonTree::HandlePasteBones(const bool bUsingGlobalTransforms)
 {
+	if (!ensure(Modifier.IsValid()))
+	{
+		return;
+	}
+
 	TArray<FName> BoneNames;
 	GetSelectedBoneNames(BoneNames);
 
@@ -675,7 +685,7 @@ void SReferenceSkeletonTree::HandlePasteBones()
 
 	BeginChange();
 	
-	const TArray<FName> NewBones = SkeletonClipboard::PasteFromClipboard(*Modifier.Get(), DefaultParent);
+	const TArray<FName> NewBones = SkeletonClipboard::PasteFromClipboard(*Modifier.Get(), DefaultParent, bUsingGlobalTransforms);
 	if (NewBones.IsEmpty())
 	{
 		CancelChange();
@@ -708,16 +718,18 @@ bool SReferenceSkeletonTree::CanPasteBones() const
 
 void SReferenceSkeletonTree::HandleDuplicateBones()
 {
-	HandleCopyBones();
-
-	if (!SkeletonClipboard::IsClipboardValid())
+	if (!ensure(Modifier.IsValid()))
 	{
 		return;
-	}  
+	}
 
 	BeginChange();
+
+	TArray<FName> SelectedBones;
+	GetSelectedBoneNames(SelectedBones);
 	
-	const TArray<FName> NewBones = SkeletonClipboard::PasteFromClipboard(*Modifier.Get(), NAME_None);
+	constexpr bool bUsingGlobalTransforms = false;
+	const TArray<FName> NewBones = SkeletonClipboard::ClipboardDuplicate(*Modifier.Get(), SelectedBones);
 	if (NewBones.IsEmpty())
 	{
 		CancelChange();
@@ -1158,7 +1170,9 @@ TSharedPtr<SWidget> SReferenceSkeletonTree::CreateContextMenu()
 
 			MenuBuilder.BeginSection("CopyPasteBones", LOCTEXT("CopyPasteBonesOperations", "Copy & Paste"));
 
-			MenuBuilder.AddMenuEntry(Commands.PasteBones);
+			MenuBuilder.AddMenuEntry(Commands.CopyBones);
+			MenuBuilder.AddMenuEntry(Commands.PasteBonesLocal);
+			MenuBuilder.AddMenuEntry(Commands.PasteBonesGlobal);
 
 			MenuBuilder.EndSection();
 		}
@@ -1184,7 +1198,8 @@ TSharedPtr<SWidget> SReferenceSkeletonTree::CreateContextMenu()
 			MenuBuilder.BeginSection("CopyPasteBones", LOCTEXT("CopyPasteBonesOperations", "Copy & Paste"));
 			{
 				MenuBuilder.AddMenuEntry(Commands.CopyBones);
-				MenuBuilder.AddMenuEntry(Commands.PasteBones);
+				MenuBuilder.AddMenuEntry(Commands.PasteBonesLocal);
+				MenuBuilder.AddMenuEntry(Commands.PasteBonesGlobal);
 				MenuBuilder.AddMenuEntry(Commands.DuplicateBones);
 			}
 			MenuBuilder.EndSection();
diff --git a/Engine/Plugins/Runtime/MeshModelingToolset/Source/MeshModelingToolsEditorOnly/Private/SkeletalMesh/SkeletonEditingTool.cpp b/Engine/Plugins/Runtime/MeshModelingToolset/Source/MeshModelingToolsEditorOnly/Private/SkeletalMesh/SkeletonEditingTool.cpp
--- a/Engine/Plugins/Runtime/MeshModelingToolset/Source/MeshModelingToolsEditorOnly/Private/SkeletalMesh/SkeletonEditingTool.cpp
+++ b/Engine/Plugins/Runtime/MeshModelingToolset/Source/MeshModelingToolsEditorOnly/Private/SkeletalMesh/SkeletonEditingTool.cpp
@@ -92,8 +92,15 @@ void FSkeletonEditingToolCommands::RegisterCommands()
 	UI_COMMAND(AddSuffix, "Add Suffix", "Add a suffix to the selected bone(s) name(s).", EUserInterfaceActionType::Button, FInputChord());
 	
 	UI_COMMAND(CopyBones, "Copy Bone(s)", "Copy selected bone(s).", EUserInterfaceActionType::Button, FInputChord(EKeys::C, EModifierKey::Control));
-	UI_COMMAND(PasteBones, "Paste Bone(s)", "Paste selected bone(s).", EUserInterfaceActionType::Button, FInputChord(EKeys::V, EModifierKey::Control));
 	UI_COMMAND(DuplicateBones, "Duplicate Bone(s)", "Duplicate selected bone(s).", EUserInterfaceActionType::Button, FInputChord(EKeys::D, EModifierKey::Control));
+
+	UI_COMMAND(PasteBonesLocal, "Paste Bone(s) Local", "Paste selected bone(s) using the local transform of the copied bones.", EUserInterfaceActionType::Button, FInputChord(EKeys::V, EModifierKey::Control));
+	UI_COMMAND(PasteBonesGlobal, "Paste Bone(s) Global", "Paste selected bone(s) using the global transform of the copied bones.", EUserInterfaceActionType::Button, FInputChord(EKeys::V, EModifierKey::Control | EModifierKey::Shift));
+
+	// DEPRECATED 5.8
+	PRAGMA_DISABLE_DEPRECATION_WARNINGS
+	UI_COMMAND(PasteBones, "Paste Bone(s) DEPRECATED", "Paste selected bone(s) DEPRECATED, instead use PasteBonesLocal.", EUserInterfaceActionType::Button, FInputChord(EKeys::V, EModifierKey::Control));
+	PRAGMA_ENABLE_DEPRECATION_WARNINGS
 }
 
 
@@ -521,6 +528,7 @@ void USkeletonEditingTool::RegisterActions(FInteractiveToolActionSet& InOutActio
 	RegisterUnParentAction(InOutActionSet, GetActionId());
 	RegisterCopyAction(InOutActionSet, GetActionId());
 	RegisterPasteAction(InOutActionSet, GetActionId());
+	RegisterPasteUsingGlobalTransformsAction(InOutActionSet, GetActionId());
 	RegisterDuplicateAction(InOutActionSet, GetActionId());
 	RegisterSelectComponentsAction(InOutActionSet, GetActionId());
 	RegisterSelectionFilterCyclingAction(InOutActionSet, GetActionId());
@@ -605,47 +613,38 @@ void USkeletonEditingTool::RegisterCopyAction(FInteractiveToolActionSet& InOutAc
 		EModifierKey::Control, EKeys::C,
 		[this]()
 		{
-			if (Selection.IsEmpty())
+			if (!Modifier)
 			{
 				return;
 			}
+			
 			SkeletonClipboard::CopyToClipboard(*Modifier.Get(), Selection);
 		});
 }
 
 void USkeletonEditingTool::RegisterPasteAction(FInteractiveToolActionSet& InOutActionSet, const int32 InActionId)
 {
-	InOutActionSet.RegisterAction(this, InActionId, TEXT("PasteBones"),
-		LOCTEXT("PasteBones", "Paste Bone(s)"),
-		LOCTEXT("PasteBonesDesc", "Paste Bone(s)"),
+	InOutActionSet.RegisterAction(this, InActionId, TEXT("PasteBonesLocal"),
+		LOCTEXT("PasteBonesLocal", "Paste Bone(s) Local"),
+		LOCTEXT("PasteBonesDesc", "Paste selected bone(s) using the local transform of the copied bones."),
 		EModifierKey::Control, EKeys::V,
 		[this]()
 		{
-			if (!SkeletonClipboard::IsClipboardValid())
-			{
-				return;
-			}
-			
-			const FName DefaultParent = Selection.IsEmpty() ? NAME_None : Selection[0];
-
-			TGuardValue OperationGuard(Operation, EEditingOperation::Create);
-			BeginChange();
-				
-			const TArray<FName> NewBones = SkeletonClipboard::PasteFromClipboard(*Modifier.Get(), DefaultParent);
-			if (!NewBones.IsEmpty())
-			{
-				Selection = NewBones;
-
-				if (NeedsNotification())
-				{
-					GetNotifier()->Notify(NewBones, ESkeletalMeshNotifyType::HierarchyChanged);
-					GetNotifier()->Notify(NewBones, ESkeletalMeshNotifyType::BonesSelected);
-				}
+			constexpr bool bUsingGlobalTransforms = false;
+			PerformPasteAction(bUsingGlobalTransforms);
+		});
+}
 
-				EndChange();
-			}
-			
-			CancelChange();
+void USkeletonEditingTool::RegisterPasteUsingGlobalTransformsAction(FInteractiveToolActionSet& InOutActionSet, const int32 InActionId)
+{
+	InOutActionSet.RegisterAction(this, InActionId, TEXT("PasteBonesGlobal"),
+		LOCTEXT("PasteBonesUsingGlobalTransforms", "Paste Bone(s) Global"),
+		LOCTEXT("PasteBonesUsingGlobalTransformsDesc", "Paste selected bone(s) at the global transform of the copied bones"),
+		EModifierKey::Control | EModifierKey::Shift, EKeys::V, 
+		[this]()
+		{
+			constexpr bool bUsingGlobalTransforms = true;
+			PerformPasteAction(bUsingGlobalTransforms);
 		});
 }
 
@@ -661,18 +660,11 @@ void USkeletonEditingTool::RegisterDuplicateAction(FInteractiveToolActionSet& In
 			{
 				return;
 			}
-			
-			SkeletonClipboard::CopyToClipboard(*Modifier.Get(), Selection);
-
-			if (!SkeletonClipboard::IsClipboardValid())
-			{
-				return;
-			} 
 
 			TGuardValue OperationGuard(Operation, EEditingOperation::Create);
 			BeginChange();
 				
-			const TArray<FName> NewBones = SkeletonClipboard::PasteFromClipboard(*Modifier.Get(), NAME_None);
+			const TArray<FName> NewBones = SkeletonClipboard::ClipboardDuplicate(*Modifier.Get(), Selection);
 			if (!NewBones.IsEmpty())
 			{
 				Selection = NewBones;
@@ -765,6 +757,36 @@ void USkeletonEditingTool::RegisterSnapAction(FInteractiveToolActionSet& InOutAc
 		});
 }
 
+void USkeletonEditingTool::PerformPasteAction(const bool bUsingGlobalTransforms)
+{
+	if (!Modifier ||
+		!SkeletonClipboard::IsClipboardValid())
+	{
+		return;
+	}
+	
+	const FName DefaultParent = Selection.IsEmpty() ? NAME_None : Selection[0];
+
+	TGuardValue OperationGuard(Operation, EEditingOperation::Create);
+	BeginChange();
+			
+	const TArray<FName> NewBones = SkeletonClipboard::PasteFromClipboard(*Modifier.Get(), DefaultParent, bUsingGlobalTransforms);
+	if (!NewBones.IsEmpty())
+	{
+		Selection = NewBones;
+
+		if (NeedsNotification())
+		{
+			GetNotifier()->Notify(NewBones, ESkeletalMeshNotifyType::HierarchyChanged);
+			GetNotifier()->Notify(NewBones, ESkeletalMeshNotifyType::BonesSelected);
+		}
+
+		EndChange();
+	}
+			
+	CancelChange();
+}
+
 TArray<int32> USkeletonEditingTool::GetSelectedComponents() const
 {
 	TArray<int32> IDs;
diff --git a/Engine/Plugins/Runtime/MeshModelingToolset/Source/MeshModelingToolsEditorOnly/Public/SkeletalMesh/SReferenceSkeletonTree.h b/Engine/Plugins/Runtime/MeshModelingToolset/Source/MeshModelingToolsEditorOnly/Public/SkeletalMesh/SReferenceSkeletonTree.h
--- a/Engine/Plugins/Runtime/MeshModelingToolset/Source/MeshModelingToolsEditorOnly/Public/SkeletalMesh/SReferenceSkeletonTree.h
+++ b/Engine/Plugins/Runtime/MeshModelingToolset/Source/MeshModelingToolsEditorOnly/Public/SkeletalMesh/SReferenceSkeletonTree.h
@@ -260,7 +260,7 @@ private:
 	// Copy & Paste
 	void HandleCopyBones() const;
 	bool CanCopyBones() const;
-	void HandlePasteBones();
+	void HandlePasteBones(const bool bUsingGlobalTransforms);
 	bool CanPasteBones() const;
 	void HandleDuplicateBones();
 	bool CanDuplicateBones() const;
diff --git a/Engine/Plugins/Runtime/MeshModelingToolset/Source/MeshModelingToolsEditorOnly/Public/SkeletalMesh/SkeletonEditingTool.h b/Engine/Plugins/Runtime/MeshModelingToolset/Source/MeshModelingToolsEditorOnly/Public/SkeletalMesh/SkeletonEditingTool.h
--- a/Engine/Plugins/Runtime/MeshModelingToolset/Source/MeshModelingToolsEditorOnly/Public/SkeletalMesh/SkeletonEditingTool.h
+++ b/Engine/Plugins/Runtime/MeshModelingToolset/Source/MeshModelingToolsEditorOnly/Public/SkeletalMesh/SkeletonEditingTool.h
@@ -78,8 +78,12 @@ public:
 
 	// Copy & Paste commands
 	TSharedPtr<FUICommandInfo> CopyBones;
-	TSharedPtr<FUICommandInfo> PasteBones;
+	TSharedPtr<FUICommandInfo> PasteBonesLocal;
+	TSharedPtr<FUICommandInfo> PasteBonesGlobal;
 	TSharedPtr<FUICommandInfo> DuplicateBones;
+	
+	UE_DEPRECATED(5.8, "Instead use PasteBonesLocal")
+	TSharedPtr<FUICommandInfo> PasteBones;
 };	
 
 /**
@@ -302,10 +306,13 @@ private:
 	UE_API void RegisterUnParentAction(FInteractiveToolActionSet& InOutActionSet, const int32 InActionId);
 	UE_API void RegisterCopyAction(FInteractiveToolActionSet& InOutActionSet, const int32 InActionId);
 	UE_API void RegisterPasteAction(FInteractiveToolActionSet& InOutActionSet, const int32 InActionId);
+	UE_API void RegisterPasteUsingGlobalTransformsAction(FInteractiveToolActionSet& InOutActionSet, const int32 InActionId);
 	UE_API void RegisterDuplicateAction(FInteractiveToolActionSet& InOutActionSet, const int32 InActionId);
 	UE_API void RegisterSelectComponentsAction(FInteractiveToolActionSet& InOutActionSet, const int32 InActionId);
 	UE_API void RegisterSelectionFilterCyclingAction(FInteractiveToolActionSet& InOutActionSet, const int32 InActionId);
 	UE_API void RegisterSnapAction(FInteractiveToolActionSet& InOutActionSet, const int32 InActionId);
+	
+	void PerformPasteAction(const bool bUsingGlobalTransforms);
 
 	UE_API TArray<int32> GetSelectedComponents() const;
 
diff --git a/Engine/Plugins/Runtime/MeshModelingToolset/Source/SkeletalMeshModifiers/Private/SkeletonClipboard.cpp b/Engine/Plugins/Runtime/MeshModelingToolset/Source/SkeletalMeshModifiers/Private/SkeletonClipboard.cpp
--- a/Engine/Plugins/Runtime/MeshModelingToolset/Source/SkeletalMeshModifiers/Private/SkeletonClipboard.cpp
+++ b/Engine/Plugins/Runtime/MeshModelingToolset/Source/SkeletalMeshModifiers/Private/SkeletonClipboard.cpp
@@ -14,14 +14,44 @@ namespace SkeletonClipboard
 
 void CopyToClipboard(const USkeletonModifier& InModifier, const TArray<FName>& InBonesToCopy)
 {
-	FHierarchyClipboardData ClipboardData;
-	ClipboardData.Bones.Reserve(InBonesToCopy.Num());
-
 	const FReferenceSkeleton& RefSkeleton = InModifier.GetReferenceSkeleton();
 
+	// Copy children as well
+	TFunction<void(const int32, TArray<FName>&)> GetChildrenRecursive = 
+		[&](const int32 BoneIndex, TArray<FName>& InOutBoneAndChildren)
+			{
+				if (!ensureMsgf(BoneIndex != INDEX_NONE, TEXT("Unexpected cannot find bone index in skeleton, cannot copy bone")))
+				{
+					return;
+				}					
+					
+				TArray<int32> DirectChildIndices;
+				RefSkeleton.GetDirectChildBones(BoneIndex, DirectChildIndices);
+					
+				for (int32 ChildIndex : DirectChildIndices)
+				{
+					if (ChildIndex != INDEX_NONE)
+					{
+						InOutBoneAndChildren.AddUnique(RefSkeleton.GetBoneName(ChildIndex));
+						GetChildrenRecursive(ChildIndex, InOutBoneAndChildren);
+					}
+				}
+			};
+			
+	TArray<FName> BonesToCopyWithChildren;
+	for (const FName& BoneToCopy : InBonesToCopy)
+	{
+		BonesToCopyWithChildren.AddUnique(BoneToCopy);
+				
+		const int32 BoneIndex = RefSkeleton.FindBoneIndex(BoneToCopy);
+		if (BoneIndex != INDEX_NONE)
+		{
+			GetChildrenRecursive(BoneIndex, BonesToCopyWithChildren);
+		}
+	}
+
 	// ensure ordering
-	TArray<FName> BonesToCopy(InBonesToCopy);
-	BonesToCopy.StableSort([&](const FName Bone0, const FName Bone1)
+	BonesToCopyWithChildren.StableSort([&](const FName Bone0, const FName Bone1)
 	{
 		const int32 Index0 = RefSkeleton.FindBoneIndex(Bone0);
 		const int32 Index1 = RefSkeleton.FindBoneIndex(Bone1);
@@ -29,14 +59,30 @@ void CopyToClipboard(const USkeletonModifier& InModifier, const TArray<FName>& I
 	});
 
 	// convert bone names to clipboard data
-	for (const FName& BoneName: BonesToCopy)
+	FHierarchyClipboardData ClipboardData;
+	ClipboardData.Bones.Reserve(InBonesToCopy.Num());
+	
+	for (const FName& BoneName : BonesToCopyWithChildren)
 	{
 		const FName ParentName = InModifier.GetParentName(BoneName);
-
+		
 		FBoneClipboardData BoneData;
+		{
 			BoneData.BoneName = BoneName;
-			BoneData.ParentIndex = ParentName != NAME_None ? BonesToCopy.IndexOfByKey(ParentName) : INDEX_NONE;
-			BoneData.Global = InModifier.GetBoneTransform(BoneName, true);
+			
+			BoneData.ParentBoneName = ParentName;
+			
+			constexpr bool bGlobal = true;
+			BoneData.Global = InModifier.GetBoneTransform(BoneName, bGlobal);
+			
+			constexpr bool bLocal = false;
+			BoneData.Local = InModifier.GetBoneTransform(BoneName, bLocal);
+			
+			// DEPRECATED 5.8
+			PRAGMA_DISABLE_DEPRECATION_WARNINGS 		
+			BoneData.ParentIndex_DEPRECATED = ParentName != NAME_None ? BonesToCopyWithChildren.IndexOfByKey(ParentName) : INDEX_NONE;
+			PRAGMA_ENABLE_DEPRECATION_WARNINGS
+		}
 		
 		ClipboardData.Bones.Add(MoveTemp(BoneData));
 	}
@@ -76,6 +122,14 @@ private:
 };
 
 TArray<FName> PasteFromClipboard(USkeletonModifier& InOutModifier, const FName& InDefaultParent)
+{
+	// DEPRECATED 5.8
+
+	constexpr bool bUsingGlobalTransforms = true;
+	return PasteFromClipboard(InOutModifier, InDefaultParent, bUsingGlobalTransforms);
+}
+
+TArray<FName> PasteFromClipboard(USkeletonModifier& InOutModifier, const FName& InDefaultParent, const bool bUsingGlobalTransforms)
 {
 	// Get the text from the clipboard.
 	FString ClipboardText;
@@ -91,48 +145,140 @@ TArray<FName> PasteFromClipboard(USkeletonModifier& InOutModifier, const FName&
 	{
 		return {};
 	}
+	
+	TArray<FName> NewBoneNames;
+	NewBoneNames.Reserve(ClipboardData.Bones.Num());
 
-	// get new names
-	TArray<FName> NewBones;
-	NewBones.Reserve(ClipboardData.Bones.Num());
-
-	Algo::Transform(ClipboardData.Bones, NewBones, [&](const FBoneClipboardData& BoneData)
-	{
-		return InOutModifier.GetUniqueName(BoneData.BoneName, NewBones);
-	});
-
-	// get parent names & local transforms
-	TArray<FName> NewParents;
-	NewParents.Reserve(ClipboardData.Bones.Num());
+	TArray<FName> NewParentBoneNames;
+	NewParentBoneNames.Reserve(ClipboardData.Bones.Num());
 
 	TArray<FTransform> LocalTransforms;
 	LocalTransforms.Reserve(ClipboardData.Bones.Num());
-
-	for (const FBoneClipboardData& BoneData: ClipboardData.Bones)
+	
+	// Generate new bone data
 	{
-		const bool bIsParentFromClipboard = NewBones.IsValidIndex(BoneData.ParentIndex);
-		if (bIsParentFromClipboard)
-		{ // compute from clipboard
-			NewParents.Add(NewBones[BoneData.ParentIndex]);
-
-			const FTransform& ParentGlobal = ClipboardData.Bones[BoneData.ParentIndex].Global;
-			LocalTransforms.Add(BoneData.Global.GetRelativeTransform(ParentGlobal));
+		for (const FBoneClipboardData& BoneData : ClipboardData.Bones)
+		{
+			const bool bIsChildFromClipboard = ClipboardData.Bones.ContainsByPredicate(
+				[&BoneData](const FBoneClipboardData& ClipboardData)
+				{
+					return ClipboardData.BoneName == BoneData.ParentBoneName;
+				});
+				
+			if (bIsChildFromClipboard)
+			{ 
+				NewBoneNames.Add(BoneData.BoneName);
+				NewParentBoneNames.Add(BoneData.ParentBoneName);
+				
+				// Children don't need to consider global transform
+				LocalTransforms.Add(BoneData.Local);
+			}
+			else
+			{ 
+				NewBoneNames.Add(BoneData.BoneName);
+				NewParentBoneNames.Add(InDefaultParent);
+				
+				if (bUsingGlobalTransforms)
+				{
+					constexpr bool bGlobal = true;
+					const FTransform& ParentGlobal = InOutModifier.GetBoneTransform(InDefaultParent, bGlobal);
+					LocalTransforms.Add(BoneData.Global.GetRelativeTransform(ParentGlobal));					
+				}
+				else
+				{
+					LocalTransforms.Add(BoneData.Local);
+				}
+			}
 		}
-		else
-		{ // compute from modifier
-			const bool bUseDefaultParent = InDefaultParent != NAME_None && InDefaultParent != BoneData.BoneName;
-			const FName ParentName = bUseDefaultParent ? InDefaultParent : InOutModifier.GetParentName(BoneData.BoneName);
-			NewParents.Add(ParentName);
-
-			const FTransform& ParentGlobal = InOutModifier.GetBoneTransform(ParentName, true);
-			LocalTransforms.Add(BoneData.Global.GetRelativeTransform(ParentGlobal));
+	}
+	
+	// Make sure all new bones use unique names
+	{
+		check(NewBoneNames.Num() == NewParentBoneNames.Num());
+		
+		TArray<FName> ExistingBoneNames = InOutModifier.GetAllBoneNames();
+		
+		for (FName& BoneName : NewBoneNames)
+		{
+			if (ExistingBoneNames.Contains(BoneName))
+			{
+				const FName NewBoneName = InOutModifier.GetUniqueName(BoneName, ExistingBoneNames);
+				for (FName& ParentBoneName : NewParentBoneNames)
+				{
+					const bool bUsesDefaultParent = ParentBoneName == InDefaultParent;
+					if (ParentBoneName == BoneName &&
+						!bUsesDefaultParent)
+					{
+						ParentBoneName = NewBoneName;
+					}
+				}
+				
+				BoneName = NewBoneName;
+				ExistingBoneNames.Add(BoneName);
+			}
 		}
 	}
 
 	// add bones
-	InOutModifier.AddBones(NewBones, NewParents, LocalTransforms);
+	InOutModifier.AddBones(NewBoneNames, NewParentBoneNames, LocalTransforms);
+
+	return NewBoneNames;
+}
+	
+TArray<FName> ClipboardDuplicate(USkeletonModifier& InOutModifier, const TArray<FName>& InBonesToDuplicate)
+{
+	if (InBonesToDuplicate.IsEmpty())
+	{
+		return {}; 
+	}
+
+	SkeletonClipboard::CopyToClipboard(InOutModifier, InBonesToDuplicate);
+
+	if (!SkeletonClipboard::IsClipboardValid())
+	{
+		return {};
+	} 
+
+	const FReferenceSkeleton& RefSkeleton = InOutModifier.GetReferenceSkeleton();
 
-	return NewBones;
+	// Find the common parent of the selection
+	const int32 CommonParentIndex = [&RefSkeleton, &InBonesToDuplicate]() -> int32
+		{
+			const auto GetParentIndies = [&RefSkeleton](const FName& BoneName)
+				{
+					TArray<int32> ParentIndices;
+					int32 BoneIndex = RefSkeleton.FindBoneIndex(BoneName);
+					while (BoneIndex != INDEX_NONE)
+					{
+						BoneIndex = RefSkeleton.GetParentIndex(BoneIndex);
+						ParentIndices.Add(BoneIndex);
+					}
+									
+					return ParentIndices;
+				};
+					
+			TArray<int32> CommonParentIndices = GetParentIndies(InBonesToDuplicate[0]);
+			for (const FName& BoneName : InBonesToDuplicate)
+			{
+				if (BoneName != InBonesToDuplicate[0])
+				{
+					const TArray<int32> ParentIndices = GetParentIndies(BoneName);
+					CommonParentIndices.RemoveAll([ParentIndices](const int32 CommonParentIndex)
+						{
+							return !ParentIndices.Contains(CommonParentIndex);
+						});
+				}
+			}
+					
+			return CommonParentIndices.IsEmpty() ? INDEX_NONE : CommonParentIndices[0];
+		}();
+				
+	const FName CommonParentName = CommonParentIndex == INDEX_NONE ?
+		NAME_None : 
+		RefSkeleton.GetBoneName(CommonParentIndex);
+						
+	constexpr bool bUsingGlobalTransforms = true;
+	return SkeletonClipboard::PasteFromClipboard(InOutModifier, CommonParentName, bUsingGlobalTransforms);
 }
 	
 bool IsClipboardValid()
diff --git a/Engine/Plugins/Runtime/MeshModelingToolset/Source/SkeletalMeshModifiers/Public/SkeletonClipboard.h b/Engine/Plugins/Runtime/MeshModelingToolset/Source/SkeletalMeshModifiers/Public/SkeletonClipboard.h
--- a/Engine/Plugins/Runtime/MeshModelingToolset/Source/SkeletalMeshModifiers/Public/SkeletonClipboard.h
+++ b/Engine/Plugins/Runtime/MeshModelingToolset/Source/SkeletalMeshModifiers/Public/SkeletonClipboard.h
@@ -18,17 +18,30 @@ class USkeletonModifier;
 namespace SkeletonClipboard
 {
 	/**
-	 * Copies InBonesToCopy data (name, parent, transform) into the clipboard.
+	 * Copies InBonesToCopy data (name, parent, transform) into the clipboard. Also copies children of these bones.
 	 * Hierarchy info & transforms are retrieved via the InModifier.
 	 */
 	SKELETALMESHMODIFIERS_API void CopyToClipboard(const USkeletonModifier& InModifier, const TArray<FName>& InBonesToCopy);
 
 	/**
-	 * Paste new bones into the InOutModifier the clipboard InBonesToCopy data (name, parent, transform) into the clipboard.
-	 * InDefaultParent is used as the default parent if a proper parent can't be retrieve from the clipboard data. 
+	 * Paste new bones from the clipboard InBonesToCopy data (name, parent, transform) into the InOutModifier.
+	 * InDefaultParent is used as the default parent if a proper parent can't be retrieved from the clipboard data. 
 	 */
+	UE_DEPRECATED(5.8, "Instead use the overload that offers the bUsingGlobalTransforms argument")
 	SKELETALMESHMODIFIERS_API TArray<FName> PasteFromClipboard(USkeletonModifier& InOutModifier, const FName& InDefaultParent);
 	
+	/**
+	 * Paste new bones from the clipboard InBonesToCopy data (name, parent, transform) into the InOutModifier.
+	 * InDefaultParent is used as the default parent if a proper parent can't be retrieved from the clipboard data. 
+	 */
+   SKELETALMESHMODIFIERS_API TArray<FName> PasteFromClipboard(USkeletonModifier& InOutModifier, const FName& InDefaultParent, const bool bUsingGlobalTransforms);
+	
+	/**
+	 * Duplicates InBonesToCopy data (name, parent, transform) using the clipboard. Also duplicates children of these bones.
+	 * Hierarchy info & transforms are retrieved via the InModifier.
+	 */
+	SKELETALMESHMODIFIERS_API TArray<FName> ClipboardDuplicate(USkeletonModifier& InOutModifier, const TArray<FName>& InBonesToDuplicate);
+	
 	/**
 	 * Checks whether the clipboard data is valid. 
 	 */
@@ -39,15 +52,32 @@ USTRUCT()
 struct FBoneClipboardData
 {
 	GENERATED_USTRUCT_BODY()
+	
+	FBoneClipboardData() = default;
+	
+	// removing deprecation for default copy operator/constructor to avoid deprecation warnings
+	PRAGMA_DISABLE_DEPRECATION_WARNINGS
+	FBoneClipboardData(const FBoneClipboardData&) = default;
+	FBoneClipboardData(FBoneClipboardData&&) = default;
+	FBoneClipboardData& operator=(const FBoneClipboardData&) = default;
+	FBoneClipboardData& operator=(FBoneClipboardData&&) = default;
+	PRAGMA_ENABLE_DEPRECATION_WARNINGS
 
 	UPROPERTY()
 	FName BoneName = NAME_None;
-
+	
+	UPROPERTY()
+	FName ParentBoneName = NAME_None;
+	
+	UE_DEPRECATED(5.8, "Instead refer to ParentBoneName")
+	UPROPERTY()
+	int32 ParentIndex_DEPRECATED = INDEX_NONE;
+	
 	UPROPERTY()
 	FTransform Global = FTransform::Identity;
-
+	
 	UPROPERTY()
-	int32 ParentIndex = INDEX_NONE;
+	FTransform Local = FTransform::Identity;
 };
 
 USTRUCT()
