add iteration option to smooth brushes

#Synced-CL 51328950

diff --git a/Engine/Plugins/Runtime/MeshModelingToolset/Source/MeshModelingTools/Public/Sculpting/MeshSmoothingBrushOps.h b/Engine/Plugins/Runtime/MeshModelingToolset/Source/MeshModelingTools/Public/Sculpting/MeshSmoothingBrushOps.h
--- a/Engine/Plugins/Runtime/MeshModelingToolset/Source/MeshModelingTools/Public/Sculpting/MeshSmoothingBrushOps.h
+++ b/Engine/Plugins/Runtime/MeshModelingToolset/Source/MeshModelingTools/Public/Sculpting/MeshSmoothingBrushOps.h
@@ -17,6 +17,7 @@ class UBaseSmoothBrushOpProps : public UMeshSculptBrushOpProps
 public:
 	virtual bool GetPreserveUVFlow() { return false; }
 	virtual bool SupportsStrengthPressure() override { return true; }
+	virtual int32 GetIterations() { return 1; }
 };
 
 
@@ -33,6 +34,9 @@ public:
 	UPROPERTY(EditAnywhere, Category = SmoothBrush, meta = (DisplayName = "Falloff", UIMin = "0.0", UIMax = "1.0", ClampMin = "0.0", ClampMax = "1.0"))
 	float Falloff = 0.5;
 
+	UPROPERTY(EditAnywhere, Category = SmoothBrush, meta = (ClampMin = 1, UIMax = 50))
+	int32 Iterations = 1;
+
 	/** If true, try to preserve the shape of the UV/3D mapping. This will limit Smoothing and Remeshing in some cases. */
 	UPROPERTY(EditAnywhere, Category = SmoothBrush)
 	bool bPreserveUVFlow = false;
@@ -41,6 +45,7 @@ public:
 	virtual void SetStrength(float NewStrength) override { Strength = FMathf::Clamp(NewStrength, 0.0f, 1.0f); }
 	virtual float GetFalloff() override { return Falloff; }
 	virtual bool GetPreserveUVFlow() override { return bPreserveUVFlow; }
+	virtual int32 GetIterations() { return Iterations; }
 };
 
 
@@ -58,6 +63,9 @@ public:
 	UPROPERTY(EditAnywhere, Category = ShiftSmoothBrush, meta = (DisplayName = "Falloff", UIMin = "0.0", UIMax = "1.0", ClampMin = "0.0", ClampMax = "1.0"))
 	float Falloff = 0.5;
 
+	UPROPERTY(EditAnywhere, Category = ShiftSmoothBrush, meta = (ClampMin = 1, UIMax = 50))
+	int32 Iterations = 1;
+
 	/** If true, try to preserve the shape of the UV/3D mapping. This will limit Smoothing and Remeshing in some cases. */
 	UPROPERTY(EditAnywhere, Category = ShiftSmoothBrush)
 	bool bPreserveUVFlow = false;
@@ -66,6 +74,7 @@ public:
 	virtual void SetStrength(float NewStrength) override { Strength = FMathf::Clamp(NewStrength, 0.0f, 1.0f); }
 	virtual float GetFalloff() override { return Falloff; }
 	virtual bool GetPreserveUVFlow() override { return bPreserveUVFlow; }
+	virtual int32 GetIterations() { return Iterations; }
 };
 
 
@@ -75,32 +84,67 @@ class FSmoothBrushOp : public FMeshSculptBrushOp
 
 public:
 
+	static void ApplyIterativeStamp(const FDynamicMesh3* Mesh, const TArray<int32>& Vertices, TArray<FVector3d>& NewPositionsOut, int32 NumIters,
+		TFunctionRef<FVector3d(int32 Iter, FVector3d OrigPos, int32 VID, TFunctionRef<FVector3d(int32)> GetVertexFn)> GetStampVertPos)
+	{
+		TMap<int32, int32> VIDToIdx;
+		if (NumIters > 0)
+		{
+			// We need this map for iterations after the first, to go from neighbor vertex ID -> index in NewPositionsOut
+			// (otherwise we'd use stale positions for the neighbor vertices)
+			VIDToIdx.Reserve(Vertices.Num());
+			for (int32 Idx = 0; Idx < Vertices.Num(); ++Idx)
+			{
+				VIDToIdx.Add(Vertices[Idx], Idx);
+			}
+		}
+
+		auto MeshVertPos = [&Mesh](int32 VID) { return Mesh->GetVertex(VID); };
+		auto BufferVertPos = [&Mesh, &VIDToIdx, &NewPositionsOut](int32 VID)
+		{
+			if (int32* Idx = VIDToIdx.Find(VID))
+			{
+				return NewPositionsOut[*Idx];
+			}
+			else
+			{
+				return Mesh->GetVertex(VID);
+			}
+		};
+		for (int32 Iter = 0; Iter < NumIters; ++Iter)
+		{
+			ParallelFor(Vertices.Num(), [&](int32 k)
+			{
+				int32 VertIdx = Vertices[k];
+
+				FVector3d OrigPos = Iter > 0 ? NewPositionsOut[k] : Mesh->GetVertex(VertIdx);
+				TFunctionRef<FVector3d(int32)> GetVertexFn = Iter == 0 ?
+					TFunctionRef<FVector3d(int32)>(MeshVertPos) :
+					TFunctionRef<FVector3d(int32)>(BufferVertPos);
+				NewPositionsOut[k] = GetStampVertPos(Iter, OrigPos, VertIdx, GetVertexFn);
+			});
+		}
+	}
+
 	virtual void ApplyStamp(const FDynamicMesh3* Mesh, const FSculptBrushStamp& Stamp, const TArray<int32>& Vertices, TArray<FVector3d>& NewPositionsOut) override
 	{
 		const FVector3d& StampPos = Stamp.LocalFrame.Origin;
 		bool bPreserveUVFlow = GetPropertySetAs<UBaseSmoothBrushOpProps>()->GetPreserveUVFlow();
-
-		ParallelFor(Vertices.Num(), [&](int32 k)
+		int32 NumIters = GetPropertySetAs<UBaseSmoothBrushOpProps>()->GetIterations();
+		ApplyIterativeStamp(Mesh, Vertices, NewPositionsOut, NumIters,
+		[&](int32 Iter, FVector3d OrigPos, int32 VID, TFunctionRef<FVector3d(int32)> GetVertexFn)->FVector3d
 		{
-			int32 VertIdx = Vertices[k];
-
-			FVector3d OrigPos = Mesh->GetVertex(VertIdx);
-
 			double Falloff = GetFalloff().Evaluate(Stamp, OrigPos);
-
 			FVector3d SmoothedPos = OrigPos;
 			if (bPreserveUVFlow)
 			{
-				SmoothedPos = UE::Geometry::FMeshWeights::CotanCentroidSafe(*Mesh, VertIdx, 10.0);
+				SmoothedPos = UE::Geometry::FMeshWeights::CotanCentroidSafe(*Mesh, VID, GetVertexFn, 10.0);
 			}
 			else
 			{
-				SmoothedPos = UE::Geometry::FMeshWeights::UniformCentroid(*Mesh, VertIdx);
+				SmoothedPos = UE::Geometry::FMeshWeights::UniformCentroid(*Mesh, VID, GetVertexFn);
 			}
-
-			FVector3d NewPos = UE::Geometry::Lerp(OrigPos, SmoothedPos, Falloff * Stamp.Power);
-
-			NewPositionsOut[k] = NewPos;
+			return UE::Geometry::Lerp(OrigPos, SmoothedPos, Falloff * Stamp.Power);
 		});
 	}
 };
@@ -150,6 +194,7 @@ public:
 			int32 VertIdx = Vertices[k];
 
 			FVector3d OrigPos = Mesh->GetVertex(VertIdx);
+			// TODO: If adding Iterations option to this brush, need to also make a ComputeVertexNormals that takes a TFunctionRef for getting vertex positions
 			FVector3d Normal = UE::Geometry::FMeshNormals::ComputeVertexNormal(*Mesh, VertIdx);
 
 			double Falloff = GetFalloff().Evaluate(Stamp, OrigPos);
