[SkeletalMeshEditingTools] The skeleton editing tool can now copy/paste bone transforms

#Synced-CL 51338101

diff --git a/Engine/Plugins/Runtime/MeshModelingToolset/Source/MeshModelingToolsEditorOnly/Private/SkeletalMesh/SReferenceSkeletonTree.cpp b/Engine/Plugins/Runtime/MeshModelingToolset/Source/MeshModelingToolsEditorOnly/Private/SkeletalMesh/SReferenceSkeletonTree.cpp
--- a/Engine/Plugins/Runtime/MeshModelingToolset/Source/MeshModelingToolsEditorOnly/Private/SkeletalMesh/SReferenceSkeletonTree.cpp
+++ b/Engine/Plugins/Runtime/MeshModelingToolset/Source/MeshModelingToolsEditorOnly/Private/SkeletalMesh/SReferenceSkeletonTree.cpp
@@ -315,6 +315,7 @@ SReferenceSkeletonTree::~SReferenceSkeletonTree()
 
 void SReferenceSkeletonTree::Construct(const FArguments& InArgs)
 {
+	checkf(InArgs._Modifier.IsValid(), TEXT("Modifier is a required argument for SReferenceSkeletonTree"));
 	Modifier = InArgs._Modifier;
 	bIsReadOnly = Modifier->IsReadOnly();
 
@@ -653,12 +654,10 @@ void SReferenceSkeletonTree::HandleCopyBones() const
 	TArray<FName> SelectedBones;
 	GetSelectedBoneNames(SelectedBones);
 
-	if (SelectedBones.IsEmpty())
+	if (!SelectedBones.IsEmpty())
 	{
-		return;
+		SkeletonClipboard::CopyBonesToClipboard(*Modifier.Get(), SelectedBones);
 	}
-
-	SkeletonClipboard::CopyToClipboard(*Modifier.Get(), SelectedBones);
 }
 
 bool SReferenceSkeletonTree::CanCopyBones() const
@@ -685,7 +684,7 @@ void SReferenceSkeletonTree::HandlePasteBones(const bool bUsingGlobalTransforms)
 
 	BeginChange();
 	
-	const TArray<FName> NewBones = SkeletonClipboard::PasteFromClipboard(*Modifier.Get(), DefaultParent, bUsingGlobalTransforms);
+	const TArray<FName> NewBones = SkeletonClipboard::PasteBonesFromClipboard(*Modifier.Get(), DefaultParent, bUsingGlobalTransforms);
 	if (NewBones.IsEmpty())
 	{
 		CancelChange();
@@ -713,7 +712,7 @@ bool SReferenceSkeletonTree::CanPasteBones() const
 		return false;
 	}
 	
-	return Modifier.IsValid() && SkeletonClipboard::IsClipboardValid();
+	return Modifier.IsValid() && SkeletonClipboard::IsBoneClipboardValid();
 }
 
 void SReferenceSkeletonTree::HandleDuplicateBones()
@@ -729,7 +728,7 @@ void SReferenceSkeletonTree::HandleDuplicateBones()
 	GetSelectedBoneNames(SelectedBones);
 	
 	constexpr bool bUsingGlobalTransforms = false;
-	const TArray<FName> NewBones = SkeletonClipboard::ClipboardDuplicate(*Modifier.Get(), SelectedBones);
+	const TArray<FName> NewBones = SkeletonClipboard::ClipboardDuplicateBones(*Modifier.Get(), SelectedBones);
 	if (NewBones.IsEmpty())
 	{
 		CancelChange();
diff --git a/Engine/Plugins/Runtime/MeshModelingToolset/Source/MeshModelingToolsEditorOnly/Private/SkeletalMesh/SkeletonEditingTool.cpp b/Engine/Plugins/Runtime/MeshModelingToolset/Source/MeshModelingToolsEditorOnly/Private/SkeletalMesh/SkeletonEditingTool.cpp
--- a/Engine/Plugins/Runtime/MeshModelingToolset/Source/MeshModelingToolsEditorOnly/Private/SkeletalMesh/SkeletonEditingTool.cpp
+++ b/Engine/Plugins/Runtime/MeshModelingToolset/Source/MeshModelingToolsEditorOnly/Private/SkeletalMesh/SkeletonEditingTool.cpp
@@ -18,6 +18,7 @@
 #include "BaseBehaviors/ClickDragBehavior.h"
 #include "BaseGizmos/GizmoViewContext.h"
 #include "DynamicMesh/MeshNormals.h"
+#include "Framework/MultiBox/MultiBoxBuilder.h"
 #include "Preferences/PersonaOptions.h"
 #include "Selection/PolygonSelectionMechanic.h"
 
@@ -92,10 +93,14 @@ void FSkeletonEditingToolCommands::RegisterCommands()
 	UI_COMMAND(AddSuffix, "Add Suffix", "Add a suffix to the selected bone(s) name(s).", EUserInterfaceActionType::Button, FInputChord());
 	
 	UI_COMMAND(CopyBones, "Copy Bone(s)", "Copy selected bone(s).", EUserInterfaceActionType::Button, FInputChord(EKeys::C, EModifierKey::Control));
-	UI_COMMAND(DuplicateBones, "Duplicate Bone(s)", "Duplicate selected bone(s).", EUserInterfaceActionType::Button, FInputChord(EKeys::D, EModifierKey::Control));
+	UI_COMMAND(PasteBonesLocal, "Paste Bone(s) Local", "Paste bone(s) using the local transform of the copied bones.", EUserInterfaceActionType::Button, FInputChord(EKeys::V, EModifierKey::Control));
+	UI_COMMAND(PasteBonesGlobal, "Paste Bone(s) Global", "Paste bone(s) using the global transform of the copied bones.", EUserInterfaceActionType::Button, FInputChord(EKeys::V, EModifierKey::Control | EModifierKey::Shift));
+	
+	UI_COMMAND(CopyBoneTransforms, "Copy Bone Transform(s)", "Copy selected bone(s) transforms.", EUserInterfaceActionType::Button, FInputChord(EKeys::C, EModifierKey::Control));
+	UI_COMMAND(PasteBoneTransformsLocal, "Paste Bone Transform(s) Local", "Paste local bone transform(s) to the selected bones.", EUserInterfaceActionType::Button, FInputChord(EKeys::V, EModifierKey::Control | EModifierKey::Alt));
+	UI_COMMAND(PasteBoneTransformsGlobal, "Paste Bone Transform(s) Global", "Paste global bone transfrom(s) to the selected bones.", EUserInterfaceActionType::Button, FInputChord(EKeys::V, EModifierKey::Control | EModifierKey::Shift | EModifierKey::Alt));
 
-	UI_COMMAND(PasteBonesLocal, "Paste Bone(s) Local", "Paste selected bone(s) using the local transform of the copied bones.", EUserInterfaceActionType::Button, FInputChord(EKeys::V, EModifierKey::Control));
-	UI_COMMAND(PasteBonesGlobal, "Paste Bone(s) Global", "Paste selected bone(s) using the global transform of the copied bones.", EUserInterfaceActionType::Button, FInputChord(EKeys::V, EModifierKey::Control | EModifierKey::Shift));
+	UI_COMMAND(DuplicateBones, "Duplicate Bone(s)", "Duplicate selected bone(s).", EUserInterfaceActionType::Button, FInputChord(EKeys::D, EModifierKey::Control));
 
 	// DEPRECATED 5.8
 	PRAGMA_DISABLE_DEPRECATION_WARNINGS
@@ -527,8 +532,10 @@ void USkeletonEditingTool::RegisterActions(FInteractiveToolActionSet& InOutActio
 	RegisterParentAction(InOutActionSet, GetActionId());
 	RegisterUnParentAction(InOutActionSet, GetActionId());
 	RegisterCopyAction(InOutActionSet, GetActionId());
-	RegisterPasteAction(InOutActionSet, GetActionId());
-	RegisterPasteUsingGlobalTransformsAction(InOutActionSet, GetActionId());
+	RegisterPasteBonesLocalAction(InOutActionSet, GetActionId());
+	RegisterPasteBonesGlobalAction(InOutActionSet, GetActionId());
+	RegisterPasteBoneTransformsLocalAction(InOutActionSet, GetActionId());
+	RegisterPasteBoneTransformsGlobalAction(InOutActionSet, GetActionId());
 	RegisterDuplicateAction(InOutActionSet, GetActionId());
 	RegisterSelectComponentsAction(InOutActionSet, GetActionId());
 	RegisterSelectionFilterCyclingAction(InOutActionSet, GetActionId());
@@ -618,11 +625,11 @@ void USkeletonEditingTool::RegisterCopyAction(FInteractiveToolActionSet& InOutAc
 				return;
 			}
 			
-			SkeletonClipboard::CopyToClipboard(*Modifier.Get(), Selection);
+			SkeletonClipboard::CopyBonesToClipboard(*Modifier.Get(), Selection);
 		});
 }
 
-void USkeletonEditingTool::RegisterPasteAction(FInteractiveToolActionSet& InOutActionSet, const int32 InActionId)
+void USkeletonEditingTool::RegisterPasteBonesLocalAction(FInteractiveToolActionSet& InOutActionSet, const int32 InActionId)
 {
 	InOutActionSet.RegisterAction(this, InActionId, TEXT("PasteBonesLocal"),
 		LOCTEXT("PasteBonesLocal", "Paste Bone(s) Local"),
@@ -631,23 +638,49 @@ void USkeletonEditingTool::RegisterPasteAction(FInteractiveToolActionSet& InOutA
 		[this]()
 		{
 			constexpr bool bUsingGlobalTransforms = false;
-			PerformPasteAction(bUsingGlobalTransforms);
+			PerformPasteBonesAction(bUsingGlobalTransforms);
 		});
 }
 
-void USkeletonEditingTool::RegisterPasteUsingGlobalTransformsAction(FInteractiveToolActionSet& InOutActionSet, const int32 InActionId)
+void USkeletonEditingTool::RegisterPasteBonesGlobalAction(FInteractiveToolActionSet& InOutActionSet, const int32 InActionId)
 {
 	InOutActionSet.RegisterAction(this, InActionId, TEXT("PasteBonesGlobal"),
-		LOCTEXT("PasteBonesUsingGlobalTransforms", "Paste Bone(s) Global"),
-		LOCTEXT("PasteBonesUsingGlobalTransformsDesc", "Paste selected bone(s) at the global transform of the copied bones"),
+		LOCTEXT("PasteBonesGlobal", "Paste Bone(s) Global"),
+		LOCTEXT("PasteBonesGlobalDesc", "Paste selected bone(s) using the global transform of the copied bones"),
 		EModifierKey::Control | EModifierKey::Shift, EKeys::V, 
 		[this]()
 		{
 			constexpr bool bUsingGlobalTransforms = true;
-			PerformPasteAction(bUsingGlobalTransforms);
+			PerformPasteBonesAction(bUsingGlobalTransforms);
 		});
 }
 
+void USkeletonEditingTool::RegisterPasteBoneTransformsLocalAction(FInteractiveToolActionSet& InOutActionSet, const int32 InActionId)
+{
+	InOutActionSet.RegisterAction(this, InActionId, TEXT("PasteBoneTransformsLocal"),
+		LOCTEXT("PasteBoneTransformsLocal", "Paste Bone Transform(s) Local"),
+		LOCTEXT("PasteBoneTransformsLocalDesc", "Paste local bone transforms(s) to the selected bones"),
+		EModifierKey::Control | EModifierKey::Alt, EKeys::V, 
+		[this]()
+		{
+			constexpr bool bUsingGlobalTransforms = false;
+			PerformPasteBoneTransformsAction(bUsingGlobalTransforms);
+		});	
+}
+
+void USkeletonEditingTool::RegisterPasteBoneTransformsGlobalAction(FInteractiveToolActionSet& InOutActionSet, const int32 InActionId)
+{
+	InOutActionSet.RegisterAction(this, InActionId, TEXT("PasteBoneTransformsGlobal"),
+	LOCTEXT("PasteBoneTransformsGlobal", "Paste Bone Transform(s) Global"),
+	LOCTEXT("PasteBoneTransformsGlobalDesc", "Paste global bone transfroms(s) to the selected bones"),
+	EModifierKey::Control | EModifierKey::Shift | EModifierKey::Alt, EKeys::V, 
+	[this]()
+	{
+		constexpr bool bUsingGlobalTransforms = true;
+		PerformPasteBoneTransformsAction(bUsingGlobalTransforms);
+	});
+}
+
 void USkeletonEditingTool::RegisterDuplicateAction(FInteractiveToolActionSet& InOutActionSet, const int32 InActionId)
 {
 	InOutActionSet.RegisterAction(this, InActionId, TEXT("DuplicateBones"),
@@ -664,7 +697,7 @@ void USkeletonEditingTool::RegisterDuplicateAction(FInteractiveToolActionSet& In
 			TGuardValue OperationGuard(Operation, EEditingOperation::Create);
 			BeginChange();
 				
-			const TArray<FName> NewBones = SkeletonClipboard::ClipboardDuplicate(*Modifier.Get(), Selection);
+			const TArray<FName> NewBones = SkeletonClipboard::ClipboardDuplicateBones(*Modifier.Get(), Selection);
 			if (!NewBones.IsEmpty())
 			{
 				Selection = NewBones;
@@ -757,20 +790,19 @@ void USkeletonEditingTool::RegisterSnapAction(FInteractiveToolActionSet& InOutAc
 		});
 }
 
-void USkeletonEditingTool::PerformPasteAction(const bool bUsingGlobalTransforms)
+void USkeletonEditingTool::PerformPasteBonesAction(const bool bUsingGlobalTransforms)
 {
 	if (!Modifier ||
-		!SkeletonClipboard::IsClipboardValid())
+		!SkeletonClipboard::IsBoneClipboardValid())
 	{
 		return;
 	}
 	
-	const FName DefaultParent = Selection.IsEmpty() ? NAME_None : Selection[0];
-
-	TGuardValue OperationGuard(Operation, EEditingOperation::Create);
+	const TGuardValue OperationGuard(Operation, EEditingOperation::Create);
 	BeginChange();
-			
-	const TArray<FName> NewBones = SkeletonClipboard::PasteFromClipboard(*Modifier.Get(), DefaultParent, bUsingGlobalTransforms);
+	
+	const FName DefaultParent = Selection.IsEmpty() ? NAME_None : Selection[0];
+	const TArray<FName> NewBones = SkeletonClipboard::PasteBonesFromClipboard(*Modifier.Get(), DefaultParent, bUsingGlobalTransforms);
 	if (!NewBones.IsEmpty())
 	{
 		Selection = NewBones;
@@ -787,6 +819,47 @@ void USkeletonEditingTool::PerformPasteAction(const bool bUsingGlobalTransforms)
 	CancelChange();
 }
 
+void USkeletonEditingTool::PerformPasteBoneTransformsAction(const bool bUsingGlobalTransforms)
+{
+	if (!Modifier ||
+		!SkeletonClipboard::IsBoneTransformClipboardValid())
+	{
+		return;
+	}
+	
+	const SkeletonClipboard::EPasteBoneTransformsFromClipboardFlags PasteFlags = 
+	[bUsingGlobalTransforms, this]()
+		{
+			SkeletonClipboard::EPasteBoneTransformsFromClipboardFlags Flags = SkeletonClipboard::EPasteBoneTransformsFromClipboardFlags::None;
+			if (bUsingGlobalTransforms)
+			{	
+				EnumAddFlags(Flags, SkeletonClipboard::EPasteBoneTransformsFromClipboardFlags::UsingGlobalTransforms);
+			}
+			if (Properties && Properties->bUpdateChildren)
+			{
+				EnumAddFlags(Flags, SkeletonClipboard::EPasteBoneTransformsFromClipboardFlags::UpdateChildren);
+			}
+			return Flags;
+		}();
+	
+	const TGuardValue OperationGuard(Operation, EEditingOperation::Create);
+	BeginChange();
+	
+	const TArray<FName> AffectedBones = SkeletonClipboard::PasteBoneTransformsFromClipboard(*Modifier.Get(), Selection, PasteFlags);
+	if (!AffectedBones.IsEmpty())
+	{
+		if (NeedsNotification())
+		{
+			GetNotifier()->Notify(AffectedBones, ESkeletalMeshNotifyType::HierarchyChanged);
+			GetNotifier()->Notify(AffectedBones, ESkeletalMeshNotifyType::BonesMoved);
+		}
+
+		EndChange();
+	}
+			
+	CancelChange();
+}
+
 TArray<int32> USkeletonEditingTool::GetSelectedComponents() const
 {
 	TArray<int32> IDs;
@@ -1510,19 +1583,56 @@ bool USkeletonEditingTool::CanModifySkeleton() const
 
 TSharedPtr<SWidget> USkeletonEditingTool::GetCustomSkeletonTreeWidget(TSharedPtr<ISkeletalMeshEditorBinding> EditorBinding)
 {
-	TSharedPtr<SReferenceSkeletonTree> ToolSkeletonTree =
-		SNew(SReferenceSkeletonTree)
-		.Modifier(Modifier);
+	FRefSkeletonTreeDelegates SkeletonTreeDelegates;
+
+	SkeletonTreeDelegates.OnExtendCommandList.BindLambda(
+		[this](TSharedPtr<FUICommandList> InOutCommandList)
+		{
+			const FSkeletonEditingToolCommands& Commands = FSkeletonEditingToolCommands::Get();
+
+			InOutCommandList->MapAction(Commands.CopyBoneTransforms,
+				FExecuteAction::CreateUObject(this, &USkeletonEditingTool::HandleCopyBoneTransforms),
+				FCanExecuteAction::CreateUObject(this, &USkeletonEditingTool::CanCopyBoneTransforms));
+		
+			constexpr bool bPasteUsingGlobalTransforms = true;
+			InOutCommandList->MapAction(Commands.PasteBoneTransformsLocal,
+				FExecuteAction::CreateUObject(this, &USkeletonEditingTool::HandlePasteBoneTransforms, !bPasteUsingGlobalTransforms),
+				FCanExecuteAction::CreateUObject(this, &USkeletonEditingTool::CanPasteBoneTransforms));
+		
+			InOutCommandList->MapAction(Commands.PasteBoneTransformsGlobal,
+				FExecuteAction::CreateUObject(this, &USkeletonEditingTool::HandlePasteBoneTransforms, bPasteUsingGlobalTransforms),
+				FCanExecuteAction::CreateUObject(this, &USkeletonEditingTool::CanPasteBoneTransforms));
+		});
+		
+	SkeletonTreeDelegates.OnExtendContextMenu.BindLambda([](FMenuBuilder& MenuBuilder)
+		{
+			// The transform section is added here as it relies on this tool's bUpdateChildren setting
+			const FSkeletonEditingToolCommands& Commands = FSkeletonEditingToolCommands::Get();
+			
+			MenuBuilder.BeginSection("CopyPasteTransforms", LOCTEXT("CopyPasteTransformsSection", "Copy & Paste Transforms"));
+			{
+				MenuBuilder.AddMenuEntry(Commands.CopyBoneTransforms);
+				MenuBuilder.AddMenuEntry(Commands.PasteBoneTransformsLocal);
+				MenuBuilder.AddMenuEntry(Commands.PasteBoneTransformsGlobal);
+			}
+			MenuBuilder.EndSection();
+
+			return true;
+		});
+
+	SAssignNew(SkeletonTree, SReferenceSkeletonTree)
+		.Modifier(Modifier)
+		.Delegates(SkeletonTreeDelegates);
 
 	SkeletonTreeNotifierBindScope.Reset(
 		new FSkeletalMeshNotifierBindScope(
 			EditorBinding->GetNotifier(),
-			ToolSkeletonTree->GetNotifier()));	
+			SkeletonTree->GetNotifier()));	
 		
 	TArray<FName> SelectedBones = EditorBinding->GetSelectedBones();
-	ToolSkeletonTree->GetNotifier()->HandleNotification(SelectedBones, ESkeletalMeshNotifyType::BonesSelected);		
+	SkeletonTree->GetNotifier()->HandleNotification(SelectedBones, ESkeletalMeshNotifyType::BonesSelected);		
 	
-	return ToolSkeletonTree;
+	return SkeletonTree;
 }
 
 EEditingOperation USkeletonEditingTool::GetOperation() const
@@ -2058,6 +2168,90 @@ FName USkeletonEditingTool::GetCurrentBone() const
 	return Selection.IsEmpty() ? NAME_None : Selection.Last(); 
 }
 
+void USkeletonEditingTool::HandleCopyBoneTransforms() const
+{
+	if (!ensure(SkeletonTree.IsValid() && Modifier))
+	{
+		return;
+	}
+
+	TArray<FName> SelectedBones;
+	SkeletonTree->GetSelectedBoneNames(SelectedBones);
+
+	if (!SelectedBones.IsEmpty())
+	{
+		SkeletonClipboard::CopyBoneTransformsToClipboard(*Modifier.Get(), SelectedBones);
+	}
+}
+
+bool USkeletonEditingTool::CanCopyBoneTransforms() const
+{
+	if (!SkeletonTree.IsValid() ||
+		!Modifier || 
+		Modifier->IsReadOnly())
+	{
+		return false;
+	}
+
+	return SkeletonTree->HasSelectedItems();
+}
+
+void USkeletonEditingTool::HandlePasteBoneTransforms(const bool bUsingGlobalTransforms)
+{
+	if (!ensure(SkeletonTree.IsValid() && Modifier))
+	{
+		return;
+	}
+
+	TArray<FName> BoneNames;
+	SkeletonTree->GetSelectedBoneNames(BoneNames);
+
+	const SkeletonClipboard::EPasteBoneTransformsFromClipboardFlags PasteFlags = 
+		[bUsingGlobalTransforms, this]()
+		{
+			SkeletonClipboard::EPasteBoneTransformsFromClipboardFlags Flags = SkeletonClipboard::EPasteBoneTransformsFromClipboardFlags::None;
+			if (bUsingGlobalTransforms)
+			{	
+				Flags |= SkeletonClipboard::EPasteBoneTransformsFromClipboardFlags::UsingGlobalTransforms;
+			}
+			if (Properties && Properties->bUpdateChildren)
+			{
+				Flags |= SkeletonClipboard::EPasteBoneTransformsFromClipboardFlags::UpdateChildren;
+			}
+
+			return Flags;
+		}();
+
+	BeginChange();
+	
+	const TArray<FName> AffectedBones = SkeletonClipboard::PasteBoneTransformsFromClipboard(*Modifier.Get(), BoneNames, PasteFlags);
+	if (AffectedBones.IsEmpty())
+	{
+		CancelChange();
+		return;
+	}
+	
+	if (const TSharedPtr<ISkeletalMeshNotifier> MeshNotifier = GetNotifier())
+	{
+		MeshNotifier->Notify(AffectedBones, ESkeletalMeshNotifyType::HierarchyChanged);
+		MeshNotifier->Notify(AffectedBones, ESkeletalMeshNotifyType::BonesMoved);
+	}
+	UpdateGizmo();
+	
+	EndChange();
+}
+
+bool USkeletonEditingTool::CanPasteBoneTransforms() const
+{
+	if (!Modifier ||
+		Modifier->IsReadOnly())
+	{
+		return false;
+	}
+	
+	return SkeletonClipboard::IsBoneTransformClipboardValid();
+}
+
 TArray<int32> USkeletonEditingTool::GetSelectedBoneIndexes() const
 {
 	TArray<int32> Indexes;
diff --git a/Engine/Plugins/Runtime/MeshModelingToolset/Source/MeshModelingToolsEditorOnly/Public/SkeletalMesh/SReferenceSkeletonTree.h b/Engine/Plugins/Runtime/MeshModelingToolset/Source/MeshModelingToolsEditorOnly/Public/SkeletalMesh/SReferenceSkeletonTree.h
--- a/Engine/Plugins/Runtime/MeshModelingToolset/Source/MeshModelingToolsEditorOnly/Public/SkeletalMesh/SReferenceSkeletonTree.h
+++ b/Engine/Plugins/Runtime/MeshModelingToolset/Source/MeshModelingToolsEditorOnly/Public/SkeletalMesh/SReferenceSkeletonTree.h
@@ -216,7 +216,7 @@ public:
 		{}
 		SLATE_ARGUMENT(TWeakObjectPtr<USkeletonModifier>, Modifier)
 		SLATE_ARGUMENT(FRefSkeletonTreeDelegates, Delegates)
-		SLATE_ARGUMENT(bool, bIsReadOnly)
+		SLATE_ARGUMENT_DEPRECATED(bool, bIsReadOnly, 5.8, "This argument did not provide functionality and is now deprecated")
 	SLATE_END_ARGS()
 
 	UE_API SReferenceSkeletonTree();
diff --git a/Engine/Plugins/Runtime/MeshModelingToolset/Source/MeshModelingToolsEditorOnly/Public/SkeletalMesh/SkeletonEditingTool.h b/Engine/Plugins/Runtime/MeshModelingToolset/Source/MeshModelingToolsEditorOnly/Public/SkeletalMesh/SkeletonEditingTool.h
--- a/Engine/Plugins/Runtime/MeshModelingToolset/Source/MeshModelingToolsEditorOnly/Public/SkeletalMesh/SkeletonEditingTool.h
+++ b/Engine/Plugins/Runtime/MeshModelingToolset/Source/MeshModelingToolsEditorOnly/Public/SkeletalMesh/SkeletonEditingTool.h
@@ -20,6 +20,7 @@
 
 #define UE_API MESHMODELINGTOOLSEDITORONLY_API
 
+class SReferenceSkeletonTree;
 class USingleClickInputBehavior;
 class UGizmoViewContext;
 class UTransformGizmo;
@@ -78,8 +79,12 @@ public:
 
 	// Copy & Paste commands
 	TSharedPtr<FUICommandInfo> CopyBones;
+	TSharedPtr<FUICommandInfo> CopyBoneTransforms;
 	TSharedPtr<FUICommandInfo> PasteBonesLocal;
-	TSharedPtr<FUICommandInfo> PasteBonesGlobal;
+	TSharedPtr<FUICommandInfo> PasteBonesGlobal;		
+	TSharedPtr<FUICommandInfo> PasteBoneTransformsLocal;
+	TSharedPtr<FUICommandInfo> PasteBoneTransformsGlobal;
+	
 	TSharedPtr<FUICommandInfo> DuplicateBones;
 	
 	UE_DEPRECATED(5.8, "Instead use PasteBonesLocal")
@@ -266,6 +271,18 @@ protected:
 	friend class SkeletonEditingTool::FRefSkeletonChange;
 	
 private:
+	/** Copies bone transforms to clipboard */
+	void HandleCopyBoneTransforms() const;
+	
+	/** Returns true if bone transforms can be copied to the clipboard */
+	bool CanCopyBoneTransforms() const;
+	
+	/** Pastes bone transforms from the clipboard to selected bones  */
+	void HandlePasteBoneTransforms(const bool bUsingGlobalTransforms);
+	
+	/** Returns true if bone transforms can be pasted from the clipboard to selected bones */
+	bool CanPasteBoneTransforms() const;
+
 	UE_API TArray<int32> GetSelectedBoneIndexes() const;
 	
 	enum class EBoneSelectionMode : uint8
@@ -305,14 +322,17 @@ private:
 	UE_API void RegisterParentAction(FInteractiveToolActionSet& InOutActionSet, const int32 InActionId);
 	UE_API void RegisterUnParentAction(FInteractiveToolActionSet& InOutActionSet, const int32 InActionId);
 	UE_API void RegisterCopyAction(FInteractiveToolActionSet& InOutActionSet, const int32 InActionId);
-	UE_API void RegisterPasteAction(FInteractiveToolActionSet& InOutActionSet, const int32 InActionId);
-	UE_API void RegisterPasteUsingGlobalTransformsAction(FInteractiveToolActionSet& InOutActionSet, const int32 InActionId);
+	UE_API void RegisterPasteBonesLocalAction(FInteractiveToolActionSet& InOutActionSet, const int32 InActionId);
+	UE_API void RegisterPasteBonesGlobalAction(FInteractiveToolActionSet& InOutActionSet, const int32 InActionId);
+	UE_API void RegisterPasteBoneTransformsLocalAction(FInteractiveToolActionSet& InOutActionSet, const int32 InActionId);
+	UE_API void RegisterPasteBoneTransformsGlobalAction(FInteractiveToolActionSet& InOutActionSet, const int32 InActionId);
 	UE_API void RegisterDuplicateAction(FInteractiveToolActionSet& InOutActionSet, const int32 InActionId);
 	UE_API void RegisterSelectComponentsAction(FInteractiveToolActionSet& InOutActionSet, const int32 InActionId);
 	UE_API void RegisterSelectionFilterCyclingAction(FInteractiveToolActionSet& InOutActionSet, const int32 InActionId);
 	UE_API void RegisterSnapAction(FInteractiveToolActionSet& InOutActionSet, const int32 InActionId);
 	
-	void PerformPasteAction(const bool bUsingGlobalTransforms);
+	void PerformPasteBonesAction(const bool bUsingGlobalTransforms);
+	void PerformPasteBoneTransformsAction(const bool bUsingGlobalTransforms);
 
 	UE_API TArray<int32> GetSelectedComponents() const;
 
@@ -335,6 +355,9 @@ private:
 	TWeakObjectPtr<USkeletalMesh> WeakMesh = nullptr;
 
 	TUniquePtr<FSkeletalMeshNotifierBindScope> SkeletonTreeNotifierBindScope;
+	
+	/** The skeleton tree this tool displays */
+	TSharedPtr<SReferenceSkeletonTree> SkeletonTree;
 };
 
 ENUM_CLASS_FLAGS(USkeletonEditingTool::EBoneSelectionMode);
diff --git a/Engine/Plugins/Runtime/MeshModelingToolset/Source/SkeletalMeshModifiers/Private/SkeletonClipboard.cpp b/Engine/Plugins/Runtime/MeshModelingToolset/Source/SkeletalMeshModifiers/Private/SkeletonClipboard.cpp
--- a/Engine/Plugins/Runtime/MeshModelingToolset/Source/SkeletalMeshModifiers/Private/SkeletonClipboard.cpp
+++ b/Engine/Plugins/Runtime/MeshModelingToolset/Source/SkeletalMeshModifiers/Private/SkeletonClipboard.cpp
@@ -4,6 +4,7 @@
 #include "SkeletonClipboard.h"
 
 #include "SkeletonModifier.h"
+#include "Algo/AllOf.h"
 
 #include "HAL/PlatformApplicationMisc.h"
 
@@ -13,6 +14,11 @@ namespace SkeletonClipboard
 {
 
 void CopyToClipboard(const USkeletonModifier& InModifier, const TArray<FName>& InBonesToCopy)
+{
+	CopyBonesToClipboard(InModifier, InBonesToCopy);
+}
+
+void CopyBonesToClipboard(const USkeletonModifier& InModifier, const TArray<FName>& InBonesToCopy)
 {
 	const FReferenceSkeleton& RefSkeleton = InModifier.GetReferenceSkeleton();
 
@@ -60,6 +66,7 @@ void CopyToClipboard(const USkeletonModifier& InModifier, const TArray<FName>& I
 
 	// convert bone names to clipboard data
 	FHierarchyClipboardData ClipboardData;
+	ClipboardData.Type = EBoneClipboardDataType::Bones;
 	ClipboardData.Bones.Reserve(InBonesToCopy.Num());
 	
 	for (const FName& BoneName : BonesToCopyWithChildren)
@@ -126,10 +133,10 @@ TArray<FName> PasteFromClipboard(USkeletonModifier& InOutModifier, const FName&
 	// DEPRECATED 5.8
 
 	constexpr bool bUsingGlobalTransforms = true;
-	return PasteFromClipboard(InOutModifier, InDefaultParent, bUsingGlobalTransforms);
+	return PasteBonesFromClipboard(InOutModifier, InDefaultParent, bUsingGlobalTransforms);
 }
 
-TArray<FName> PasteFromClipboard(USkeletonModifier& InOutModifier, const FName& InDefaultParent, const bool bUsingGlobalTransforms)
+TArray<FName> PasteBonesFromClipboard(USkeletonModifier& InOutModifier, const FName& InDefaultParent, const bool bUsingGlobalTransforms)
 {
 	// Get the text from the clipboard.
 	FString ClipboardText;
@@ -145,7 +152,7 @@ TArray<FName> PasteFromClipboard(USkeletonModifier& InOutModifier, const FName&
 	{
 		return {};
 	}
-	
+
 	TArray<FName> NewBoneNames;
 	NewBoneNames.Reserve(ClipboardData.Bones.Num());
 
@@ -181,8 +188,8 @@ TArray<FName> PasteFromClipboard(USkeletonModifier& InOutModifier, const FName&
 				if (bUsingGlobalTransforms)
 				{
 					constexpr bool bGlobal = true;
-					const FTransform& ParentGlobal = InOutModifier.GetBoneTransform(InDefaultParent, bGlobal);
-					LocalTransforms.Add(BoneData.Global.GetRelativeTransform(ParentGlobal));					
+					const FTransform& ParentTransformGlobal = InOutModifier.GetBoneTransform(InDefaultParent, bGlobal);
+					LocalTransforms.Add(BoneData.Global.GetRelativeTransform(ParentTransformGlobal));					
 				}
 				else
 				{
@@ -224,17 +231,17 @@ TArray<FName> PasteFromClipboard(USkeletonModifier& InOutModifier, const FName&
 
 	return NewBoneNames;
 }
-	
-TArray<FName> ClipboardDuplicate(USkeletonModifier& InOutModifier, const TArray<FName>& InBonesToDuplicate)
+
+TArray<FName> ClipboardDuplicateBones(USkeletonModifier& InOutModifier, const TArray<FName>& InBonesToDuplicate)
 {
 	if (InBonesToDuplicate.IsEmpty())
 	{
 		return {}; 
 	}
 
-	SkeletonClipboard::CopyToClipboard(InOutModifier, InBonesToDuplicate);
+	SkeletonClipboard::CopyBonesToClipboard(InOutModifier, InBonesToDuplicate);
 
-	if (!SkeletonClipboard::IsClipboardValid())
+	if (!SkeletonClipboard::IsBoneClipboardValid())
 	{
 		return {};
 	} 
@@ -278,11 +285,166 @@ TArray<FName> ClipboardDuplicate(USkeletonModifier& InOutModifier, const TArray<
 		RefSkeleton.GetBoneName(CommonParentIndex);
 						
 	constexpr bool bUsingGlobalTransforms = true;
-	return SkeletonClipboard::PasteFromClipboard(InOutModifier, CommonParentName, bUsingGlobalTransforms);
+	return SkeletonClipboard::PasteBonesFromClipboard(InOutModifier, CommonParentName, bUsingGlobalTransforms);
+}
+	
+void CopyBoneTransformsToClipboard(const USkeletonModifier& InModifier, TArray<FName> InBonesToCopy)
+{
+	const FReferenceSkeleton& RefSkeleton = InModifier.GetReferenceSkeleton();
+	
+	// ensure ordering
+	InBonesToCopy.StableSort([&](const FName Bone0, const FName Bone1)
+	{
+		const int32 Index0 = RefSkeleton.FindBoneIndex(Bone0);
+		const int32 Index1 = RefSkeleton.FindBoneIndex(Bone1);
+		return Index0 < Index1;
+	});
+
+	// convert bone names to clipboard data
+	FHierarchyClipboardData ClipboardData;
+	ClipboardData.Type = EBoneClipboardDataType::BoneTransforms;
+	ClipboardData.Bones.Reserve(InBonesToCopy.Num());
+	
+	for (const FName& BoneName : InBonesToCopy)
+	{
+		const FName ParentName = InModifier.GetParentName(BoneName);
+		
+		FBoneClipboardData BoneData;
+		{
+			BoneData.BoneName = BoneName;
+			
+			BoneData.ParentBoneName = ParentName;
+			
+			constexpr bool bGlobal = true;
+			BoneData.Global = InModifier.GetBoneTransform(BoneName, bGlobal);
+			
+			constexpr bool bLocal = false;
+			BoneData.Local = InModifier.GetBoneTransform(BoneName, bLocal);
+		}
+		
+		ClipboardData.Bones.Add(BoneData);
+	}
+	
+	// convert data to text
+	FString ClipboardText;
+	static const FHierarchyClipboardData DefaultData;
+	FHierarchyClipboardData::StaticStruct()->ExportText(ClipboardText, &ClipboardData, &DefaultData, nullptr, PPF_None, nullptr);
+
+	// copy to clipboard
+	FPlatformApplicationMisc::ClipboardCopy(*ClipboardText);
 }
+
+TArray<FName> PasteBoneTransformsFromClipboard(
+	USkeletonModifier& InOutModifier,
+	TArray<FName> InCandidateBoneNames,
+	const EPasteBoneTransformsFromClipboardFlags Flags)
+{
+	// Get the text from the clipboard.
+	FString ClipboardText;
+	FPlatformApplicationMisc::ClipboardPaste(ClipboardText);
+
+	FHierarchyClipboardData ClipboardData;
+	FHierarchyImportErrorContext ErrorPipe(true);
+	
+	UScriptStruct* DataStruct = FHierarchyClipboardData::StaticStruct();
+	DataStruct->ImportText(*ClipboardText, &ClipboardData, nullptr, PPF_None, &ErrorPipe, DataStruct->GetName(), true);
+
+	if (ErrorPipe.NumErrors > 0 || ClipboardData.Bones.IsEmpty())
+	{
+		return {};
+	}
+	
+	TArray<FName> TargetBoneNames;
+	TargetBoneNames.Reserve(ClipboardData.Bones.Num());
+
+	TArray<FTransform> LocalTransforms;
+	LocalTransforms.Reserve(ClipboardData.Bones.Num());
+	
+	const bool bTargetMatchesSource = Algo::AllOf(ClipboardData.Bones,
+		[&InCandidateBoneNames](const FBoneClipboardData& BoneData)
+		{
+			return InCandidateBoneNames.Contains(BoneData.BoneName);
+		});
+		
+	const auto GetTargetTransformLambda(
+		[&InOutModifier, Flags](const FBoneClipboardData& BoneData, const FName& TargetBoneName)
+		{			
+			const FName ParentBoneName = InOutModifier.GetParentName(TargetBoneName);
+
+			if (ParentBoneName == NAME_None)
+			{
+				return BoneData.Global;				
+			}
+			else if (EnumHasAnyFlags(Flags, EPasteBoneTransformsFromClipboardFlags::UsingGlobalTransforms))
+			{
+				constexpr bool bGlobal = true;
+				const FTransform& ParentTransformGlobal = InOutModifier.GetBoneTransform(ParentBoneName, bGlobal);
+				return BoneData.Global.GetRelativeTransform(ParentTransformGlobal);					
+			}
+			else
+			{
+				return BoneData.Local;
+			}
+		});
+	
+	if (bTargetMatchesSource)
+	{
+		// the clipboard data matches bone names exactly, paste transforms by name
+		for (const FBoneClipboardData& BoneData : ClipboardData.Bones)
+		{
+			const FName* TargetBoneNamePtr = InCandidateBoneNames.FindByKey(BoneData.BoneName);
+			check(TargetBoneNamePtr);
+			
+			TargetBoneNames.Add(*TargetBoneNamePtr);
+			LocalTransforms.Add(GetTargetTransformLambda(BoneData, *TargetBoneNamePtr));
+		}
+	}
+	else
+	{
+		// tagets don't match the clipboard data exactly, paste by bone index
+		const FReferenceSkeleton& RefSkeleton = InOutModifier.GetReferenceSkeleton();
+	
+		for (int32 BoneIndex = 0; BoneIndex < ClipboardData.Bones.Num(); BoneIndex++)
+		{
+			if (InCandidateBoneNames.IsValidIndex(BoneIndex))
+			{
+				const FBoneClipboardData& BoneData = ClipboardData.Bones[BoneIndex];
+				const FName& TargetBoneName = InCandidateBoneNames[BoneIndex];
+				
+				TargetBoneNames.Add(TargetBoneName);
+				LocalTransforms.Add(GetTargetTransformLambda(BoneData, TargetBoneName));
+			}
+			else
+			{
+				break;
+			}
+		}
+				
+		// ensure ordering of targets (source was already ordered when copying)
+		TargetBoneNames.StableSort([&](const FName Bone0, const FName Bone1)
+		{
+			const int32 Index0 = RefSkeleton.FindBoneIndex(Bone0);
+			const int32 Index1 = RefSkeleton.FindBoneIndex(Bone1);
+			return Index0 < Index1;
+		});
+	}
+	
+	// set transforms
+	const bool bMoveChildren = EnumHasAnyFlags(Flags, EPasteBoneTransformsFromClipboardFlags::UpdateChildren);
+	InOutModifier.SetBonesTransforms(TargetBoneNames, LocalTransforms, bMoveChildren);
+	
+	return TargetBoneNames;
+}	
 	
 bool IsClipboardValid()
 {
+	// DEPRECATED 5.8
+	return IsBoneClipboardValid();
+}
+
+bool IsBoneClipboardValid()
+{
+	// DEPRECATED 5.8
 	FString ClipboardText;
 	FPlatformApplicationMisc::ClipboardPaste(ClipboardText);
 
@@ -292,7 +454,28 @@ bool IsClipboardValid()
 	UScriptStruct* DataStruct = FHierarchyClipboardData::StaticStruct();
 	DataStruct->ImportText(*ClipboardText, &ClipboardData, nullptr, PPF_None, &ErrorPipe, DataStruct->GetName(), true);
 
-	return ErrorPipe.NumErrors == 0 && !ClipboardData.Bones.IsEmpty();
+	return 
+		ErrorPipe.NumErrors == 0 && 
+		ClipboardData.Type == EBoneClipboardDataType::Bones &&
+		!ClipboardData.Bones.IsEmpty();
 }
+
+bool IsBoneTransformClipboardValid()
+{
+	// DEPRECATED 5.8
+	FString ClipboardText;
+	FPlatformApplicationMisc::ClipboardPaste(ClipboardText);
+
+	FHierarchyClipboardData ClipboardData;
+	FHierarchyImportErrorContext ErrorPipe(false);
 	
+	UScriptStruct* DataStruct = FHierarchyClipboardData::StaticStruct();
+	DataStruct->ImportText(*ClipboardText, &ClipboardData, nullptr, PPF_None, &ErrorPipe, DataStruct->GetName(), true);
+
+	return 
+		ErrorPipe.NumErrors == 0 && 
+		ClipboardData.Type == EBoneClipboardDataType::BoneTransforms &&
+		!ClipboardData.Bones.IsEmpty();
+}
+
 }
diff --git a/Engine/Plugins/Runtime/MeshModelingToolset/Source/SkeletalMeshModifiers/Public/SkeletonClipboard.h b/Engine/Plugins/Runtime/MeshModelingToolset/Source/SkeletalMeshModifiers/Public/SkeletonClipboard.h
--- a/Engine/Plugins/Runtime/MeshModelingToolset/Source/SkeletalMeshModifiers/Public/SkeletonClipboard.h
+++ b/Engine/Plugins/Runtime/MeshModelingToolset/Source/SkeletalMeshModifiers/Public/SkeletonClipboard.h
@@ -17,35 +17,97 @@ class USkeletonModifier;
 
 namespace SkeletonClipboard
 {
+	/** DEPRECATED 5.8, instead use CopyBonesToClipboard */
+	UE_DEPRECATED(5.8, "Instead use CopyBonesToClipboard to avoid ambiguity with CopyBoneTransformsToClipboard")
+	SKELETALMESHMODIFIERS_API void CopyToClipboard(
+		const USkeletonModifier& InModifier, 
+		const TArray<FName>& InBonesToCopy);
+
 	/**
 	 * Copies InBonesToCopy data (name, parent, transform) into the clipboard. Also copies children of these bones.
-	 * Hierarchy info & transforms are retrieved via the InModifier.
+	 *
+	 * @param InModifier				The skeleton modifier that provides the skeleton and the bones to copy.
+	 * @param InBonesToCopy				The bones from which transforms are copied
 	 */
-	SKELETALMESHMODIFIERS_API void CopyToClipboard(const USkeletonModifier& InModifier, const TArray<FName>& InBonesToCopy);
+	SKELETALMESHMODIFIERS_API void CopyBonesToClipboard(
+		const USkeletonModifier& InModifier, 
+		const TArray<FName>& InBonesToCopy);
 
+	/** DEPRECATED 5.8, instead use PasteBonesFromClipboard	 */
+	UE_DEPRECATED(5.8, "Instead use the PasteBonesFromClipboard that offers the bUsingGlobalTransforms argument")
+	SKELETALMESHMODIFIERS_API TArray<FName> PasteFromClipboard(
+		USkeletonModifier& InOutModifier, 
+		const FName& InDefaultParent);
+	
 	/**
 	 * Paste new bones from the clipboard InBonesToCopy data (name, parent, transform) into the InOutModifier.
-	 * InDefaultParent is used as the default parent if a proper parent can't be retrieved from the clipboard data. 
+	 * 
+	 * @param InOutModifier				The skeleton modifier that provides the skeleton to which the bones are pasted
+	 * @param InDefaultParent			The parent to which bones are pasted
+	 * @param bUsingGlobalTransforms	If true, pastes global transforms
 	 */
-	UE_DEPRECATED(5.8, "Instead use the overload that offers the bUsingGlobalTransforms argument")
-	SKELETALMESHMODIFIERS_API TArray<FName> PasteFromClipboard(USkeletonModifier& InOutModifier, const FName& InDefaultParent);
+	 SKELETALMESHMODIFIERS_API TArray<FName> PasteBonesFromClipboard(
+		USkeletonModifier& InOutModifier, 
+		const FName& InDefaultParent, 
+		const bool bUsingGlobalTransforms);
 	
 	/**
-	 * Paste new bones from the clipboard InBonesToCopy data (name, parent, transform) into the InOutModifier.
-	 * InDefaultParent is used as the default parent if a proper parent can't be retrieved from the clipboard data. 
+	 * Duplicates InBonesToCopy data (name, parent, transform) using the clipboard. Also duplicates children of these bones.
+	 * 
+	 * @param InOutModifier				The skeleton modifier that provides the skeleton and the bones to duplicate.
+	 * @param InBonesToDuplicate		The bones that are being duplicated
 	 */
-   SKELETALMESHMODIFIERS_API TArray<FName> PasteFromClipboard(USkeletonModifier& InOutModifier, const FName& InDefaultParent, const bool bUsingGlobalTransforms);
+	SKELETALMESHMODIFIERS_API TArray<FName> ClipboardDuplicateBones(
+		USkeletonModifier& InOutModifier, 
+		const TArray<FName>& InBonesToDuplicate);
 	
 	/**
-	 * Duplicates InBonesToCopy data (name, parent, transform) using the clipboard. Also duplicates children of these bones.
+	 * Copies transforms of InBonesToCopy datainto the clipboard. 
 	 * Hierarchy info & transforms are retrieved via the InModifier.
+	 * 
+	 * @param InModifier				The skeleton modifier that provides the skeleton and the bones from which transforms are copied.
+	 * @param InBonesToCopy				The bones from which transforms are copied
 	 */
-	SKELETALMESHMODIFIERS_API TArray<FName> ClipboardDuplicate(USkeletonModifier& InOutModifier, const TArray<FName>& InBonesToDuplicate);
-	
+	SKELETALMESHMODIFIERS_API void CopyBoneTransformsToClipboard(
+		const USkeletonModifier& InModifier, 
+		TArray<FName> InBonesToCopy);
+
+	/** Flags for PasteBoneTransformsFromClipboard */
+	enum class EPasteBoneTransformsFromClipboardFlags
+	{
+		None					= 0,
+		UpdateChildren			= 1 << 0,
+		UsingGlobalTransforms	= 1 << 1,
+	};
+	ENUM_CLASS_FLAGS(EPasteBoneTransformsFromClipboardFlags);
+
 	/**
-	 * Checks whether the clipboard data is valid. 
+	 * Paste transforms from the clipboard InBonesToCopy data into the InOutModifier.
+	 * 
+	 * @param InOutModifier				The skeleton modifier that provides the skeleton and the bones to which transforms are pasted.
+	 * @param InCandidateBoneNames		Bones where the transforms should be pasted, not necessarily considered
+	 * @param Flags						Flags to define how bone transforms should be pasted
+	 * @return							The bones that were affected by the paste action
 	 */
+	SKELETALMESHMODIFIERS_API TArray<FName> PasteBoneTransformsFromClipboard(
+		USkeletonModifier& InOutModifier,
+		TArray<FName> InCandidateBoneNames,
+		const EPasteBoneTransformsFromClipboardFlags Flags);
+		
+	/** DEPRECATED 5.8, instead use IsBoneClipboardValid and IsBoneTransformClipboardValid */
+	UE_DEPRECATED(5.8, "Instead use IsBoneClipboardValid and IsBoneTransformClipboardValid")
 	SKELETALMESHMODIFIERS_API bool IsClipboardValid();
+	
+	/**
+	 * Checks whether bone clipboard data is valid. 
+	 */
+	SKELETALMESHMODIFIERS_API bool IsBoneClipboardValid();
+
+	/**
+	 * Checks whether bone transform clipboard data is valid. 
+	 */
+	SKELETALMESHMODIFIERS_API bool IsBoneTransformClipboardValid();
+
 }
 
 USTRUCT()
@@ -80,6 +142,14 @@ struct FBoneClipboardData
 	FTransform Local = FTransform::Identity;
 };
 
+UENUM()
+enum class EBoneClipboardDataType : uint8
+{
+	Invalid,
+	Bones,
+	BoneTransforms
+};
+
 USTRUCT()
 struct FHierarchyClipboardData
 {
@@ -87,5 +157,7 @@ struct FHierarchyClipboardData
 	
 	UPROPERTY()
 	TArray<FBoneClipboardData> Bones;
+	
+	UPROPERTY()
+	EBoneClipboardDataType Type = EBoneClipboardDataType::Invalid;
 };
-
