[SkeletalMeshEditingTools] Improve skeleton tree views retaining their states after edits

#Synced-CL 51338101

diff --git a/Engine/Plugins/Runtime/MeshModelingToolset/Source/MeshModelingToolsEditorOnly/Private/SkeletalMesh/SReferenceSkeletonTree.cpp b/Engine/Plugins/Runtime/MeshModelingToolset/Source/MeshModelingToolsEditorOnly/Private/SkeletalMesh/SReferenceSkeletonTree.cpp
--- a/Engine/Plugins/Runtime/MeshModelingToolset/Source/MeshModelingToolsEditorOnly/Private/SkeletalMesh/SReferenceSkeletonTree.cpp
+++ b/Engine/Plugins/Runtime/MeshModelingToolset/Source/MeshModelingToolsEditorOnly/Private/SkeletalMesh/SReferenceSkeletonTree.cpp
@@ -14,6 +14,7 @@
 
 #include "Framework/MultiBox/MultiBoxBuilder.h"
 #include "Editor.h"
+#include "Algo/RemoveIf.h"
 
 #include "Misc/ITransaction.h"
 
@@ -48,6 +49,42 @@ void FSkeletonModifierChange::Revert(UObject* Object)
 	Modifier->ExternalUpdate(PreChangeSkeleton, PreBoneTracker);
 }
 
+/** The actual tree view for this widget */
+class SReferenceSkeletonTreeView 
+	: public STreeView<TSharedPtr<FBoneElement>>
+{
+public:
+	/** Stores sparse item infos */
+	void StoreSparseItemInfos()
+	{
+		LastSparseItemInfos = STreeView<TSharedPtr<FBoneElement>>::SparseItemInfos;
+	}
+	
+	/** Restores sparse item infos from last known infos */
+	void RestoreSparseItemInfos(const TArray<TSharedPtr<FBoneElement>> TreeItems)
+	{
+		for (const auto& [TreeItem, SparseInfo] : LastSparseItemInfos)
+		{
+			const TSharedPtr<FBoneElement>* TreeItemPtr = TreeItems.FindByPredicate(
+				[&TreeItem](const TSharedPtr<FBoneElement>& OtherTreeItem)
+				{
+					return TreeItem->BoneName == OtherTreeItem->BoneName;
+				});
+		}
+		
+		for (auto It = LastSparseItemInfos.CreateIterator(); It; ++It)
+		{
+			if (!TreeItems.Contains(It->Key))
+			{
+				It.RemoveCurrent();
+			}
+		}
+	}
+	
+private:
+	/** Holds the last known sparse item infos */
+	STreeView<TSharedPtr<FBoneElement>>::TSparseItemMap LastSparseItemInfos;
+};
 }
 
 FBoneElement::FBoneElement(const FName& InBoneName, TWeakObjectPtr<USkeletonModifier> InModifier)
@@ -375,7 +412,7 @@ void SReferenceSkeletonTree::Construct(const FArguments& InArgs)
 				.Padding(2.0f)
 				.BorderImage(FAppStyle::GetBrush("SCSEditor.TreePanel"))
 				[
-					SAssignNew(TreeView, SRefSkeletonTreeView)
+					SAssignNew(TreeView, SReferenceSkeletonTreeView)
 					.TreeItemsSource(&RootElements)
 					.SelectionMode(ESelectionMode::Multi)
 					.AllowInvisibleItemSelection(true)
@@ -974,6 +1011,8 @@ void SReferenceSkeletonTree::RefreshTreeView(bool IsInitialSetup /*=false*/)
 		return;
 	}
 	
+	TreeView->StoreSparseItemInfos();
+	
 	const FReferenceSkeleton& RefSkeleton = Modifier->GetReferenceSkeleton();
 	const TArray<FMeshBoneInfo>& BoneInfos = RefSkeleton.GetRawRefBoneInfo();
 	
@@ -1092,6 +1131,8 @@ void SReferenceSkeletonTree::RefreshTreeView(bool IsInitialSetup /*=false*/)
 	}
 	
 	TreeView->RequestTreeRefresh();
+	
+	TreeView->RestoreSparseItemInfos(AllElements);
 }
 
 void SReferenceSkeletonTree::HandleGetChildrenForTree(
@@ -1419,4 +1460,4 @@ void SReferenceSkeletonTree::CancelChange()
 	ActiveChange.Reset();
 }
 
-#undef LOCTEXT_NAMESPACE
\ No newline at end of file
+#undef LOCTEXT_NAMESPACE
diff --git a/Engine/Plugins/Runtime/MeshModelingToolset/Source/MeshModelingToolsEditorOnly/Public/SkeletalMesh/SReferenceSkeletonTree.h b/Engine/Plugins/Runtime/MeshModelingToolset/Source/MeshModelingToolsEditorOnly/Public/SkeletalMesh/SReferenceSkeletonTree.h
--- a/Engine/Plugins/Runtime/MeshModelingToolset/Source/MeshModelingToolsEditorOnly/Public/SkeletalMesh/SReferenceSkeletonTree.h
+++ b/Engine/Plugins/Runtime/MeshModelingToolset/Source/MeshModelingToolsEditorOnly/Public/SkeletalMesh/SReferenceSkeletonTree.h
@@ -22,6 +22,7 @@ class SInlineEditableTextBlock;
 
 namespace ReferenceSkeletonTreeLocals
 {
+class SReferenceSkeletonTreeView;
 
 class FSkeletonModifierChange : public FCommandChange
 {
@@ -194,8 +195,6 @@ private:
 	bool bIsReadOnly = false;
 };
 
-typedef STreeView< TSharedPtr<FBoneElement> > SRefSkeletonTreeView;
-
 class SReferenceSkeletonTree;
 
 class FReferenceSkeletonWidgetNotifier: public ISkeletalMeshNotifier
@@ -210,6 +209,8 @@ private:
 
 class SReferenceSkeletonTree : public SCompoundWidget, public FEditorUndoClient
 {
+	using SReferenceSkeletonTreeView = ReferenceSkeletonTreeLocals::SReferenceSkeletonTreeView;
+	
 public:
 	SLATE_BEGIN_ARGS(SReferenceSkeletonTree)
 		: _bIsReadOnly(false)
@@ -304,7 +305,7 @@ private:
 	TSharedPtr<FUICommandList> CommandList;
 	
 	// Tree view
-	TSharedPtr<SRefSkeletonTreeView> TreeView;
+	TSharedPtr<SReferenceSkeletonTreeView> TreeView;
 	TArray<TSharedPtr<FBoneElement>> RootElements;
 	TArray<TSharedPtr<FBoneElement>> AllElements;
 
diff --git a/Engine/Source/Editor/SkeletonEditor/Private/SSkeletonTree.cpp b/Engine/Source/Editor/SkeletonEditor/Private/SSkeletonTree.cpp
--- a/Engine/Source/Editor/SkeletonEditor/Private/SSkeletonTree.cpp
+++ b/Engine/Source/Editor/SkeletonEditor/Private/SSkeletonTree.cpp
@@ -84,6 +84,7 @@ template <typename ItemType>
 class SSkeletonTreeView : public STreeView<ItemType>
 {
 public:
+
 	bool Private_CanItemBeSelected(ItemType InItem) const
 	{
 		return !(InItem->GetFilterResult() == ESkeletonTreeFilterResult::ShownDescendant && GetMutableDefault<UPersonaOptions>()->bHideParentsWhenFiltering);
@@ -114,7 +115,7 @@ public:
 			Swap( RangeStartIndex, RangeEndIndex );
 		}
 
-		for( int32 ItemIndex = RangeStartIndex; ItemIndex <= RangeEndIndex; ++ItemIndex )
+		for( int32 ItemIndex = RangeStartIndex; ItemIndex < RangeEndIndex; ++ItemIndex )
 		{
 			// check if this item can actually be selected
 			if(Private_CanItemBeSelected(ItemsSourceRef[ItemIndex]))
@@ -149,6 +150,36 @@ public:
 			}
 		}
 	}
+
+	/** Stores sparse item infos */
+	void StoreSparseItemInfos()
+	{
+		CachedSparseItemInfos = STreeView<ItemType>::SparseItemInfos;
+	}
+	
+	/** Restores sparse item infos from last known infos */
+	void RestoreSparseItemInfos(const TArray<TSharedPtr<ISkeletonTreeItem>> TreeItems)
+	{
+		STreeView<ItemType>::SparseItemInfos.Reset();
+		
+		for (const auto& [TreeItem, SparseItemInfo] : CachedSparseItemInfos)
+		{
+			const TSharedPtr<ISkeletonTreeItem>* TreeItemPtr = TreeItems.FindByPredicate(
+				[&TreeItem](const TSharedPtr<ISkeletonTreeItem>& OtherTreeItem)
+				{
+					return TreeItem->GetRowItemName() == OtherTreeItem->GetRowItemName();
+				});
+				
+			if (TreeItemPtr)
+			{
+				STreeView<ItemType>::SparseItemInfos.Add(*TreeItemPtr, SparseItemInfo);
+			}
+		}
+	}
+	
+private:
+	/** Holds the last known sparse item infos */
+	STreeView<ItemType>::TSparseItemMap CachedSparseItemInfos;
 };
 
 void SSkeletonTree::Construct(const FArguments& InArgs, const TSharedRef<FEditableSkeleton>& InEditableSkeleton, const FSkeletonTreeArgs& InSkeletonTreeArgs)
@@ -701,6 +732,8 @@ void SSkeletonTree::CreateTreeColumns()
 
 void SSkeletonTree::CreateFromSkeleton()
 {	
+	SkeletonTreeView->StoreSparseItemInfos();
+
 	// save selected items
 	FScopedSavedSelection ScopedSelection(SharedThis(this));
 
@@ -711,6 +744,8 @@ void SSkeletonTree::CreateFromSkeleton()
 	FSkeletonTreeBuilderOutput Output(Items, LinearItems);
 	Builder->Build(Output);
 	ApplyFilter();
+	
+	SkeletonTreeView->RestoreSparseItemInfos(LinearItems);
 }
 
 void SSkeletonTree::ApplyFilter()
@@ -1944,6 +1979,11 @@ void SSkeletonTree::DeselectAll()
 	}
 }
 
+TArray<TSharedPtr<ISkeletonTreeItem>> SSkeletonTree::GetSelectedItems() const 
+{ 
+	return SkeletonTreeView->GetSelectedItems(); 
+}
+
 void SSkeletonTree::NotifyUser( FNotificationInfo& NotificationInfo )
 {
 	TSharedPtr<SNotificationItem> Notification = FSlateNotificationManager::Get().AddNotification( NotificationInfo );
diff --git a/Engine/Source/Editor/SkeletonEditor/Private/SSkeletonTree.h b/Engine/Source/Editor/SkeletonEditor/Private/SSkeletonTree.h
--- a/Engine/Source/Editor/SkeletonEditor/Private/SSkeletonTree.h
+++ b/Engine/Source/Editor/SkeletonEditor/Private/SSkeletonTree.h
@@ -36,6 +36,7 @@ class FUICommandList_Pinnable;
 class IPersonaPreviewScene;
 class SBlendProfilePicker;
 class SComboButton;
+template <typename ItemType> class SSkeletonTreeView;
 class UBlendProfile;
 class UToolMenu;
 struct FNotificationInfo;
@@ -79,7 +80,7 @@ public:
 	virtual void SetSelectedBone(const FName& InBoneName, ESelectInfo::Type InSelectInfo) override;
 	virtual void SetSelectedBones(const TArray<FName>& InBoneNames, ESelectInfo::Type InSelectInfo) override;
 	virtual void DeselectAll() override;
-	virtual TArray<TSharedPtr<ISkeletonTreeItem>> GetSelectedItems() const override { return SkeletonTreeView->GetSelectedItems(); }
+	virtual TArray<TSharedPtr<ISkeletonTreeItem>> GetSelectedItems() const override;
 	virtual void SelectItemsBy(TFunctionRef<bool(const TSharedRef<ISkeletonTreeItem>&, bool&)> Predicate, const ESelectInfo::Type SelectInfo) const override;
 	virtual void DuplicateAndSelectSocket(const FSelectedSocketInfo& SocketInfoToDuplicate, const FName& NewParentBoneName = FName()) override;
 	virtual void ResetBoneTransforms(const bool bSelectedOnly) override;
@@ -388,7 +389,7 @@ private:
 	TSharedPtr<SOverlay> TreeHolder;
 
 	/** Widget used to display the skeleton hierarchy */
-	TSharedPtr<STreeView<TSharedPtr<class ISkeletonTreeItem>>> SkeletonTreeView;
+	TSharedPtr<SSkeletonTreeView<TSharedPtr<class ISkeletonTreeItem>>> SkeletonTreeView;
 
 	/** A tree of unfiltered items */
 	TArray<TSharedPtr<class ISkeletonTreeItem>> Items;
