[SkeletalMeshEditingTools] Fix implementaiton of RestoreSparseItemInfos to properly restore skeleton tree states in the skeletal mesh editor

#rnx

#Synced-CL 51413426

diff --git a/Engine/Plugins/Runtime/MeshModelingToolset/Source/MeshModelingToolsEditorOnly/Private/SkeletalMesh/SReferenceSkeletonTree.cpp b/Engine/Plugins/Runtime/MeshModelingToolset/Source/MeshModelingToolsEditorOnly/Private/SkeletalMesh/SReferenceSkeletonTree.cpp
--- a/Engine/Plugins/Runtime/MeshModelingToolset/Source/MeshModelingToolsEditorOnly/Private/SkeletalMesh/SReferenceSkeletonTree.cpp
+++ b/Engine/Plugins/Runtime/MeshModelingToolset/Source/MeshModelingToolsEditorOnly/Private/SkeletalMesh/SReferenceSkeletonTree.cpp
@@ -57,33 +57,38 @@ public:
 	/** Stores sparse item infos */
 	void StoreSparseItemInfos()
 	{
-		LastSparseItemInfos = STreeView<TSharedPtr<FBoneElement>>::SparseItemInfos;
+		CachedSparseItemInfos = STreeView<TSharedPtr<FBoneElement>>::SparseItemInfos;
 	}
 	
 	/** Restores sparse item infos from last known infos */
 	void RestoreSparseItemInfos(const TArray<TSharedPtr<FBoneElement>> TreeItems)
 	{
-		for (const auto& [TreeItem, SparseInfo] : LastSparseItemInfos)
+		for (auto It = STreeView<TSharedPtr<FBoneElement>>::SparseItemInfos.CreateIterator(); It; ++It)
 		{
-			const TSharedPtr<FBoneElement>* TreeItemPtr = TreeItems.FindByPredicate(
-				[&TreeItem](const TSharedPtr<FBoneElement>& OtherTreeItem)
-				{
-					return TreeItem->BoneName == OtherTreeItem->BoneName;
-				});
+			if (!TreeItems.Contains(It->Key))
+			{
+				It.RemoveCurrent();
+			}
 		}
 		
-		for (auto It = LastSparseItemInfos.CreateIterator(); It; ++It)
+		for (auto It = CachedSparseItemInfos.CreateIterator(); It; ++It)
 		{
-			if (!TreeItems.Contains(It->Key))
+			const TSharedPtr<FBoneElement>* TreeItemPtr = TreeItems.FindByPredicate(
+				[It](const TSharedPtr<FBoneElement>& TreeItem)
+				{
+					return It->Key->BoneName == TreeItem->BoneName;
+				});
+				
+			if (TreeItemPtr)
 			{
-				It.RemoveCurrent();
+				STreeView<TSharedPtr<FBoneElement>>::SparseItemInfos.Add(*TreeItemPtr, It->Value);
 			}
 		}
 	}
 	
 private:
 	/** Holds the last known sparse item infos */
-	STreeView<TSharedPtr<FBoneElement>>::TSparseItemMap LastSparseItemInfos;
+	STreeView<TSharedPtr<FBoneElement>>::TSparseItemMap CachedSparseItemInfos;
 };
 }
 
diff --git a/Engine/Source/Editor/SkeletonEditor/Private/SSkeletonTree.cpp b/Engine/Source/Editor/SkeletonEditor/Private/SSkeletonTree.cpp
--- a/Engine/Source/Editor/SkeletonEditor/Private/SSkeletonTree.cpp
+++ b/Engine/Source/Editor/SkeletonEditor/Private/SSkeletonTree.cpp
@@ -159,8 +159,14 @@ public:
 	
 	/** Restores sparse item infos from last known infos */
 	void RestoreSparseItemInfos(const TArray<TSharedPtr<ISkeletonTreeItem>> TreeItems)
-	{
-		STreeView<ItemType>::SparseItemInfos.Reset();
+	{		
+		for (auto It = STreeView<ItemType>::SparseItemInfos.CreateIterator(); It; ++It)
+		{
+			if (!TreeItems.Contains(It->Key))
+			{
+				It.RemoveCurrent();
+			}
+		}
 		
 		for (const auto& [TreeItem, SparseItemInfo] : CachedSparseItemInfos)
 		{
