[SkeletalMeshEditingTools] Ctrl-clicking the selection mode buttons in the skin weigh tool now converts selection from one mode to another

#Synced-CL 51413426

diff --git a/Engine/Plugins/Animation/SkeletalMeshModelingTools/Source/SkeletalMeshModelingTools/Private/DetailCustomization/SkinWeightDetailCustomization.cpp b/Engine/Plugins/Animation/SkeletalMeshModelingTools/Source/SkeletalMeshModelingTools/Private/DetailCustomization/SkinWeightDetailCustomization.cpp
--- a/Engine/Plugins/Animation/SkeletalMeshModelingTools/Source/SkeletalMeshModelingTools/Private/DetailCustomization/SkinWeightDetailCustomization.cpp
+++ b/Engine/Plugins/Animation/SkeletalMeshModelingTools/Source/SkeletalMeshModelingTools/Private/DetailCustomization/SkinWeightDetailCustomization.cpp
@@ -382,7 +382,8 @@ void FSkinWeightDetailCustomization::AddSelectionUI(IDetailLayoutBuilder& Detail
 		{
 			if (ToolSettings.IsValid())
 			{
-				ToolSettings->SetComponentMode(Mode);
+				const bool bConvertSelection = FSlateApplication::Get().GetModifierKeys().IsControlDown();
+				ToolSettings->SetComponentSelectionMode(Mode, bConvertSelection);
 			}
 		}),
 		FCanExecuteAction::CreateLambda([this]()
@@ -403,17 +404,17 @@ void FSkinWeightDetailCustomization::AddSelectionUI(IDetailLayoutBuilder& Detail
 	AddToggleButtonForBool(
 		EComponentSelectionMode::Vertices,
 		LOCTEXT("VerticesLabel", "Vertices"),
-		LOCTEXT("VerticesTooltip", "Select mesh vertices."),
+		LOCTEXT("VerticesTooltip", "Select mesh vertices. Ctrl-clicking converts the current selection to vertices."),
 		"PolyEd.SelectCorners");
 	AddToggleButtonForBool(
 		EComponentSelectionMode::Edges,
 		LOCTEXT("EdgesLabel", "Edges"),
-		LOCTEXT("EdgesTooltip", "Select mesh edges."),
+		LOCTEXT("EdgesTooltip", "Select mesh edges. Ctrl-clicking converts the current selection to edges."),
 		"PolyEd.SelectEdges");
 	AddToggleButtonForBool(
 		EComponentSelectionMode::Faces,
 		LOCTEXT("FacesLabel", "Faces"),
-		LOCTEXT("FacesTooltip", "Select mesh faces."),
+		LOCTEXT("FacesTooltip", "Select mesh faces. Ctrl-clicking converts the current selection to faces."),
 		"PolyEd.SelectFaces");
 
 	ToolbarBuilder.EndBlockGroup();
@@ -429,18 +430,20 @@ void FSkinWeightDetailCustomization::AddSelectionUI(IDetailLayoutBuilder& Detail
 		[
 			SNew(SHorizontalBox)
 
-			// ISOLATE SELECTION
 			+SHorizontalBox::Slot()
 			.HAlign(HAlign_Center)
 			.VAlign(VAlign_Center)
 			.Padding(2.0f, WeightEditVerticalPadding)
 			[
 				SNew(SHorizontalBox)
+				
+				// Selection Modes
 				+SHorizontalBox::Slot()
 				[
 					ToolbarBuilder.MakeWidget()
 				]
 
+				// Isolate selection
 				+SHorizontalBox::Slot()
 				.AutoWidth()
 				.VAlign(VAlign_Center)
diff --git a/Engine/Plugins/Runtime/MeshModelingToolset/Source/MeshModelingToolsEditorOnly/Private/SkeletalMesh/SkinWeightsPaintTool.cpp b/Engine/Plugins/Runtime/MeshModelingToolset/Source/MeshModelingToolsEditorOnly/Private/SkeletalMesh/SkinWeightsPaintTool.cpp
--- a/Engine/Plugins/Runtime/MeshModelingToolset/Source/MeshModelingToolsEditorOnly/Private/SkeletalMesh/SkinWeightsPaintTool.cpp
+++ b/Engine/Plugins/Runtime/MeshModelingToolset/Source/MeshModelingToolsEditorOnly/Private/SkeletalMesh/SkinWeightsPaintTool.cpp
@@ -73,6 +73,312 @@
 
 #define LOCTEXT_NAMESPACE "USkinWeightsPaintTool"
 
+namespace UE::MeshModelingToolsEditorOnly::SelectionUtils
+{
+	/** Returns the geometry type that corresponds to a selection mode */
+	UE::Geometry::EGeometryElementType GetGeometryElementTypeFromSelectionMode(const EComponentSelectionMode SelectionMode)
+	{
+		using namespace UE::Geometry;
+		if (SelectionMode == EComponentSelectionMode::Vertices)
+		{
+			return EGeometryElementType::Vertex;
+		}
+		else if (SelectionMode == EComponentSelectionMode::Edges)
+		{
+			return EGeometryElementType::Edge;
+		}
+		else if (SelectionMode == EComponentSelectionMode::Faces)
+		{
+			return EGeometryElementType::Face;
+		}
+		else
+		{
+			ensureMsgf(0, TEXT("Unhandled selection mode"));
+			return EGeometryElementType::Vertex;
+		}
+	}
+
+	/** Util to convert selection from one selection mode to another */
+	struct FWeightToolScopedConvertSelection
+		: FNoncopyable
+	{
+		FWeightToolScopedConvertSelection(
+			USkinWeightsPaintTool& WeightTool,
+			const UE::Geometry::FDynamicMesh3& Mesh,
+			const EComponentSelectionMode TargetSelectionMode)
+		{
+			using namespace UE::Geometry;
+
+			UToolMeshSelector* MeshSelector = WeightTool.GetMainMeshSelector();
+			SelectionMechanic = MeshSelector ? MeshSelector->GetSelectionMechanic() : nullptr;
+			const USkinWeightsPaintToolProperties* ToolProperties = WeightTool.GetWeightToolProperties();
+
+			if (ensure(ToolProperties && SelectionMechanic))
+			{
+				const EComponentSelectionMode SourceSelectionMode = ToolProperties->ComponentSelectionMode;
+				if (SourceSelectionMode == TargetSelectionMode)
+				{
+					bPerformConversion = false;
+					return;
+				}
+
+				FGeometrySelection SourceSelection;
+				SourceSelection.TopologyType = EGeometryTopologyType::Triangle;
+				SourceSelection.ElementType = GetGeometryElementTypeFromSelectionMode(SourceSelectionMode);
+				
+				SelectionMechanic->GetSelection_AsTriangleTopology(SourceSelection);
+				
+				const auto VerticesToEdges = [&Mesh](const FGeometrySelection& InVertexSelection)
+					{
+						FGeometrySelection Result;						
+						Result.ElementType = EGeometryElementType::Edge;
+						Result.TopologyType = EGeometryTopologyType::Triangle;
+
+						for (const int EdgeID : Mesh.EdgeIndicesItr())
+						{
+							const FDynamicMesh3::FEdge Edge = Mesh.GetEdge(EdgeID);
+							if (InVertexSelection.Selection.Contains(Edge.Vert.A) &&
+								InVertexSelection.Selection.Contains(Edge.Vert.B))
+							{
+								const FMeshTriEdgeID TriEdgeID = Mesh.GetTriEdgeIDFromEdgeID(EdgeID);
+								Result.Selection.Add(TriEdgeID.Encoded());
+							}
+						}
+						
+						return Result;
+					};
+					
+				const auto VerticesToFaces = [&Mesh](const FGeometrySelection& InVertexSelection)
+					{
+						FGeometrySelection Result;
+						Result.ElementType = EGeometryElementType::Face;
+						Result.TopologyType = EGeometryTopologyType::Triangle;
+
+						for (const int32& TriangleID : Mesh.TriangleIndicesItr())
+						{
+							const FIndex3i Triangle = Mesh.GetTriangle(TriangleID);
+							if (InVertexSelection.Selection.Contains(Triangle.A) &&
+								InVertexSelection.Selection.Contains(Triangle.B) &&
+								InVertexSelection.Selection.Contains(Triangle.C))
+							{
+								Result.Selection.Add(TriangleID);
+							}
+						}
+						
+						return Result;
+					};
+									
+				const auto EdgesToVertices = [this, &Mesh](const FGeometrySelection& InVertexSelection)
+					{
+						FGeometrySelection Result;
+						Result.ElementType = EGeometryElementType::Vertex;
+						Result.TopologyType = EGeometryTopologyType::Triangle;
+
+						for (const uint64 EdgeID : InVertexSelection.Selection)
+						{
+							const FMeshTriEdgeID TriEdgeID(EdgeID);	
+							const int TriEdge = Mesh.GetTriEdge(TriEdgeID.TriangleID, TriEdgeID.TriEdgeIndex);
+							const FDynamicMesh3::FEdge Edge = Mesh.GetEdge(TriEdge);
+							
+							Result.Selection.Add(Edge.Vert.A);
+							Result.Selection.Add(Edge.Vert.B);
+						}
+						
+						return Result;
+					};
+					
+				const auto FacesToVertices = [this, &Mesh](const FGeometrySelection& InVertexSelection)
+					{
+						FGeometrySelection Result;
+						Result.ElementType = EGeometryElementType::Vertex;
+						Result.TopologyType = EGeometryTopologyType::Triangle;
+
+						for (const uint64 TriangleID : InVertexSelection.Selection)
+						{
+							const FIndex3i Triangle = Mesh.GetTriangle(TriangleID);
+							Result.Selection.Add(Triangle.A);
+							Result.Selection.Add(Triangle.B);
+							Result.Selection.Add(Triangle.C);
+						}
+						
+						return Result;
+					};
+				
+				if (SourceSelectionMode == EComponentSelectionMode::Vertices &&
+					TargetSelectionMode == EComponentSelectionMode::Edges)
+				{	
+					TargetSelection = VerticesToEdges(SourceSelection);
+				}
+				else if (
+					SourceSelectionMode == EComponentSelectionMode::Vertices &&
+					TargetSelectionMode == EComponentSelectionMode::Faces)
+				{	
+					TargetSelection = VerticesToFaces(SourceSelection);
+				}
+				else if (
+					SourceSelectionMode == EComponentSelectionMode::Edges &&
+					TargetSelectionMode == EComponentSelectionMode::Vertices)
+				{	
+					TargetSelection = EdgesToVertices(SourceSelection);
+				}
+				else if (
+					SourceSelectionMode == EComponentSelectionMode::Edges &&
+					TargetSelectionMode == EComponentSelectionMode::Faces)
+				{
+					TargetSelection = EdgesToVertices(SourceSelection);
+					TargetSelection = VerticesToFaces(TargetSelection);
+				}
+				else if (
+					SourceSelectionMode == EComponentSelectionMode::Faces &&
+					TargetSelectionMode == EComponentSelectionMode::Vertices)
+				{	
+					TargetSelection = FacesToVertices(SourceSelection);
+				}
+				else if (
+					SourceSelectionMode == EComponentSelectionMode::Faces &&
+					TargetSelectionMode == EComponentSelectionMode::Edges)
+				{
+					TargetSelection = FacesToVertices(SourceSelection);
+					TargetSelection = VerticesToEdges(TargetSelection);
+				}
+				else
+				{
+					ensureMsgf(0, TEXT("Unhandled selection mode"));
+				}
+				
+				SelectionMechanic->ClearSelection();
+			}
+		}
+		
+		~FWeightToolScopedConvertSelection()
+		{
+			if (bPerformConversion &&
+				SelectionMechanic)
+			{					
+				// Restore selection
+				constexpr bool bBroadcastSelectionChanged = true;
+				SelectionMechanic->SetSelection_AsTriangleTopology(TargetSelection, bBroadcastSelectionChanged);
+			}
+		}
+		
+	private:
+		bool bPerformConversion = true;
+		
+		UE::Geometry::FGeometrySelection TargetSelection;
+
+		UPolygonSelectionMechanic* SelectionMechanic = nullptr;
+	};
+	
+	/** Modes for the scoped retain selection util */
+	enum class EWeightToolScopedRetainSelectionDirection : uint8
+	{
+		RestoreInPartialMesh,
+		RestoreInFullMesh
+	};
+
+	/** Util to retain the selection when isolating or unisolating the selection */
+	struct FWeightToolScopedRetainSelection
+		: FNoncopyable
+	{
+		FWeightToolScopedRetainSelection(
+			USkinWeightsPaintTool& WeightTool,
+			const UE::Geometry::FDynamicSubmesh3& Mesh,
+			const EWeightToolScopedRetainSelectionDirection Direction)
+		{
+			using namespace UE::Geometry;
+
+			UToolMeshSelector* MeshSelector = WeightTool.GetMainMeshSelector();
+			SelectionMechanic = MeshSelector ? MeshSelector->GetSelectionMechanic() : nullptr;
+			const USkinWeightsPaintToolProperties* ToolProperties = WeightTool.GetWeightToolProperties();
+
+			if (ensure(ToolProperties && SelectionMechanic))
+			{
+				const EComponentSelectionMode SelectionMode = ToolProperties->ComponentSelectionMode;
+				const EGeometryElementType ElementType = GetGeometryElementTypeFromSelectionMode(SelectionMode);
+
+				UE::Geometry::FGeometrySelection OldSelection;
+
+				OldSelection.TopologyType = EGeometryTopologyType::Triangle;
+				NewSelection.TopologyType = EGeometryTopologyType::Triangle;
+
+				OldSelection.ElementType = ElementType;
+				NewSelection.ElementType = ElementType;
+
+				SelectionMechanic->GetSelection_AsTriangleTopology(OldSelection);
+
+				if (SelectionMode == EComponentSelectionMode::Vertices)
+				{
+					Algo::Transform(OldSelection.Selection, NewSelection.Selection,
+						[&Mesh, Direction](const uint64 ElementID)
+						{
+							if (Direction == EWeightToolScopedRetainSelectionDirection::RestoreInFullMesh)
+							{
+								return Mesh.MapVertexToBaseMesh(ElementID);
+							}
+							else
+							{
+								return Mesh.MapVertexToSubmesh(ElementID);
+							}
+						});
+				}
+				else if (SelectionMode == EComponentSelectionMode::Edges)
+				{
+					for (const uint64 EdgeID : OldSelection.Selection)
+					{
+						if (Direction == EWeightToolScopedRetainSelectionDirection::RestoreInFullMesh)
+						{
+							UE::Geometry::FMeshTriEdgeID TriEdgeID(EdgeID);							
+							TriEdgeID.TriangleID = Mesh.MapTriangleToBaseMesh(TriEdgeID.TriangleID);
+
+							NewSelection.Selection.Add(TriEdgeID.Encoded());
+						}
+						else
+						{
+							UE::Geometry::FMeshTriEdgeID TriEdgeID(EdgeID);
+							TriEdgeID.TriangleID = Mesh.MapTriangleToSubmesh(TriEdgeID.TriangleID);
+
+							NewSelection.Selection.Add(TriEdgeID.Encoded());
+						}
+					}
+				}
+				else if (SelectionMode == EComponentSelectionMode::Faces)
+				{
+					Algo::Transform(OldSelection.Selection, NewSelection.Selection,
+						[&Mesh, Direction](const uint64 Index)
+						{
+							if (Direction == EWeightToolScopedRetainSelectionDirection::RestoreInFullMesh)
+							{
+								return Mesh.MapTriangleToBaseMesh(Index);
+							}
+							else
+							{
+								return Mesh.MapTriangleToSubmesh(Index);
+							}
+						});
+				}
+				else
+				{
+					ensureMsgf(0, TEXT("Unhandled selection mode"));
+				}
+			}
+		}
+
+		~FWeightToolScopedRetainSelection()
+		{
+			if (SelectionMechanic)
+			{					
+				// Restore selection
+				SelectionMechanic->SetSelection_AsTriangleTopology(NewSelection);
+			}
+		}
+
+	private:
+		UE::Geometry::FGeometrySelection NewSelection;
+
+		UPolygonSelectionMechanic* SelectionMechanic = nullptr;
+	};
+}
+
 using namespace SkinPaintTool;
 
 static void ShowEditorMessage(ELogVerbosity::Type InMessageType, const FText& InMessage)
@@ -2327,6 +2633,40 @@
 
 void USkinWeightsPaintToolProperties::SetComponentMode(EComponentSelectionMode InComponentMode)
 {
+	// DEPRECATED 5.8
+	constexpr bool bConvertSelection = false;
+	SetComponentSelectionMode(InComponentMode, bConvertSelection);
+}
+
+void USkinWeightsPaintToolProperties::SetComponentSelectionMode(const EComponentSelectionMode InComponentMode, const bool bInConvertSelection)
+{
+	using namespace UE::MeshModelingToolsEditorOnly::SelectionUtils;
+	
+	if (!WeightTool)
+	{
+		return;
+	}
+
+	TSharedPtr<FWeightToolScopedConvertSelection> ConvertSelection;
+	if (bInConvertSelection)
+	{
+		const UPreviewMesh* PreviewMesh = WeightTool->GetPreviewMesh();
+		const FDynamicMesh3* DynamicMesh = PreviewMesh ? PreviewMesh->GetMesh() : nullptr;
+		if (DynamicMesh)
+		{
+			ConvertSelection = MakeShared<FWeightToolScopedConvertSelection>(*WeightTool, *DynamicMesh, InComponentMode);
+		}
+	}
+	else
+	{
+		UToolMeshSelector* MeshSelector = WeightTool->GetMainMeshSelector();
+		UPolygonSelectionMechanic* SelectionMechanic = MeshSelector ? MeshSelector->GetSelectionMechanic() : nullptr;
+		if (SelectionMechanic)
+		{	
+			SelectionMechanic->ClearSelection();
+		}
+	}
+
 	ComponentSelectionMode = InComponentMode;
 	
 	WeightTool->UpdateSelectorState();
@@ -4630,140 +4970,9 @@
 	bIsolatedMeshNeedsUpdated = true;
 }
 
-namespace UE::MeshModelingToolsEditorOnly
-{
-	enum class EWeightToolScopedRetainSelectionDirection : uint8
-	{
-		RestoreInPartialMesh,
-		RestoreInFullMesh
-	};
-
-	/** Util to retain the selection when isolating or unisolating the selection */
-	struct FWeightToolScopedRetainSelection
-		: FNoncopyable
-	{
-		FWeightToolScopedRetainSelection(
-			USkinWeightsPaintTool& InWeightTool,
-			const UE::Geometry::FDynamicSubmesh3& Mesh,
-			const EWeightToolScopedRetainSelectionDirection Direction)
-		{
-			using namespace UE::Geometry;
-
-			UToolMeshSelector* MeshSelector = InWeightTool.GetMainMeshSelector();
-			SelectionMechanic = MeshSelector ? MeshSelector->GetSelectionMechanic() : nullptr;
-			const USkinWeightsPaintToolProperties* ToolProperties = InWeightTool.GetWeightToolProperties();
-
-			if (ensure(ToolProperties && SelectionMechanic))
-			{
-				const EComponentSelectionMode SelectionMode = ToolProperties->ComponentSelectionMode;
-
-				const EGeometryElementType ElementType = [SelectionMode]()
-					{
-						if (SelectionMode == EComponentSelectionMode::Vertices)
-						{
-							return EGeometryElementType::Vertex;
-						}
-						else if (SelectionMode == EComponentSelectionMode::Edges)
-						{
-							return EGeometryElementType::Edge;
-						}
-						else if (SelectionMode == EComponentSelectionMode::Faces)
-						{
-							return EGeometryElementType::Face;
-						}
-						else
-						{
-							ensureMsgf(0, TEXT("Unhandled selection mode"));
-							return EGeometryElementType::Vertex;
-						}
-					}();
-
-				UE::Geometry::FGeometrySelection OldSelection;
-
-				OldSelection.TopologyType = EGeometryTopologyType::Triangle;
-				NewSelection.TopologyType = EGeometryTopologyType::Triangle;
-
-				OldSelection.ElementType = ElementType;
-				NewSelection.ElementType = ElementType;
-
-				SelectionMechanic->GetSelection_AsTriangleTopology(OldSelection);
-
-				if (SelectionMode == EComponentSelectionMode::Vertices)
-				{
-					Algo::Transform(OldSelection.Selection, NewSelection.Selection,
-						[&Mesh, Direction](const int32 ElementID)
-						{
-							if (Direction == EWeightToolScopedRetainSelectionDirection::RestoreInFullMesh)
-							{
-								return Mesh.MapVertexToBaseMesh(ElementID);
-							}
-							else
-							{
-								return Mesh.MapVertexToSubmesh(ElementID);
-							}
-						});
-				}
-				else if (SelectionMode == EComponentSelectionMode::Edges)
-				{
-					for (const int32 EdgeID : OldSelection.Selection)
-					{
-						if (Direction == EWeightToolScopedRetainSelectionDirection::RestoreInFullMesh)
-						{
-							UE::Geometry::FMeshTriEdgeID TriEdgeID(EdgeID);							
-							TriEdgeID.TriangleID = Mesh.MapTriangleToBaseMesh(TriEdgeID.TriangleID);
-
-							NewSelection.Selection.Add(TriEdgeID.Encoded());
-						}
-						else
-						{
-							UE::Geometry::FMeshTriEdgeID TriEdgeID(EdgeID);
-							TriEdgeID.TriangleID = Mesh.MapTriangleToSubmesh(TriEdgeID.TriangleID);
-
-							NewSelection.Selection.Add(TriEdgeID.Encoded());
-						}
-					}
-				}
-				else if (SelectionMode == EComponentSelectionMode::Faces)
-				{
-					Algo::Transform(OldSelection.Selection, NewSelection.Selection,
-						[&Mesh, Direction](const int32 Index)
-						{
-							if (Direction == EWeightToolScopedRetainSelectionDirection::RestoreInFullMesh)
-							{
-								return Mesh.MapTriangleToBaseMesh(Index);
-							}
-							else
-							{
-								return Mesh.MapTriangleToSubmesh(Index);
-							}
-						});
-				}
-				else
-				{
-					ensureMsgf(0, TEXT("Unhandled selection mode"));
-				}
-			}
-		}
-
-		~FWeightToolScopedRetainSelection()
-		{
-			if (SelectionMechanic)
-			{					
-				// Restore selection
-				SelectionMechanic->SetSelection_AsTriangleTopology(NewSelection);
-			}
-		}
-
-	private:
-		UE::Geometry::FGeometrySelection NewSelection;
-
-		UPolygonSelectionMechanic* SelectionMechanic = nullptr;
-	};
-}
-
 void UWeightToolSelectionIsolator::RestoreFullMesh()
 {
-	using namespace UE::MeshModelingToolsEditorOnly;
+	using namespace UE::MeshModelingToolsEditorOnly::SelectionUtils;
 
 	if (!IsSelectionIsolated())
 	{
@@ -4822,7 +5031,7 @@
 
 void UWeightToolSelectionIsolator::CreatePartialMesh()
 {
-	using namespace UE::MeshModelingToolsEditorOnly;
+	using namespace UE::MeshModelingToolsEditorOnly::SelectionUtils;
 
 	UPolygonSelectionMechanic* SelectionMechanic = WeightTool->GetMainMeshSelector()->GetSelectionMechanic();
 	if (!ensure(SelectionMechanic))
diff --git a/Engine/Plugins/Runtime/MeshModelingToolset/Source/MeshModelingToolsEditorOnly/Public/SkeletalMesh/SkinWeightsPaintTool.h b/Engine/Plugins/Runtime/MeshModelingToolset/Source/MeshModelingToolsEditorOnly/Public/SkeletalMesh/SkinWeightsPaintTool.h
--- a/Engine/Plugins/Runtime/MeshModelingToolset/Source/MeshModelingToolsEditorOnly/Public/SkeletalMesh/SkinWeightsPaintTool.h
+++ b/Engine/Plugins/Runtime/MeshModelingToolset/Source/MeshModelingToolsEditorOnly/Public/SkeletalMesh/SkinWeightsPaintTool.h
@@ -466,7 +466,17 @@ public:
 	// pointer back to paint tool
 	TObjectPtr<USkinWeightsPaintTool> WeightTool;
 
+	UE_DEPRECATED(5.8, "Instead use SetComponentSelectionMode")
 	UE_API void SetComponentMode(EComponentSelectionMode InComponentMode);
+	
+	/** 
+	 * Sets the component selection mode 
+	 * 
+	 * @param InComponentMode		The component selection mode to set
+	 * @param bInConvertSelection	When set to true, converts the current selection to the new mode
+	 */
+	UE_API void SetComponentSelectionMode(const EComponentSelectionMode InComponentMode, const bool bInConvertSelection);
+	
 	UE_API void SetFalloffMode(EWeightBrushFalloffMode InFalloffMode);
 	UE_API void SetColorMode(EWeightColorMode InColorMode);
 	UE_API void SetBrushMode(EWeightEditOperation InBrushMode);
diff --git a/Engine/Plugins/Runtime/MeshModelingToolset/Source/ModelingComponents/Private/Selection/PolygonSelectionMechanic.cpp b/Engine/Plugins/Runtime/MeshModelingToolset/Source/ModelingComponents/Private/Selection/PolygonSelectionMechanic.cpp
--- a/Engine/Plugins/Runtime/MeshModelingToolset/Source/ModelingComponents/Private/Selection/PolygonSelectionMechanic.cpp
+++ b/Engine/Plugins/Runtime/MeshModelingToolset/Source/ModelingComponents/Private/Selection/PolygonSelectionMechanic.cpp
@@ -187,12 +187,20 @@ void UPolygonSelectionMechanic::GetSelection_AsTriangleTopology(UE::Geometry::FG
 }
 
 
-void UPolygonSelectionMechanic::SetSelection_AsGroupTopology(const UE::Geometry::FGeometrySelection& Selection)
+void UPolygonSelectionMechanic::SetSelection_AsGroupTopology(const UE::Geometry::FGeometrySelection& Selection, const bool bBroadcastSelectionChanged)
 {
 	if (Selection.TopologyType != EGeometryTopologyType::Polygroup)
 	{
 		return;
 	}
+	
+	TOptional<FGroupTopologySelection> PreviousSelection;
+	if (bBroadcastSelectionChanged)
+	{
+		// Only remember the previous selection when broadcasting the change
+		PreviousSelection = PersistentSelection;
+	}
+	
 	PersistentSelection.Clear();
 
 	if (Selection.ElementType == EGeometryElementType::Vertex)
@@ -231,14 +239,30 @@ void UPolygonSelectionMechanic::SetSelection_AsGroupTopology(const UE::Geometry:
 			}
 		}
 	}
+	
+	if (bBroadcastSelectionChanged &&
+		ensure(PreviousSelection.IsSet()) &&
+		PersistentSelection != PreviousSelection.GetValue())
+	{
+		SelectionTimestamp++;
+		OnSelectionChanged.Broadcast();
+	}
 }
 
-void UPolygonSelectionMechanic::SetSelection_AsTriangleTopology(const UE::Geometry::FGeometrySelection& Selection)
+void UPolygonSelectionMechanic::SetSelection_AsTriangleTopology(const UE::Geometry::FGeometrySelection& Selection, const bool bBroadcastSelectionChanged)
 {
 	if (Selection.TopologyType != EGeometryTopologyType::Triangle)
 	{
 		return;
 	}
+
+	TOptional<FGroupTopologySelection> PreviousSelection;
+	if (bBroadcastSelectionChanged)
+	{
+		// Only remember the previous selection when broadcasting the change
+		PreviousSelection = PersistentSelection;
+	}
+
 	PersistentSelection.Clear();
 	
 	if (Selection.ElementType == EGeometryElementType::Vertex)
@@ -277,6 +301,14 @@ void UPolygonSelectionMechanic::SetSelection_AsTriangleTopology(const UE::Geomet
 			}
 		}
 	}
+	
+	if (bBroadcastSelectionChanged &&
+		ensure(PreviousSelection.IsSet()) &&
+		PersistentSelection != PreviousSelection.GetValue())
+	{
+		SelectionTimestamp++;
+		OnSelectionChanged.Broadcast();
+	}
 }
 
 bool UPolygonSelectionMechanic::ExecuteActionThroughGeometrySelection(bool bAsTriangleTopology, const FText& TransactionName, 
diff --git a/Engine/Plugins/Runtime/MeshModelingToolset/Source/ModelingComponents/Public/Selection/PolygonSelectionMechanic.h b/Engine/Plugins/Runtime/MeshModelingToolset/Source/ModelingComponents/Public/Selection/PolygonSelectionMechanic.h
--- a/Engine/Plugins/Runtime/MeshModelingToolset/Source/ModelingComponents/Public/Selection/PolygonSelectionMechanic.h
+++ b/Engine/Plugins/Runtime/MeshModelingToolset/Source/ModelingComponents/Public/Selection/PolygonSelectionMechanic.h
@@ -100,12 +100,12 @@ public:
 	/**
 	 * Initialize the Active Selection based on the provided PolyGroup-topology FGeometrySelection
 	 */
-	UE_API void SetSelection_AsGroupTopology(const UE::Geometry::FGeometrySelection& Selection);
+	UE_API void SetSelection_AsGroupTopology(const UE::Geometry::FGeometrySelection& Selection, const bool bBroadcastSelectionChanged = false);
 
 	/**
 	 * Initialize the Active Selection based on the provided Triangle-topology FGeometrySelection
 	 */
-	UE_API void SetSelection_AsTriangleTopology(const UE::Geometry::FGeometrySelection& Selection);
+	UE_API void SetSelection_AsTriangleTopology(const UE::Geometry::FGeometrySelection& Selection, const bool bBroadcastSelectionChanged = false);
 
 
 private:
