fix crash in hole fill op where it could use reference a triangle with unset UVs

#Synced-CL 51501523

diff --git a/Engine/Plugins/Runtime/MeshModelingToolset/Source/ModelingOperators/Private/CleaningOps/HoleFillOp.cpp b/Engine/Plugins/Runtime/MeshModelingToolset/Source/ModelingOperators/Private/CleaningOps/HoleFillOp.cpp
--- a/Engine/Plugins/Runtime/MeshModelingToolset/Source/ModelingOperators/Private/CleaningOps/HoleFillOp.cpp
+++ b/Engine/Plugins/Runtime/MeshModelingToolset/Source/ModelingOperators/Private/CleaningOps/HoleFillOp.cpp
@@ -264,9 +264,10 @@ bool FHoleFillOp::FillSingleTriangleHole(const FEdgeLoop& Loop, int32& NewGroupI
 		}
 		else
 		{
+			// Since we didn't have a matching element ID above, at least one of ExistingTriangles[0] or [1] must have set UVs --
 			// Pick an arbitrary set UV edge and add a new third UV element
 			int32 NeighborTidToUse = UVOverlay->IsSetTriangle(ExistingTriangles[0]) ? ExistingTriangles[0] : ExistingTriangles[1];
-			SetUVTriangleFromAdjacentTriangle(*UVOverlay, ExistingTriangles[0], NewTriangleID, NewTriangleElements);
+			SetUVTriangleFromAdjacentTriangle(*UVOverlay, NeighborTidToUse, NewTriangleID, NewTriangleElements);
 		}
 
 		EMeshResult Result = EMeshResult::Ok;
