Fixed Modeling - Mesh Component Selection test

#rnx
#Synced-CL 51261145

diff --git a/Engine/Plugins/Runtime/MeshModelingToolset/Source/ModelingComponents/Public/Selection/MeshTopologySelectionMechanic.h b/Engine/Plugins/Runtime/MeshModelingToolset/Source/ModelingComponents/Public/Selection/MeshTopologySelectionMechanic.h
--- a/Engine/Plugins/Runtime/MeshModelingToolset/Source/ModelingComponents/Public/Selection/MeshTopologySelectionMechanic.h
+++ b/Engine/Plugins/Runtime/MeshModelingToolset/Source/ModelingComponents/Public/Selection/MeshTopologySelectionMechanic.h
@@ -69,7 +69,7 @@ public:
 	UPROPERTY(EditAnywhere, Category = AdditionalSelectionOptions, AdvancedDisplay)
 	EMeshTopologySelectionMechanicDragTool SelectionDragTool = EMeshTopologySelectionMechanicDragTool::Marquee;
 
-	UPROPERTY(EditAnywhere, Category = AdditionalSelectionOptions, meta = (EditCondition = "SelectionDragTool != UMeshTopologySelectionMechanicDragTool::None", EditConditionHides))
+	UPROPERTY(EditAnywhere, Category = AdditionalSelectionOptions, meta = (EditCondition = "SelectionDragTool != EMeshTopologySelectionMechanicDragTool::None", EditConditionHides))
 	bool bIgnoreOcclusionForDragTools = true;
 
 #if WITH_EDITORONLY_DATA
