[ModelingTools] ability to generate, copy and paint per triangle label attributes

#Synced-CL 51662791

diff --git a/Engine/Plugins/Runtime/GeometryProcessing/Source/DynamicMesh/Private/FaceGroupUtil.cpp b/Engine/Plugins/Runtime/GeometryProcessing/Source/DynamicMesh/Private/FaceGroupUtil.cpp
--- a/Engine/Plugins/Runtime/GeometryProcessing/Source/DynamicMesh/Private/FaceGroupUtil.cpp
+++ b/Engine/Plugins/Runtime/GeometryProcessing/Source/DynamicMesh/Private/FaceGroupUtil.cpp
@@ -3,6 +3,7 @@
 #include "FaceGroupUtil.h"
 #include "Selections/MeshConnectedComponents.h"
 #include "Async/ParallelFor.h"
+#include "Polygroups/GroupSetAdapter.h"
 
 using namespace UE::Geometry;
 
@@ -231,17 +232,25 @@ void FGroupVisualizationCache::UpdateGroupInfo_ConnectedComponents(
 	const FPolygroupSet& GroupSet,
 	bool bParallel)
 {
-	
+	const FPolygroupLayerAdapter Adapter(GroupSet.Mesh, GroupSet.PolygroupAttrib);
+	UpdateGroupInfo_ConnectedComponents(SourceMesh, Adapter, bParallel);
+}
+
+void FGroupVisualizationCache::UpdateGroupInfo_ConnectedComponents(
+	const FDynamicMesh3& SourceMesh,
+	const ITriangleGroupSetAdapter& Adapter,
+	bool bParallel)
+{
 	// find connected group components
 	FMeshConnectedComponents Components(&SourceMesh);
-	Components.FindConnectedTriangles([&](int32 t1, int32 t2) { return GroupSet.GetGroup(t1) == GroupSet.GetGroup(t2); });
+	Components.FindConnectedTriangles([&](int32 t1, int32 t2) { return Adapter.GetGroup(t1) == Adapter.GetGroup(t2); });
 
 	GroupInfo.SetNum(Components.Num());
 
 	EParallelForFlags UseParallelFlags = (bParallel) ? EParallelForFlags::None : EParallelForFlags::ForceSingleThread;
 	ParallelFor(Components.Num(), [&](int32 ci) {
 		const FMeshConnectedComponents::FComponent& Component = Components[ci];
-		GroupInfo[ci].GroupID = GroupSet.GetGroup(Component.Indices[0]);
+		GroupInfo[ci].GroupID = Adapter.GetGroup(Component.Indices[0]);
 
 		// compute bounds
 		GroupInfo[ci].Bounds = FAxisAlignedBox3d::Empty();
diff --git a/Engine/Plugins/Runtime/GeometryProcessing/Source/DynamicMesh/Private/Polygroups/GroupSetAdapter.cpp b/Engine/Plugins/Runtime/GeometryProcessing/Source/DynamicMesh/Private/Polygroups/GroupSetAdapter.cpp
new file mode 100644
--- /dev/null
+++ b/Engine/Plugins/Runtime/GeometryProcessing/Source/DynamicMesh/Private/Polygroups/GroupSetAdapter.cpp
@@ -0,0 +1,110 @@
+// Copyright Epic Games, Inc. All Rights Reserved.
+
+#include "Polygroups/GroupSetAdapter.h"
+
+using namespace UE::Geometry;
+
+ITriangleGroupSetAdapter::~ITriangleGroupSetAdapter() = default;
+
+ITriangleGroupSetAdapter::ITriangleGroupSetAdapter(const FDynamicMesh3* InMesh)
+	: Mesh(InMesh)
+{
+	check(InMesh);
+}
+
+void ITriangleGroupSetAdapter::SetGroup(const int32 TriangleID, const int32 GroupID, FDynamicMesh3& WritableMesh)
+{}
+
+int32 ITriangleGroupSetAdapter::GetMaxGroupID() const
+{
+	return MaxGroupID;
+}
+
+const FDynamicMesh3* ITriangleGroupSetAdapter::GetMesh() const
+{
+	return Mesh;	
+}
+
+FPolygroupLayerAdapter::FPolygroupLayerAdapter(const FDynamicMesh3* InMesh, const FDynamicMeshPolygroupAttribute* Attribute)
+	: ITriangleGroupSetAdapter(InMesh)
+	, PolygroupSet(InMesh, Attribute)
+{}
+
+int32 FPolygroupLayerAdapter::GetGroup(const int32 TriangleID) const
+{
+	return PolygroupSet.GetGroup(TriangleID);
+}
+
+void FPolygroupLayerAdapter::SetGroup(const int32 TriangleID, const int32 NewGroupID, FDynamicMesh3& WritableMesh)
+{
+	return PolygroupSet.SetGroup(TriangleID, NewGroupID, WritableMesh);
+}
+
+int32 FPolygroupLayerAdapter::GetLayerIndex() const
+{
+	return PolygroupSet.GroupLayerIndex;
+}
+
+FTriangleLabelAdapter::FTriangleLabelAdapter(const FDynamicMesh3* InMesh, TriangleAttributeType* InAttribute, const FName InBaseLabel)
+	: FTriangleAttributeAdapter(InMesh, InAttribute)
+	, BaseLabel(InBaseLabel)
+{
+	bool AttributeFound = false;
+
+	if (const FDynamicMeshAttributeSet* Attributes = InMesh && InAttribute ? InMesh->Attributes() : nullptr)
+	{
+		for (const auto&[Name, Attribute]: Attributes->GetTriangleLabelAttributes())
+		{
+			if (Attribute.Get() == InAttribute)
+			{
+				LayerName = Name;
+				AttributeFound = true;
+				break;
+			}
+		}
+	}
+	
+	check(AttributeFound);
+}
+
+FTriangleLabelAdapter::FTriangleLabelAdapter(const FTriangleLabelAdapter* InCopy)
+	: FTriangleAttributeAdapter(InCopy)
+	, LayerName(InCopy->LayerName)
+{}
+
+FName FTriangleLabelAdapter::GetLayerName() const
+{
+	 return LayerName;
+}
+
+FName FTriangleLabelAdapter::GetUniqueValue() const
+{
+	int32 NumberCounter = MaxGroupID;
+	FName NextValue(BaseLabel);
+	NextValue.SetNumber(NumberCounter++);
+	while (ValueToID.Contains(NextValue))
+	{
+		NextValue.SetNumber(NumberCounter++);
+	}
+	return NextValue;
+}
+
+FTriangleReadOnlyLabelAdapter::FTriangleReadOnlyLabelAdapter(const FDynamicMesh3* InMesh, const TriangleAttributeType* InAttribute)
+	: FTriangleReadOnlyAttributeAdapter(InMesh, InAttribute)
+{
+	auto AttributeFound = [InMesh, InAttribute]()
+	{
+		if(const FDynamicMeshAttributeSet* Attributes = InMesh && InAttribute ? InMesh->Attributes() : nullptr)
+		{
+			for (const auto&[Name, Attribute]: Attributes->GetTriangleLabelAttributes())
+			{
+				if (Attribute.Get() == InAttribute)
+				{
+					return true;
+				}
+			}
+		}
+		return false;
+	};
+	check(AttributeFound());
+}
\ No newline at end of file
diff --git a/Engine/Plugins/Runtime/GeometryProcessing/Source/DynamicMesh/Private/Polygroups/PolygroupUtil.cpp b/Engine/Plugins/Runtime/GeometryProcessing/Source/DynamicMesh/Private/Polygroups/PolygroupUtil.cpp
--- a/Engine/Plugins/Runtime/GeometryProcessing/Source/DynamicMesh/Private/Polygroups/PolygroupUtil.cpp
+++ b/Engine/Plugins/Runtime/GeometryProcessing/Source/DynamicMesh/Private/Polygroups/PolygroupUtil.cpp
@@ -34,6 +34,18 @@ const FDynamicMeshPolygroupAttribute* UE::Geometry::FindPolygroupLayerByName(con
 	return nullptr;
 }
 
+FDynamicMeshTriangleLabelAttribute* UE::Geometry::FindTriangleLabelLayerByName(FDynamicMesh3& Mesh, const FName Name)
+{
+	FDynamicMeshAttributeSet* AttributeSet = Mesh.Attributes();
+	return AttributeSet && Name != NAME_None ? AttributeSet->FindTriangleLabelAttribute(Name) : nullptr;
+}
+
+const FDynamicMeshTriangleLabelAttribute* UE::Geometry::FindTriangleLabelLayerByName(const FDynamicMesh3& Mesh, FName Name)
+{
+	const FDynamicMeshAttributeSet* AttributeSet = Mesh.Attributes();
+	return AttributeSet && Name != NAME_None ? AttributeSet->FindTriangleLabelAttribute(Name) : nullptr;
+}
+
 int32 UE::Geometry::FindPolygroupLayerIndexByName(const FDynamicMesh3& Mesh, FName Name)
 {
 	const FDynamicMeshAttributeSet* AttributeSet = Mesh.Attributes();
diff --git a/Engine/Plugins/Runtime/GeometryProcessing/Source/DynamicMesh/Public/FaceGroupUtil.h b/Engine/Plugins/Runtime/GeometryProcessing/Source/DynamicMesh/Public/FaceGroupUtil.h
--- a/Engine/Plugins/Runtime/GeometryProcessing/Source/DynamicMesh/Public/FaceGroupUtil.h
+++ b/Engine/Plugins/Runtime/GeometryProcessing/Source/DynamicMesh/Public/FaceGroupUtil.h
@@ -8,6 +8,7 @@
 #include "DynamicMesh/DynamicMesh3.h"
 #include "DynamicMeshEditor.h"
 #include "Polygroups/PolygroupSet.h"
+PREDECLARE_USE_GEOMETRY_CLASS(ITriangleGroupSetAdapter);
 
 #define UE_API DYNAMICMESH_API
 
@@ -131,6 +132,18 @@ public:
 		const FDynamicMesh3& SourceMesh,
 		const FPolygroupSet& GroupSet,
 		bool bParallel = true);
+	
+	/**
+	 * Update the GroupInfo array by finding group-connected components.
+	 * Note that this may result in multiple GroupInfo items with the same GroupID, if they are not connected-components.
+	 * Computing connected components is somewhat expensive. In addition, for each component, if it is more
+	 * than 2 triangles it will be "eroded" down a single triangle to determine the group "center", also somewhat expensive.
+	 * This step can be computed in parallel if bParallel is true.
+	 */
+	UE_API void UpdateGroupInfo_ConnectedComponents(
+		const FDynamicMesh3& SourceMesh,
+		const ITriangleGroupSetAdapter& Adapter,
+		bool bParallel = true);
 
 public:
 	/**
diff --git a/Engine/Plugins/Runtime/GeometryProcessing/Source/DynamicMesh/Public/Polygroups/GroupSetAdapter.h b/Engine/Plugins/Runtime/GeometryProcessing/Source/DynamicMesh/Public/Polygroups/GroupSetAdapter.h
new file mode 100644
--- /dev/null
+++ b/Engine/Plugins/Runtime/GeometryProcessing/Source/DynamicMesh/Public/Polygroups/GroupSetAdapter.h
@@ -0,0 +1,313 @@
+// Copyright Epic Games, Inc. All Rights Reserved.
+
+#pragma once
+
+#include "PolygroupSet.h"
+#include "DynamicMesh/DynamicMesh3.h"
+#include "DynamicMesh/DynamicMeshAttributeSet.h"
+
+#define UE_API DYNAMICMESH_API
+
+
+namespace UE
+{
+	
+namespace Geometry
+{
+	
+/**
+ * ITriangleGroupSetAdapter allows tools that operate on triangle groups to also operate on arbitrary per-triangle values
+ * by mapping unique values to integers.
+ */
+	
+class UE_INTERNAL ITriangleGroupSetAdapter
+{
+public:
+	/** Default Destructor */
+	UE_API virtual ~ITriangleGroupSetAdapter() = 0;
+	/** @return Mesh this GroupSetAdapter references. */
+	const FDynamicMesh3* GetMesh() const;
+	/** @return Group ID for a TriangleID */
+	virtual int32 GetGroup(const int32 TriangleID) const = 0;
+	/** Set the Group ID for a TriangleID */
+	virtual void SetGroup(const int32 TriangleID, const int32 GroupID, FDynamicMesh3& WritableMesh) = 0;
+	/**  */
+	UE_API int32 GetMaxGroupID() const;
+	
+protected:
+	/** Initialize a GroupSetAdapter for the given Mesh. */
+	ITriangleGroupSetAdapter(const FDynamicMesh3* InMesh);
+	/** The mesh this GroupSetAdapter references. */
+	const FDynamicMesh3* Mesh = nullptr;
+	/** Max Group ID. */
+	int32 MaxGroupID = 0;
+};
+
+/**
+ * FPolygroupLayerAdapter provides an interface to Get/Set group IDs per triangle using a FDynamicMeshPolygroupAttribute as the actual data storage.  
+ */
+class UE_INTERNAL FPolygroupLayerAdapter: public ITriangleGroupSetAdapter
+{
+public:
+	/** Initialize a FPolygroupLayerAdapter for given Mesh and specific Polygroup attribute layer. */
+	UE_API FPolygroupLayerAdapter(const FDynamicMesh3* InMesh, const FDynamicMeshPolygroupAttribute* Attribute);
+	/** @return Group ID for a TriangleID. */
+	virtual int32 GetGroup(const int32 TriangleID) const override;
+	/** Set the Group ID for a TriangleID. */
+	virtual void SetGroup(const int32 TriangleID, const int32 NewGroupID, FDynamicMesh3& WritableMesh) override;
+
+protected:
+	/** the PolygroupAttribute this GroupSetAdapter references. */
+	FPolygroupSet PolygroupSet;
+	/** @return index of current PolygroupAttribute into Mesh AttributeSet, or -1 if this information does not exist. */
+	int32 GetLayerIndex() const;
+	/** Default layer index. */
+	int32 GroupLayerIndex = -1;
+};
+
+/**
+ * FTriangleReadOnlyAttributeAdapter provides a read-only typed interface to Get group IDs per triangle (as typed values instead on integers)
+ * using a TDynamicMeshSingleTriangleAttribute as the actual data storage and ensuring a group ID - AttribValueType value mapping.
+ */
+template<typename AttribValueType>
+class UE_INTERNAL FTriangleReadOnlyAttributeAdapter: public ITriangleGroupSetAdapter
+{
+public:
+	using TriangleAttributeType = TDynamicMeshSingleTriangleAttribute<AttribValueType>;
+	using ValueType = AttribValueType;
+	
+	/** Initialize a FTriangleReadOnlyAttributeAdapter for given Mesh and specific triangle attribute. */
+	FTriangleReadOnlyAttributeAdapter(const FDynamicMesh3* InMesh, const TriangleAttributeType* InAttribute)
+		: ITriangleGroupSetAdapter(InMesh)
+		, TriangleAttribute(InAttribute)
+	{
+		check(CheckValidity());
+		
+		for (int32 TriangleID : Mesh->TriangleIndicesItr())
+		{
+			const ValueType Value = TriangleAttribute->GetValue(TriangleID);
+			if (!ValueToID.Contains(Value))
+			{
+				ValueToID.Emplace(Value, MaxGroupID);
+				IDToValue.Emplace(MaxGroupID, Value);
+				MaxGroupID++;
+			}
+		}
+	}
+	
+	/** @return Group ID for a TriangleID or INDEX_NONE if there's no value mapped to that Group ID. */
+	virtual int32 GetGroup(const int32 TriangleID) const override 
+	{
+		const ValueType Value = TriangleAttribute->GetValue(TriangleID);
+		const int32* GroupID = ValueToID.Find(Value);
+		return GroupID ? *GroupID : INDEX_NONE;
+	}
+	
+	/** Do not update any group id. */
+	virtual void SetGroup(const int32 TriangleID, const int32 GroupID, FDynamicMesh3& WritableMesh) override {};
+	
+protected:
+
+	/** @return true if this adapter is valid for use. */
+	bool CheckValidity() const
+	{
+		return Mesh && TriangleAttribute &&
+			TriangleAttribute->GetParentMesh() == Mesh &&
+			TriangleAttribute->CheckValidity(true, EValidityCheckFailMode::ReturnOnly);
+	}
+	
+	/** the per-triangle attribute this adapter references. */
+	const TriangleAttributeType* TriangleAttribute = nullptr;
+	/** Value-Group ID mapping. */
+	TMap<ValueType, int32> ValueToID;
+	/** Group ID-Value mapping. */
+	TMap<int32, ValueType> IDToValue;
+};
+	
+/**
+ * FTriangleAttributeAdapter provides a typed interface to Get/Set group IDs per triangle (as typed values instead on integers)
+ * using a TDynamicMeshSingleTriangleAttribute as the actual data storage and ensuring a group ID - AttribValueType value mapping.
+ */
+template<typename AttribValueType>
+class UE_INTERNAL FTriangleAttributeAdapter: public ITriangleGroupSetAdapter
+{
+public:
+	using TriangleAttributeType = TDynamicMeshSingleTriangleAttribute<AttribValueType>;
+	using ValueType = AttribValueType;
+	
+	/** Initialize a FTriangleAttributeAdapter for given Mesh and specific triangle attribute. */
+	FTriangleAttributeAdapter(const FDynamicMesh3* InMesh, TriangleAttributeType* InAttribute)
+		: ITriangleGroupSetAdapter(InMesh)
+		, TriangleAttribute(InAttribute)
+	{
+		check(CheckValidity());
+	
+		AttribValueType DefaultValue = InAttribute->GetDefaultAttributeValue();
+		ValueToID.Emplace(DefaultValue, INDEX_NONE);
+		IDToValue.Emplace(INDEX_NONE, DefaultValue);
+		
+		for (int32 TriangleID : Mesh->TriangleIndicesItr())
+		{
+			const ValueType Value = GetValue(TriangleID);
+			if (!ValueToID.Contains(Value))
+			{
+				ValueToID.Emplace(Value, MaxGroupID);
+				IDToValue.Emplace(MaxGroupID, Value);
+				MaxGroupID++;
+			}
+		}
+	}
+	
+	FTriangleAttributeAdapter(const FTriangleAttributeAdapter* InCopy)
+		: ITriangleGroupSetAdapter(InCopy->Mesh)
+		, TriangleAttribute(InCopy->TriangleAttribute)
+		, ValueToID(InCopy->ValueToID)
+		, IDToValue(InCopy->IDToValue)
+	{
+		MaxGroupID = InCopy->MaxGroupID;
+	}
+
+	/** Set the Group ID for a TriangleID and updates the value-id mappings. */
+	void SetGroup(const int32 TriangleID, const int32 GroupID)
+	{
+		const ValueType* Value = IDToValue.Find(GroupID);
+		if (Value)
+		{
+			SetValue(TriangleID, *Value);
+		}
+		else
+		{
+			ValueType NewValue = GetUniqueValue();
+			ensure(!ValueToID.Contains(NewValue));
+			ValueToID.Emplace(NewValue, GroupID);
+			IDToValue.Emplace(GroupID, NewValue);
+			SetValue(TriangleID, NewValue);
+			MaxGroupID = FMath::Max(MaxGroupID+1, GroupID+1);
+		}
+	}
+	
+	/** Set the Group ID for a TriangleID and updates the value-id mappings. InMesh can be used to perform additional tests. */
+	virtual void SetGroup(const int32 TriangleID, const int32 GroupID, FDynamicMesh3& InMesh) override
+	{
+		SetGroup(TriangleID, GroupID);
+	}
+	
+	/** @return Group ID for a TriangleID or INDEX_NONE if there's no value mapped to that Group ID. */
+	virtual int32 GetGroup(const int32 TriangleID) const override 
+	{
+		const int32* GroupID = ValueToID.Find(GetValue(TriangleID));
+		if (ensure(GroupID))
+		{
+			return *GroupID;
+		}
+		return INDEX_NONE;
+	}
+	
+	/** @return Group ID from a Value. If the value is not yet registered, a new group ID is added and the mapping created. */
+	int32 GetGroupFromValue(const ValueType Value)
+	{
+		if (const int32* GroupID = ValueToID.Find(Value))
+		{
+			return *GroupID;
+		}
+		
+		const int32 NewGroupID = MaxGroupID;
+		ValueToID.Emplace(Value, MaxGroupID);
+		IDToValue.Emplace(MaxGroupID, Value);
+		MaxGroupID++;
+		
+		return NewGroupID;
+	}
+	
+	/** @return Value from a Group ID. If the group ID is not yet registered, a new value is added and the mapping is created. */
+	ValueType GetValueFromGroup(const int32 GroupID)
+	{
+		if (const ValueType* Value = IDToValue.Find(GroupID))
+		{
+			return *Value;
+		}
+		
+		ValueType NewValue = GetUniqueValue();
+		ensure(!ValueToID.Contains(NewValue));
+		ValueToID.Emplace(NewValue, GroupID);
+		IDToValue.Emplace(GroupID, NewValue);
+		MaxGroupID = FMath::Max(MaxGroupID+1, GroupID+1);
+		return NewValue;
+	}
+	
+	/** @return Value for a TriangleID. */
+	ValueType GetValue(const int32 TriangleID) const
+	{
+		return TriangleAttribute->GetValue(TriangleID);
+	}
+	
+	/** Set the Value for a TriangleID and register mapping if needed. */
+	void SetValue(const int32 TriangleID, ValueType NewValue)
+	{
+		if (!ValueToID.Contains(NewValue))
+		{
+			ensure(!IDToValue.Contains(MaxGroupID));
+			ValueToID.Emplace(NewValue, MaxGroupID);
+			IDToValue.Emplace(MaxGroupID, NewValue);
+			MaxGroupID++;
+		}
+		return TriangleAttribute->SetValue(TriangleID, NewValue);
+	}
+	
+protected:
+	
+	/** @return New unique value. */
+	virtual ValueType GetUniqueValue() const
+	{
+		ensure(false);
+		return TriangleAttribute->GetDefaultAttributeValue();
+	}
+
+	/** @return true if this adapter is valid for use. */
+	bool CheckValidity() const
+	{
+		return Mesh && TriangleAttribute &&
+			TriangleAttribute->GetParentMesh() == Mesh &&
+			TriangleAttribute->CheckValidity(true, EValidityCheckFailMode::ReturnOnly);
+	}
+	
+	/** the per-triangle attribute this adapter references. */
+	TriangleAttributeType* TriangleAttribute = nullptr;
+	/** Value-Group ID mapping. */
+	TMap<ValueType, int32> ValueToID;
+	/** Group ID-Value mapping. */
+	TMap<int32, ValueType> IDToValue;
+};
+	
+/**
+ * FTriangleLabelAdapter provides a FTriangleAttributeAdapter interface specialized for FName. The actual attribute that this adapter references
+ * must be contained in FDynamicMeshAttributeSet::TriangleLabelAttributes. (this is actually checked in the constructor).
+ */
+class UE_INTERNAL FTriangleLabelAdapter: public FTriangleAttributeAdapter<FName>
+{
+public:
+	UE_API FTriangleLabelAdapter(const FDynamicMesh3* InMesh, TriangleAttributeType* InAttribute, const FName InBaseLabel = "label");
+	UE_API FTriangleLabelAdapter(const FTriangleLabelAdapter* InCopy);
+	UE_API FName GetLayerName() const;
+protected:
+	/**  */
+	FName LayerName = NAME_None;
+	/** @return New unique FName. */
+	UE_API virtual FName GetUniqueValue() const override;
+	/** Base name used for unique FName generation. */
+	FName BaseLabel = "label";
+};
+
+/**
+ * FTriangleReadOnlyLabelAdapter provides a read-only FTriangleAttributeAdapter interface specialized for FName.
+ */
+class UE_INTERNAL FTriangleReadOnlyLabelAdapter: public FTriangleReadOnlyAttributeAdapter<FName>
+{
+public:
+	UE_API FTriangleReadOnlyLabelAdapter(const FDynamicMesh3* InMesh, const TriangleAttributeType* InAttribute);
+};
+	
+}	// end namespace Geometry
+}	// end namespace UE
+
+#undef UE_API
diff --git a/Engine/Plugins/Runtime/GeometryProcessing/Source/DynamicMesh/Public/Polygroups/PolygroupUtil.h b/Engine/Plugins/Runtime/GeometryProcessing/Source/DynamicMesh/Public/Polygroups/PolygroupUtil.h
--- a/Engine/Plugins/Runtime/GeometryProcessing/Source/DynamicMesh/Public/Polygroups/PolygroupUtil.h
+++ b/Engine/Plugins/Runtime/GeometryProcessing/Source/DynamicMesh/Public/Polygroups/PolygroupUtil.h
@@ -12,17 +12,29 @@ namespace Geometry
 {
 
 /**
- * Find the first FDynamicMeshPolygroupAttribute with the given FName in a the AttributeSet of a Mesh.
+ * Find the first FDynamicMeshPolygroupAttribute with the given FName in the AttributeSet of a Mesh.
  * @return nullptr if no Polygroup layer is found
  */
 DYNAMICMESH_API FDynamicMeshPolygroupAttribute* FindPolygroupLayerByName(FDynamicMesh3& Mesh, FName Name);
 
 /**
- * Find the first FDynamicMeshPolygroupAttribute with the given FName in a the AttributeSet of a Mesh.
+ * Find the first FDynamicMeshPolygroupAttribute with the given FName in the AttributeSet of a Mesh.
  * @return nullptr if no Polygroup layer is found
  */
 DYNAMICMESH_API const FDynamicMeshPolygroupAttribute* FindPolygroupLayerByName(const FDynamicMesh3& Mesh, FName Name);
 
+/**
+ * Find the FDynamicMeshTriangleLabelAttribute with the given FName in the AttributeSet of a Mesh.
+ * @return nullptr if no triangle label layer is found.
+ */
+DYNAMICMESH_API FDynamicMeshTriangleLabelAttribute* FindTriangleLabelLayerByName(FDynamicMesh3& Mesh, const FName Name);
+
+/**
+ * Find the FDynamicMeshTriangleLabelAttribute with the given FName in the AttributeSet of a Mesh.
+ * @return nullptr if no triangle label layer is found.
+ */
+DYNAMICMESH_API const FDynamicMeshTriangleLabelAttribute* FindTriangleLabelLayerByName(const FDynamicMesh3& Mesh, FName Name);
+	
 /**
  * @return index of the first Layer with the given FName in Mesh AttributeSet, or -1 if not found
  */
diff --git a/Engine/Plugins/Runtime/MeshModelingToolset/Source/MeshModelingTools/Private/ConvertToPolygonsTool.cpp b/Engine/Plugins/Runtime/MeshModelingToolset/Source/MeshModelingTools/Private/ConvertToPolygonsTool.cpp
--- a/Engine/Plugins/Runtime/MeshModelingToolset/Source/MeshModelingTools/Private/ConvertToPolygonsTool.cpp
+++ b/Engine/Plugins/Runtime/MeshModelingToolset/Source/MeshModelingTools/Private/ConvertToPolygonsTool.cpp
@@ -25,6 +25,7 @@
 #include "Selection/GeometrySelectionVisualization.h"
 #include "PropertySets/GeometrySelectionVisualizationProperties.h"
 #include "ToolTargetManager.h"
+#include "Polygroups/GroupSetAdapter.h"
 
 #include "Polygroups/PolygroupsGenerator.h"
 
@@ -34,6 +35,12 @@ using namespace UE::Geometry;
 
 #define LOCTEXT_NAMESPACE "UConvertToPolygonsTool"
 
+namespace UE::PolygroupUtilLocals
+{
+	const TCHAR* NewLayer = TEXT("Create New...");
+	const TCHAR* DefaultLayer = TEXT("Default");
+}
+
 /*
  * ToolBuilder
  */
@@ -77,7 +84,7 @@ public:
 	// input mesh
 	TSharedPtr<FSharedConstDynamicMesh3, ESPMode::ThreadSafe> OriginalMesh;
 	// input polygroups, if available
-	TSharedPtr<FPolygroupSet, ESPMode::ThreadSafe> SourcePolyGroups;
+	TSharedPtr<ITriangleGroupSetAdapter, ESPMode::ThreadSafe> SourcePolyGroups;
 
 	// result
 	FPolygroupsGenerator Generator;
@@ -104,7 +111,7 @@ public:
 			if (!ensure(SourcePolyGroups.IsValid())) return;
 			for (int32 tid : ResultMesh->TriangleIndicesItr())
 			{
-				ResultMesh->SetTriangleGroup(tid, SourcePolyGroups->GetTriangleGroup(tid));
+				ResultMesh->SetTriangleGroup(tid, SourcePolyGroups->GetGroup(tid));
 			}
 			return;
 		}
@@ -345,6 +352,7 @@ void UConvertToPolygonsTool::Setup()
 
 	// add group layer picking property set for source groups
 	CopyFromLayerProperties = NewObject<UPolygroupLayersProperties>(this);
+	CopyFromLayerProperties->bSupportLabels = true;
 	CopyFromLayerProperties->InitializeGroupLayers(OriginalDynamicMesh.Get());
 	CopyFromLayerProperties->WatchProperty(CopyFromLayerProperties->ActiveGroupLayer, [&](FName) { OnSelectedFromGroupLayerChanged(); });
 	AddToolPropertySource(CopyFromLayerProperties);
@@ -353,17 +361,55 @@ void UConvertToPolygonsTool::Setup()
 
 	// add picker for output group layer
 	Settings->OptionsList.Reset();
-	Settings->OptionsList.Add(TEXT("Default"));		// always have standard group
-	if (OriginalDynamicMesh->Attributes())
+	Settings->OptionsList.Add(UE::PolygroupUtilLocals::DefaultLayer);		// always have standard group
+	if (FDynamicMeshAttributeSet* Attributes = OriginalDynamicMesh->Attributes())
 	{
-		for (int32 k = 0; k < OriginalDynamicMesh->Attributes()->NumPolygroupLayers(); k++)
+		for (int32 k = 0; k < Attributes->NumPolygroupLayers(); k++)
+		{
+			FName Name = Attributes->GetPolygroupLayer(k)->GetName();
+			Settings->OptionsList.Add(Name.ToString());
+		}
+		for (const auto&[Name, LayerAttr]: Attributes->GetTriangleLabelAttributes())
 		{
-			FName Name = OriginalDynamicMesh->Attributes()->GetPolygroupLayer(k)->GetName();
 			Settings->OptionsList.Add(Name.ToString());
 		}
 	}
-	Settings->OptionsList.Add(TEXT("Create New..."));
-	Settings->WatchProperty(Settings->GroupLayer, [&](FName NewName) { Settings->bShowNewLayerName = (NewName == TEXT("Create New...")); });
+	
+	Settings->OptionsList.Add(UE::PolygroupUtilLocals::NewLayer);
+	
+	Settings->WatchProperty(Settings->GroupLayer, [&](FName NewName)
+	{
+		Settings->bShowNewLayerName = (NewName == UE::PolygroupUtilLocals::NewLayer);
+		
+		if (NewName == UE::PolygroupUtilLocals::DefaultLayer)
+		{
+			Settings->LayerType = EPerTriangleAttributeType::Integer;
+		}
+		else if (NewName != UE::PolygroupUtilLocals::NewLayer)
+		{
+			if (FDynamicMeshAttributeSet* Attributes = OriginalDynamicMesh->Attributes())
+			{
+				bool bPolygroupFound = false;
+				for (int32 k = 0; k < Attributes->NumPolygroupLayers(); k++)
+				{
+					if (Attributes->GetPolygroupLayer(k)->GetName() == NewName)
+					{
+						Settings->LayerType = EPerTriangleAttributeType::Integer;
+						bPolygroupFound = true;
+						break;
+					}
+				}
+		
+				if (!bPolygroupFound)
+				{
+					if (Attributes->FindTriangleLabelAttribute(NewName))
+					{
+						Settings->LayerType = EPerTriangleAttributeType::Label;
+					}
+				}
+			}
+		}
+	});
 
 	if (bUsingSelection)
 	{
@@ -422,12 +468,12 @@ void UConvertToPolygonsTool::UpdateOpParameters(FConvertToPolygonsOp& ConvertToP
 		ConvertToPolygonsOp.InitialGroupID = OriginalDynamicMesh->MaxGroupID();
 	}
 
-	if (ActiveFromGroupSet.IsValid())
+	if (ActiveFromGroupSetAdapter.IsValid())
 	{
-		ConvertToPolygonsOp.SourcePolyGroups = ActiveFromGroupSet;
+		ConvertToPolygonsOp.SourcePolyGroups = ActiveFromGroupSetAdapter;
 		if (bUsingSelection)
 		{
-			ConvertToPolygonsOp.InitialGroupID = ActiveFromGroupSet->MaxGroupID;
+			ConvertToPolygonsOp.InitialGroupID = ActiveFromGroupSetAdapter->GetMaxGroupID();
 		}
 	}
 
@@ -463,56 +509,9 @@ void UConvertToPolygonsTool::OnShutdown(EToolShutdownType ShutdownType)
 			FDynamicMesh3* DynamicMeshResult = Result.Mesh.Get();
 			if (ensure(DynamicMeshResult != nullptr))
 			{
-				if (Settings->GroupLayer != TEXT("Default"))
+				if (Settings->GroupLayer != UE::PolygroupUtilLocals::DefaultLayer)
 				{
-					FDynamicMesh3 UseResultMesh = *OriginalDynamicMesh;
-					if (UseResultMesh.HasAttributes() == false)
-					{
-						UseResultMesh.EnableAttributes();
-					}
-
-					// if we want to write to any layer other than default, we have to find or create it
-					FDynamicMeshPolygroupAttribute* UseAttribLayer = nullptr;
-					if (Settings->GroupLayer == TEXT("Create New..."))
-					{
-						// append new group layer and set it's name
-						int32 TargetLayerIdx = UseResultMesh.Attributes()->NumPolygroupLayers();
-						UseResultMesh.Attributes()->SetNumPolygroupLayers(TargetLayerIdx + 1);
-						UseAttribLayer = UseResultMesh.Attributes()->GetPolygroupLayer(TargetLayerIdx);
-						FString UseUniqueName = UE::Geometry::MakeUniqueGroupLayerName(UseResultMesh, Settings->NewLayerName);
-						UseAttribLayer->SetName( FName(UseUniqueName) );
-					}
-					else
-					{
-						UseAttribLayer = UE::Geometry::FindPolygroupLayerByName(UseResultMesh, Settings->GroupLayer);
-					}
-
-					if (UseAttribLayer)
-					{
-						// copy the generated groups from the Op output mesh (stored in primary groups) to the target layer
-						if (bUsingSelection)
-						{
-							for (int32 tid : *SelectionTriangleROI)
-							{
-								int32 GroupID = DynamicMeshResult->GetTriangleGroup(OriginalSubmesh->MapTriangleToSubmesh(tid));
-								UseAttribLayer->SetValue(tid, GroupID);
-							}
-						}
-						else
-						{
-							for (int32 tid : UseResultMesh.TriangleIndicesItr())
-							{
-								UseAttribLayer->SetValue(tid, DynamicMeshResult->GetTriangleGroup(tid));
-							}
-						}
-						UE::ToolTarget::CommitDynamicMeshUpdate(Target, UseResultMesh, true);
-					}
-					else
-					{
-						// If we can't find or create the layer (which should not be possible) the tool is going to do nothing,
-						// this is the safest option at this point
-						UE_LOG(LogGeometry, Warning, TEXT("Output group layer missing?"));
-					}
+					DoUpdateMeshAttribute(*DynamicMeshResult);
 				}
 				else
 				{
@@ -586,14 +585,23 @@ void UConvertToPolygonsTool::UpdateFromGroupLayer()
 	{
 		if (CopyFromLayerProperties->HasSelectedPolygroup() == false)
 		{
-			ActiveFromGroupSet = MakeShared<UE::Geometry::FPolygroupSet, ESPMode::ThreadSafe>(&ReadMesh);
+			ActiveFromGroupSetAdapter = MakeShared<FPolygroupLayerAdapter, ESPMode::ThreadSafe>(&ReadMesh, nullptr);
 		}
 		else
 		{
-			FName SelectedName = CopyFromLayerProperties->ActiveGroupLayer;
-			const FDynamicMeshPolygroupAttribute* FoundAttrib = UE::Geometry::FindPolygroupLayerByName(ReadMesh, SelectedName);
+			const FName SelectedName = CopyFromLayerProperties->ActiveGroupLayer;
+			const FDynamicMeshPolygroupAttribute* FoundAttrib = FindPolygroupLayerByName(ReadMesh, SelectedName);
+			if (!FoundAttrib)
+			{
+				if (const FDynamicMeshTriangleLabelAttribute* Labels = FindTriangleLabelLayerByName(ReadMesh, SelectedName))
+				{
+					ActiveFromGroupSetAdapter = MakeShared<FTriangleReadOnlyLabelAdapter>(&ReadMesh, Labels);
+					return;
+				}
+			}
+			
 			ensureMsgf(FoundAttrib, TEXT("Selected Attribute Not Found! Falling back to Default group layer."));
-			ActiveFromGroupSet = MakeShared<UE::Geometry::FPolygroupSet, ESPMode::ThreadSafe>(&ReadMesh, FoundAttrib);
+			ActiveFromGroupSetAdapter = MakeShared<FPolygroupLayerAdapter, ESPMode::ThreadSafe>(&ReadMesh, FoundAttrib);
 		}
 	});
 }
@@ -643,6 +651,122 @@ void UConvertToPolygonsTool::UpdateVisualization()
 	});
 }
 
+void UConvertToPolygonsTool::DoUpdateMeshAttribute(const FDynamicMesh3& InMesh) const
+{
+	if (Settings->GroupLayer == UE::PolygroupUtilLocals::DefaultLayer)
+	{
+		return;
+	}
+
+	FDynamicMesh3 UseResultMesh = *OriginalDynamicMesh;
+	if (UseResultMesh.HasAttributes() == false)
+	{
+		UseResultMesh.EnableAttributes();
+	}
+
+	switch(Settings->LayerType)
+	{
+	case EPerTriangleAttributeType::Integer:
+	{
+		// if we want to write to any layer other than default, we have to find or create it
+		FDynamicMeshPolygroupAttribute* UseAttribLayer = nullptr;
+		if (Settings->GroupLayer == UE::PolygroupUtilLocals::NewLayer)
+		{
+			// append new group layer and set it's name
+			int32 TargetLayerIdx = UseResultMesh.Attributes()->NumPolygroupLayers();
+			UseResultMesh.Attributes()->SetNumPolygroupLayers(TargetLayerIdx + 1);
+			UseAttribLayer = UseResultMesh.Attributes()->GetPolygroupLayer(TargetLayerIdx);
+			FString UseUniqueName = MakeUniqueGroupLayerName(UseResultMesh, Settings->NewLayerName);
+			UseAttribLayer->SetName( FName(UseUniqueName) );
+		}
+		else
+		{
+			UseAttribLayer = FindPolygroupLayerByName(UseResultMesh, Settings->GroupLayer);
+		}
+
+		if (UseAttribLayer)
+		{
+			// copy the generated groups from the Op output mesh (stored in primary groups) to the target layer
+			if (bUsingSelection)
+			{
+				for (int32 tid : *SelectionTriangleROI)
+				{
+					int32 GroupID = InMesh.GetTriangleGroup(OriginalSubmesh->MapTriangleToSubmesh(tid));
+					UseAttribLayer->SetValue(tid, GroupID);
+				}
+			}
+			else
+			{
+				for (int32 tid : UseResultMesh.TriangleIndicesItr())
+				{
+					UseAttribLayer->SetValue(tid, InMesh.GetTriangleGroup(tid));
+				}
+			}
+			UE::ToolTarget::CommitDynamicMeshUpdate(Target, UseResultMesh, true);
+		}
+		else
+		{
+			// If we can't find or create the layer (which should not be possible) the tool is going to do nothing,
+			// this is the safest option at this point
+			UE_LOG(LogGeometry, Warning, TEXT("Output group layer missing?"));
+		}
+		return;
+	}
+	case EPerTriangleAttributeType::Label:
+	{
+		// if we want to write to any layer other than default, we have to find or create it
+		FDynamicMeshTriangleLabelAttribute* LabelLayer = [&UseResultMesh, LayerName = Settings->GroupLayer, NewName = Settings->NewLayerName]
+		{
+			if (LayerName == UE::PolygroupUtilLocals::NewLayer)
+			{
+				FName NewAttrName(NewName);
+				int32 NumberCounter = 0;	
+				while (UseResultMesh.Attributes()->FindTriangleLabelAttribute(NewAttrName))
+				{
+					NewAttrName.SetNumber(NumberCounter++);
+				}
+				FDynamicMeshTriangleLabelAttribute* NewAttr = new FDynamicMeshTriangleLabelAttribute(&UseResultMesh);
+				UseResultMesh.Attributes()->AttachTriangleLabelAttribute(NewAttrName, NewAttr);
+				return NewAttr;
+			}
+			
+			return UseResultMesh.Attributes()->FindTriangleLabelAttribute(LayerName);
+		}();
+
+		if (LabelLayer)
+		{
+			FTriangleLabelAdapter LabelAdapter(&UseResultMesh, LabelLayer);
+						
+			// copy the generated groups from the Op output mesh (stored in primary groups) to the target layer
+			if (bUsingSelection)
+			{
+				for (int32 tid : *SelectionTriangleROI)
+				{
+					int32 GroupID = InMesh.GetTriangleGroup(OriginalSubmesh->MapTriangleToSubmesh(tid));
+					LabelAdapter.SetGroup(tid, GroupID);
+				}
+			}
+			else
+			{
+				for (int32 tid : UseResultMesh.TriangleIndicesItr())
+				{
+					LabelAdapter.SetGroup(tid, InMesh.GetTriangleGroup(tid));
+				}
+			}
+			UE::ToolTarget::CommitDynamicMeshUpdate(Target, UseResultMesh, true);
+		}
+		else
+		{
+			// If we can't find or create the layer (which should not be possible) the tool is going to do nothing,
+			// this is the safest option at this point
+			UE_LOG(LogGeometry, Warning, TEXT("Output group layer missing?"));
+		}
+		return;
+	}
+	default:
+		break;
+	}
+}
 
 
 #undef LOCTEXT_NAMESPACE
diff --git a/Engine/Plugins/Runtime/MeshModelingToolset/Source/MeshModelingTools/Private/MeshGroupPaintTool.cpp b/Engine/Plugins/Runtime/MeshModelingToolset/Source/MeshModelingTools/Private/MeshGroupPaintTool.cpp
--- a/Engine/Plugins/Runtime/MeshModelingToolset/Source/MeshModelingTools/Private/MeshGroupPaintTool.cpp
+++ b/Engine/Plugins/Runtime/MeshModelingToolset/Source/MeshModelingTools/Private/MeshGroupPaintTool.cpp
@@ -36,6 +36,7 @@
 #include "SceneView.h"
 #include "ToolTargetManager.h"
 #include "ToolBuilderUtil.h"
+#include "Changes/MeshTriangleLabelChange.h"
 
 #include UE_INLINE_GENERATED_CPP_BY_NAME(MeshGroupPaintTool)
 
@@ -150,6 +151,8 @@ void UMeshGroupPaintTool::Setup()
 	// initialize brush radius range interval, brush properties
 	UMeshSculptToolBase::InitializeBrushSizeRange(Bounds);
 
+	FilterProperties = NewObject<UGroupPaintBrushFilterProperties>(this);
+	
 	// Set up control points mechanic
 	PolyLassoMechanic = NewObject<UPolyLassoMarqueeMechanic>(this);
 	PolyLassoMechanic->Setup(this);
@@ -158,6 +161,7 @@ void UMeshGroupPaintTool::Setup()
 	PolyLassoMechanic->OnDrawPolyLassoFinished.AddUObject(this, &UMeshGroupPaintTool::OnPolyLassoFinished);
 
 	PolygroupLayerProperties = NewObject<UPolygroupLayersProperties>(this);
+	PolygroupLayerProperties->bSupportLabels = true;
 	PolygroupLayerProperties->RestoreProperties(this, TEXT("MeshGroupPaintTool"));
 	PolygroupLayerProperties->InitializeGroupLayers(GetSculptMesh());
 	PolygroupLayerProperties->WatchProperty(PolygroupLayerProperties->ActiveGroupLayer, [&](FName) { OnSelectedGroupLayerChanged(); });
@@ -165,7 +169,6 @@ void UMeshGroupPaintTool::Setup()
 	AddToolPropertySource(PolygroupLayerProperties);
 
 	// initialize other properties
-	FilterProperties = NewObject<UGroupPaintBrushFilterProperties>(this);
 	FilterProperties->WatchProperty(FilterProperties->SubToolType,
 		[this](EMeshGroupPaintInteractionType NewType) { UpdateSubToolType(NewType); });
 	FilterProperties->WatchProperty(FilterProperties->BrushSize,
@@ -175,7 +178,7 @@ void UMeshGroupPaintTool::Setup()
 	FilterProperties->RestoreProperties(this);
 	FilterProperties->BrushSize = UMeshSculptToolBase::BrushProperties->BrushSize.AdaptiveSize;
 	FilterProperties->bHitBackFaces = UMeshSculptToolBase::BrushProperties->bHitBackFaces;
-	FilterProperties->SetGroup = ActiveGroupSet->MaxGroupID;
+	FilterProperties->SetGroup = IsPaintingLabel() ? LabelAdapter->GetMaxGroupID() : ActiveGroupSet->MaxGroupID;
 	AddToolPropertySource(FilterProperties);
 
 	InitializeIndicator();
@@ -234,7 +237,7 @@ void UMeshGroupPaintTool::Setup()
 	// force colors update... ?
 	DynamicMeshComponent->SetTriangleColorFunction([this](const FDynamicMesh3* Mesh, int TriangleID)
 	{
-		return GetColorForGroup(ActiveGroupSet->GetGroup(TriangleID));
+		return GetColorForGroup(GetGroup(TriangleID));
 	});
 
 	// disable view properties
@@ -370,12 +373,12 @@ void UMeshGroupPaintTool::OnBeginStroke(const FRay& WorldRay)
 
 	if (PaintBrushOpProperties)
 	{
-		PaintBrushOpProperties->Group = FilterProperties->SetGroup;
+		PaintBrushOpProperties->Group = GetGroupToSet();
 		PaintBrushOpProperties->bOnlyPaintUngrouped = FilterProperties->bOnlySetUngrouped;
 	}
 	if (EraseBrushOpProperties)
 	{
-		EraseBrushOpProperties->Group = FilterProperties->EraseGroup;
+		EraseBrushOpProperties->Group = GetGroupToErase();
 		EraseBrushOpProperties->bOnlyEraseCurrent = FilterProperties->bOnlyEraseCurrent;
 	}
 
@@ -415,6 +418,7 @@ void UMeshGroupPaintTool::OnCancelStroke()
 {
 	GetActiveBrushOp()->CancelStroke();
 	ActiveGroupEditBuilder.Reset();
+	ActiveLabelEditBuilder.Reset();
 	bDrawGroupsDataValid = false;
 }
 
@@ -423,7 +427,7 @@ void UMeshGroupPaintTool::UpdateROI(const FSculptBrushStamp& BrushStamp)
 {
 	SCOPE_CYCLE_COUNTER(GroupPaintTool_UpdateROI);
 
-	int32 SetGroupID = GetInEraseStroke() ? FilterProperties->EraseGroup : FilterProperties->SetGroup;
+	const int32 SetGroupID = GetInEraseStroke() ? GetGroupToErase() : GetGroupToSet();
 
 	const FVector3d& BrushPos = BrushStamp.LocalFrame.Origin;
 	const FDynamicMesh3* Mesh = GetSculptMesh();
@@ -515,16 +519,17 @@ void UMeshGroupPaintTool::UpdateROI(const FSculptBrushStamp& BrushStamp)
 		TSet<int32> FillGroups;
 		for (int32 tid : TriangleROI)
 		{
-			if (ActiveGroupSet->GetGroup(tid) != SetGroupID)
+			const int32 Group = GetGroup(tid);
+			if (Group != SetGroupID)
 			{
 				StartROI.Add(tid);
-				FillGroups.Add(ActiveGroupSet->GetGroup(tid));
+				FillGroups.Add(Group);
 			}
 		}
 		FMeshConnectedComponents::GrowToConnectedTriangles(Mesh, StartROI, TriangleROI, &TempROIBuffer,
 			[&](int t1, int t2)
 		{
-			return (FillGroups.Contains(ActiveGroupSet->GetGroup(t2)));
+			return FillGroups.Contains(GetGroup(t2));
 		});
 	}
 
@@ -630,7 +635,7 @@ bool UMeshGroupPaintTool::SyncMeshWithGroupBuffer(FDynamicMesh3* Mesh)
 	for ( int32 k = 0; k < NumT; ++k)
 	{
 		int TriIdx = ROITriangleBuffer[k];
-		int32 CurGroupID = ActiveGroupSet->GetGroup(TriIdx);
+		const int32 CurGroupID = GetGroup(TriIdx);
 		if (FrozenGroups.Contains(CurGroupID))		// skip frozen groups
 		{
 			continue;
@@ -638,8 +643,18 @@ bool UMeshGroupPaintTool::SyncMeshWithGroupBuffer(FDynamicMesh3* Mesh)
 
 		if (ROIGroupBuffer[k] != CurGroupID)
 		{
-			ActiveGroupEditBuilder->SaveTriangle(TriIdx, CurGroupID, ROIGroupBuffer[k]);
-			ActiveGroupSet->SetGroup(TriIdx, ROIGroupBuffer[k], *Mesh);
+			if (IsPaintingLabel())
+			{
+				ActiveLabelEditBuilder->SaveTriangle(TriIdx,
+					LabelAdapter->GetValueFromGroup(CurGroupID),
+					LabelAdapter->GetValueFromGroup(ROIGroupBuffer[k]));
+			}
+			else
+			{
+				ActiveGroupEditBuilder->SaveTriangle(TriIdx, CurGroupID, ROIGroupBuffer[k]);
+			}
+			
+			SetGroup(TriIdx, ROIGroupBuffer[k], *Mesh);
 			//ActiveVertexChange->UpdateVertexColor(VertIdx, OrigColor, NewColor);
 			NumModified++;
 		}
@@ -725,7 +740,7 @@ void UMeshGroupPaintTool::OnPolyLassoFinished(const FCameraPolyLasso& Lasso, boo
 		return;
 	}
 
-	int32 SetGroupID = GetInEraseStroke() ? FilterProperties->EraseGroup : FilterProperties->SetGroup;
+	const int32 SetGroupID = GetInEraseStroke() ? GetGroupToErase() : GetGroupToSet();
 	SetTrianglesToGroupID(FaceSelection.AsSet(), SetGroupID, GetInEraseStroke());
 }
 
@@ -739,7 +754,7 @@ void UMeshGroupPaintTool::SetTrianglesToGroupID(const TSet<int32>& Triangles, in
 	TempROIBuffer.SetNum(0, EAllowShrinking::No);
 	for (int32 tid : Triangles)
 	{
-		int32 CurGroupID = ActiveGroupSet->GetGroup(tid);
+		int32 CurGroupID = GetGroup(tid);
 		if (CurGroupID == ToGroupID || FrozenGroups.Contains(CurGroupID))		// skip frozen groups
 		{
 			continue;
@@ -748,7 +763,7 @@ void UMeshGroupPaintTool::SetTrianglesToGroupID(const TSet<int32>& Triangles, in
 		{
 			continue;
 		}
-		if (bIsErase && FilterProperties->bOnlyEraseCurrent && CurGroupID != FilterProperties->SetGroup)
+		if (bIsErase && FilterProperties->bOnlyEraseCurrent && CurGroupID != GetGroupToSet())
 		{
 			continue;
 		}
@@ -764,12 +779,12 @@ void UMeshGroupPaintTool::SetTrianglesToGroupID(const TSet<int32>& Triangles, in
 		TempROIBuffer = MoveTemp(VisibleTriangles);
 	}
 
-	ActiveGroupEditBuilder->SaveTriangles(TempROIBuffer);
+	SaveTriangleEdits();
 	for (int32 tid : TempROIBuffer)
 	{
-		ActiveGroupSet->SetGroup(tid, ToGroupID, *GetSculptMesh());
+		SetGroup(tid, ToGroupID, *GetSculptMesh());
 	}
-	ActiveGroupEditBuilder->SaveTriangles(TempROIBuffer);
+	SaveTriangleEdits();
 
 	DynamicMeshComponent->FastNotifyTriangleVerticesUpdated(TempROIBuffer, EMeshRenderAttributeFlags::VertexColors);
 	GetToolManager()->PostInvalidation();
@@ -946,6 +961,15 @@ void UMeshGroupPaintTool::DrawHUD(FCanvas* Canvas, IToolsContextRenderAPI* Rende
 	FVector3d LocalEyePosition(CurTargetTransform.InverseTransformPosition((FVector3d)CamState.Position));
 
 	FDynamicMesh3* Mesh = GetSculptMesh();
+	
+	auto GetGroupString = [GroupIdLayer = ActiveGroupSet.Get(), GroupLabelAdapter = LabelAdapter.Get()](int CursorHitTID)
+	{
+		if (GroupLabelAdapter)
+		{
+			return FString::Printf(TEXT("('%s')"), *GroupLabelAdapter->GetValue(CursorHitTID).ToString());
+		}
+		return FString::Printf(TEXT("%d"), GroupIdLayer->GetTriangleGroup(CursorHitTID));
+	};
 
 	// draw the group number for the group under the cursor
 	if (FilterProperties->bShowHitGroup)
@@ -953,10 +977,9 @@ void UMeshGroupPaintTool::DrawHUD(FCanvas* Canvas, IToolsContextRenderAPI* Rende
 		int CursorHitTID = Octree.FindNearestHitObject(FRay3d(LocalEyePosition, Normalized(HoverStamp.LocalFrame.Origin - LocalEyePosition)));
 		if (Mesh->IsTriangle(CursorHitTID))
 		{
-			int32 GroupID = ActiveGroupSet->GetTriangleGroup(CursorHitTID);
 			FVector2D CursorPixelPos;
 			SceneView->WorldToPixel(HoverStamp.WorldFrame.Origin, CursorPixelPos);
-			FString CursorString = FString::Printf(TEXT("%d"), GroupID);
+			const FString CursorString = GetGroupString(CursorHitTID);
 			Canvas->DrawShadowedString(CursorPixelPos.X / (double)DPIScale, CursorPixelPos.Y / (double)DPIScale, *CursorString, UseFont, FLinearColor::White);
 		}
 	}
@@ -966,7 +989,14 @@ void UMeshGroupPaintTool::DrawHUD(FCanvas* Canvas, IToolsContextRenderAPI* Rende
 	{
 		if (bDrawGroupsDataValid == false)
 		{
-			GroupVisualizationCache.UpdateGroupInfo_ConnectedComponents(*Mesh, *ActiveGroupSet);
+			if (IsPaintingLabel())
+			{
+				GroupVisualizationCache.UpdateGroupInfo_ConnectedComponents(*Mesh, *LabelAdapter);
+			}
+			else
+			{
+				GroupVisualizationCache.UpdateGroupInfo_ConnectedComponents(*Mesh, *ActiveGroupSet);
+			}
 			bDrawGroupsDataValid = true;
 		}
 		for (const FGroupVisualizationCache::FGroupInfo& GroupInfo : GroupVisualizationCache)
@@ -1030,10 +1060,14 @@ void UMeshGroupPaintTool::OnTick(float DeltaTime)
 		{
 			if (GetSculptMesh()->IsTriangle(GetBrushTriangleID()))
 			{
-				int32 HitGroupID = ActiveGroupSet->GetGroup(GetBrushTriangleID());
+				const int32 HitGroupID = GetGroup(GetBrushTriangleID());
 				if (bPendingPickGroup)
 				{
 					FilterProperties->SetGroup = HitGroupID;
+					if (IsPaintingLabel())
+					{
+						FilterProperties->SetLabel = LabelAdapter->GetValueFromGroup(HitGroupID);
+					}
 					NotifyOfPropertyChangeByTool(FilterProperties);
 				}
 				else if (bPendingToggleFreezeGroup)
@@ -1087,7 +1121,7 @@ void UMeshGroupPaintTool::OnTick(float DeltaTime)
 
 void UMeshGroupPaintTool::AllocateNewGroupAndSetAsCurrentAction()
 {
-	FilterProperties->SetGroup = ActiveGroupSet->MaxGroupID;
+	FilterProperties->SetGroup = IsPaintingLabel() ? LabelAdapter->GetMaxGroupID() : ActiveGroupSet->MaxGroupID;
 	NotifyOfPropertyChangeByTool(FilterProperties);
 }
 
@@ -1122,7 +1156,7 @@ void UMeshGroupPaintTool::ToggleFrozenGroup(int32 FreezeGroupID)
 	TempROIBuffer.SetNum(0, EAllowShrinking::No);
 	for (int32 tid : Mesh->TriangleIndicesItr())
 	{
-		int32 TriGroupID = ActiveGroupSet->GetGroup(tid);
+		int32 TriGroupID = GetGroup(tid);
 		if (TriGroupID == FreezeGroupID)
 		{
 			TempROIBuffer.Add(tid);
@@ -1141,10 +1175,10 @@ void UMeshGroupPaintTool::FreezeOtherGroups(int32 KeepGroupID)
 	TempROIBuffer.SetNum(0, EAllowShrinking::No);
 	for (int32 tid : Mesh->TriangleIndicesItr())
 	{
-		int32 GroupID = ActiveGroupSet->GetGroup(tid);
+		int32 GroupID = GetGroup(tid);
 		if ( GroupID != 0 && GroupID != KeepGroupID)
 		{
-			FrozenGroups.AddUnique(ActiveGroupSet->GetGroup(tid));
+			FrozenGroups.AddUnique(GetGroup(tid));
 			TempROIBuffer.Add(tid);
 		}
 	}
@@ -1160,7 +1194,7 @@ void UMeshGroupPaintTool::ClearAllFrozenGroups()
 	TempROIBuffer.SetNum(0, EAllowShrinking::No);
 	for (int32 tid : Mesh->TriangleIndicesItr())
 	{
-		if ( FrozenGroups.Contains(ActiveGroupSet->GetGroup(tid)) )
+		if ( FrozenGroups.Contains(GetGroup(tid)) )
 		{
 			TempROIBuffer.Add(tid);
 		}
@@ -1194,21 +1228,21 @@ void UMeshGroupPaintTool::GrowCurrentGroupAction()
 {
 	BeginChange();
 
-	int32 CurrentGroupID = FilterProperties->SetGroup;
+	const int32 CurrentGroupID = GetGroupToSet();
 	const FDynamicMesh3* Mesh = DynamicMeshComponent->GetMesh();
 	FMeshFaceSelection InitialSelection(Mesh);
-	InitialSelection.Select([&](int32 tid) { return ActiveGroupSet->GetGroup(tid) == CurrentGroupID; });
+	InitialSelection.Select([&](int32 tid) { return GetGroup(tid) == CurrentGroupID; });
 	FMeshFaceSelection ExpandSelection(InitialSelection);
-	ExpandSelection.ExpandToOneRingNeighbours([&](int32 tid) { return FrozenGroups.Contains(ActiveGroupSet->GetGroup(tid)) == false; });
+	ExpandSelection.ExpandToOneRingNeighbours([&](int32 tid) { return FrozenGroups.Contains(GetGroup(tid)) == false; });
 	TempROIBuffer.SetNum(0, EAllowShrinking::No);
 	ExpandSelection.SetDifference(InitialSelection, TempROIBuffer);
 
-	ActiveGroupEditBuilder->SaveTriangles(TempROIBuffer);
+	SaveTriangleEdits();
 	for (int32 tid : TempROIBuffer)
 	{
-		ActiveGroupSet->SetGroup(tid, CurrentGroupID, *GetSculptMesh());
+		SetGroup(tid, CurrentGroupID, *GetSculptMesh());
 	}
-	ActiveGroupEditBuilder->SaveTriangles(TempROIBuffer);
+	SaveTriangleEdits();
 
 	DynamicMeshComponent->FastNotifyTriangleVerticesUpdated(TempROIBuffer, EMeshRenderAttributeFlags::VertexColors);
 	GetToolManager()->PostInvalidation();
@@ -1220,22 +1254,22 @@ void UMeshGroupPaintTool::ShrinkCurrentGroupAction()
 {
 	BeginChange();
 
-	int32 CurrentGroupID = FilterProperties->SetGroup;
+	const int32 CurrentGroupID = GetGroupToSet();
 	const FDynamicMesh3* Mesh = DynamicMeshComponent->GetMesh();
 	FMeshFaceSelection InitialSelection(Mesh);
-	InitialSelection.Select([&](int32 tid) { return ActiveGroupSet->GetGroup(tid) == CurrentGroupID; });
+	InitialSelection.Select([&](int32 tid) { return GetGroup(tid) == CurrentGroupID; });
 	FMeshFaceSelection ContractSelection(InitialSelection);
 	ContractSelection.ContractBorderByOneRingNeighbours();
 	TempROIBuffer.SetNum(0, EAllowShrinking::No);
 	InitialSelection.SetDifference(ContractSelection, TempROIBuffer);
 
-	ActiveGroupEditBuilder->SaveTriangles(TempROIBuffer);
+	SaveTriangleEdits();
 	for (int32 tid : TempROIBuffer)
 	{
 		// todo: could probably guess boundary groups here...
-		ActiveGroupSet->SetGroup(tid, 0, *GetSculptMesh());
+		SetGroup(tid, 0, *GetSculptMesh());
 	}
-	ActiveGroupEditBuilder->SaveTriangles(TempROIBuffer);
+	SaveTriangleEdits();
 
 	DynamicMeshComponent->FastNotifyTriangleVerticesUpdated(TempROIBuffer, EMeshRenderAttributeFlags::VertexColors);
 	GetToolManager()->PostInvalidation();
@@ -1247,23 +1281,23 @@ void UMeshGroupPaintTool::ClearCurrentGroupAction()
 {
 	BeginChange();
 
-	int32 CurrentGroupID = FilterProperties->SetGroup;
+	const int32 CurrentGroupID = GetGroupToSet();
 	const FDynamicMesh3* Mesh = DynamicMeshComponent->GetMesh();
 	TempROIBuffer.SetNum(0, EAllowShrinking::No);
 	for (int32 tid : Mesh->TriangleIndicesItr())
 	{
-		if (ActiveGroupSet->GetGroup(tid) == CurrentGroupID)
+		if (GetGroup(tid) == CurrentGroupID)
 		{
 			TempROIBuffer.Add(tid);
 		}
 	}
 
-	ActiveGroupEditBuilder->SaveTriangles(TempROIBuffer);
+	SaveTriangleEdits();
 	for (int32 tid : TempROIBuffer)
 	{
-		ActiveGroupSet->SetGroup(tid, 0, *GetSculptMesh());
+		SetGroup(tid, 0, *GetSculptMesh());
 	}
-	ActiveGroupEditBuilder->SaveTriangles(TempROIBuffer);
+	SaveTriangleEdits();
 
 	DynamicMeshComponent->FastNotifyTriangleVerticesUpdated(TempROIBuffer, EMeshRenderAttributeFlags::VertexColors);
 	GetToolManager()->PostInvalidation();
@@ -1275,24 +1309,24 @@ void UMeshGroupPaintTool::FloodFillCurrentGroupAction()
 {
 	BeginChange();
 
-	int32 SetGroupID = FilterProperties->SetGroup;
+	const int32 SetGroupID = GetGroupToSet();
 	const FDynamicMesh3* Mesh = DynamicMeshComponent->GetMesh();
 	TempROIBuffer.SetNum(0, EAllowShrinking::No);
 	for (int32 tid : Mesh->TriangleIndicesItr())
 	{
-		int32 GroupID = ActiveGroupSet->GetGroup(tid);
+		const int32 GroupID = GetGroup(tid);
 		if (GroupID == 0 && GroupID != SetGroupID && FrozenGroups.Contains(GroupID) == false)
 		{
 			TempROIBuffer.Add(tid);
 		}
 	}
 
-	ActiveGroupEditBuilder->SaveTriangles(TempROIBuffer);
+	SaveTriangleEdits();
 	for (int32 tid : TempROIBuffer)
 	{
-		ActiveGroupSet->SetGroup(tid, SetGroupID, *GetSculptMesh());
+		SetGroup(tid, SetGroupID, *GetSculptMesh());
 	}
-	ActiveGroupEditBuilder->SaveTriangles(TempROIBuffer);
+	SaveTriangleEdits();
 
 	DynamicMeshComponent->FastNotifyTriangleVerticesUpdated(TempROIBuffer, EMeshRenderAttributeFlags::VertexColors);
 	GetToolManager()->PostInvalidation();
@@ -1308,18 +1342,18 @@ void UMeshGroupPaintTool::ClearAllGroupsAction()
 	TempROIBuffer.SetNum(0, EAllowShrinking::No);
 	for (int32 tid : Mesh->TriangleIndicesItr())
 	{
-		if (ActiveGroupSet->GetGroup(tid) != 0)
+		if (GetGroup(tid) != 0)
 		{
 			TempROIBuffer.Add(tid);
 		}
 	}
 
-	ActiveGroupEditBuilder->SaveTriangles(TempROIBuffer);
+	SaveTriangleEdits();
 	for (int32 tid : TempROIBuffer)
 	{
-		ActiveGroupSet->SetGroup(tid, 0, *GetSculptMesh());
+		SetGroup(tid, 0, *GetSculptMesh());
 	}
-	ActiveGroupEditBuilder->SaveTriangles(TempROIBuffer);
+	SaveTriangleEdits();
 
 	DynamicMeshComponent->FastNotifyTriangleVerticesUpdated(TempROIBuffer, EMeshRenderAttributeFlags::VertexColors);
 	GetToolManager()->PostInvalidation();
@@ -1334,32 +1368,77 @@ void UMeshGroupPaintTool::ClearAllGroupsAction()
 //
 void UMeshGroupPaintTool::BeginChange()
 {
-	check(ActiveGroupEditBuilder == nullptr);
-	ActiveGroupEditBuilder = MakeUnique<FDynamicMeshGroupEditBuilder>(ActiveGroupSet.Get());
+	check(ActiveGroupEditBuilder == nullptr && ActiveLabelEditBuilder == nullptr);
+	
+	if (IsPaintingLabel())
+	{
+		ActiveLabelEditBuilder = MakeUnique<FDynamicMeshTriangleLabelEditBuilder>(LabelAdapter.Get());
+	}
+	else
+	{
+		ActiveGroupEditBuilder = MakeUnique<FDynamicMeshGroupEditBuilder>(ActiveGroupSet.Get());
+	}
+
 	LongTransactions.Open(LOCTEXT("GroupPaintChange", "Group Stroke"), GetToolManager());
 }
 
 void UMeshGroupPaintTool::EndChange()
 {
-	check(ActiveGroupEditBuilder);
+	check(ActiveGroupEditBuilder || ActiveLabelEditBuilder);
+
+	static const FText Description = LOCTEXT("GroupPaintChange", "Group Stroke");
+	if (ActiveLabelEditBuilder)
+	{
+		TUniquePtr<FDynamicMeshTriangleLabelEdit> EditResult = ActiveLabelEditBuilder->ExtractResult();
+		ActiveLabelEditBuilder = nullptr;
 
-	TUniquePtr<FDynamicMeshGroupEdit> EditResult = ActiveGroupEditBuilder->ExtractResult();
-	ActiveGroupEditBuilder = nullptr;
+		TUniquePtr<TWrappedToolCommandChange<FMeshTriangleLabelChange>> NewChange = MakeUnique<TWrappedToolCommandChange<FMeshTriangleLabelChange>>();
+		NewChange->WrappedChange = MakeUnique<FMeshTriangleLabelChange>(MoveTemp(EditResult));
+		NewChange->BeforeModify = [this](bool bRevert)
+		{
+			this->WaitForPendingUndoRedoUpdate();
+		};
 
-	TUniquePtr<TWrappedToolCommandChange<FMeshPolygroupChange>> NewChange = MakeUnique<TWrappedToolCommandChange<FMeshPolygroupChange>>();
-	NewChange->WrappedChange = MakeUnique<FMeshPolygroupChange>(MoveTemp(EditResult));
-	NewChange->BeforeModify = [this](bool bRevert)
+		GetToolManager()->EmitObjectChange(DynamicMeshComponent, MoveTemp(NewChange), Description);
+	}
+	else
 	{
-		this->WaitForPendingUndoRedoUpdate();
-	};
+		TUniquePtr<FDynamicMeshGroupEdit> EditResult = ActiveGroupEditBuilder->ExtractResult();
+		ActiveGroupEditBuilder = nullptr;
+
+		TUniquePtr<TWrappedToolCommandChange<FMeshPolygroupChange>> NewChange = MakeUnique<TWrappedToolCommandChange<FMeshPolygroupChange>>();
+		NewChange->WrappedChange = MakeUnique<FMeshPolygroupChange>(MoveTemp(EditResult));
+		NewChange->BeforeModify = [this](bool bRevert)
+			{
+				this->WaitForPendingUndoRedoUpdate();
+			};
 
-	GetToolManager()->EmitObjectChange(DynamicMeshComponent, MoveTemp(NewChange), LOCTEXT("GroupPaintChange", "Group Stroke"));
+		GetToolManager()->EmitObjectChange(DynamicMeshComponent, MoveTemp(NewChange), Description);
+	}
+	
 	LongTransactions.Close(GetToolManager());
 
 	// debug groups are invalid now
 	bDrawGroupsDataValid = false;
 }
 
+void UMeshGroupPaintTool::SaveTriangleEdits() const
+{
+	if (IsPaintingLabel())
+	{
+		if (ensure(ActiveLabelEditBuilder))
+		{
+			ActiveLabelEditBuilder->SaveTriangles(TempROIBuffer);
+		}
+	}
+	else
+	{
+		if (ensure(ActiveGroupEditBuilder))
+		{
+			ActiveGroupEditBuilder->SaveTriangles(TempROIBuffer);
+		}
+	}
+}
 
 void UMeshGroupPaintTool::WaitForPendingUndoRedoUpdate()
 {
@@ -1421,6 +1500,25 @@ void UMeshGroupPaintTool::PrecomputeFilterData()
 	});
 }
 
+int32 UMeshGroupPaintTool::GetGroup(const int32 TriangleID) const
+{
+	return IsPaintingLabel() ? LabelAdapter->GetGroup(TriangleID) : ActiveGroupSet->GetGroup(TriangleID);
+}
+
+void UMeshGroupPaintTool::SetGroup(const int32 TriangleID, const int32 GroupID, FDynamicMesh3& Mesh) const
+{
+	return IsPaintingLabel() ? LabelAdapter->SetGroup(TriangleID, GroupID, Mesh) : ActiveGroupSet->SetGroup(TriangleID, GroupID, Mesh);
+}
+
+int32 UMeshGroupPaintTool::GetGroupToSet() const
+{
+	return IsPaintingLabel() ? LabelAdapter->GetGroupFromValue(FilterProperties->SetLabel) : FilterProperties->SetGroup;
+}
+
+int32 UMeshGroupPaintTool::GetGroupToErase() const
+{
+	return IsPaintingLabel() ? LabelAdapter->GetGroupFromValue(FilterProperties->EraseLabel) : FilterProperties->EraseGroup;
+}
 
 void UMeshGroupPaintTool::OnSelectedGroupLayerChanged()
 {
@@ -1455,17 +1553,37 @@ void UMeshGroupPaintTool::OnSelectedGroupLayerChanged()
 
 void UMeshGroupPaintTool::UpdateActiveGroupLayer()
 {
+	ActiveGroupSet.Reset();
+	LabelAdapter.Reset();
+	FilterProperties->bPaintingLabels = false;
+	
+	FDynamicMesh3* Mesh = GetSculptMesh();
+	
 	if (PolygroupLayerProperties->HasSelectedPolygroup() == false)
 	{
-		ActiveGroupSet = MakeUnique<UE::Geometry::FPolygroupSet>(GetSculptMesh());
+		ActiveGroupSet = MakeUnique<FPolygroupSet>(Mesh);
 	}
 	else
 	{
 		FName SelectedName = PolygroupLayerProperties->ActiveGroupLayer;
-		FDynamicMeshPolygroupAttribute* FoundAttrib = UE::Geometry::FindPolygroupLayerByName(*GetSculptMesh(), SelectedName);
-		ensureMsgf(FoundAttrib, TEXT("Selected Attribute Not Found! Falling back to Default group layer."));
-		ActiveGroupSet = MakeUnique<UE::Geometry::FPolygroupSet>(GetSculptMesh(), FoundAttrib);
+		FDynamicMeshPolygroupAttribute* FoundAttrib = FindPolygroupLayerByName(*Mesh, SelectedName);
+		if (!FoundAttrib)
+		{
+			if (FDynamicMeshTriangleLabelAttribute* Labels = FindTriangleLabelLayerByName(*Mesh, SelectedName))
+			{
+				LabelAdapter = MakeShared<FTriangleLabelAdapter>(Mesh, Labels);
+				FilterProperties->bPaintingLabels = true;
+			}
+		}
+		
+		if (!LabelAdapter.IsValid())
+		{
+			ensureMsgf(FoundAttrib, TEXT("Selected Attribute Not Found! Falling back to Default group layer."));
+			ActiveGroupSet = MakeUnique<FPolygroupSet>(Mesh, FoundAttrib);
+		}
 	}
+	
+	FilterProperties->bPaintingLabels = IsPaintingLabel();
 
 	// need to reset everything here...
 	FrozenGroups.Reset();
@@ -1525,11 +1643,11 @@ void UMeshGroupPaintTool::ApplyAction(EMeshGroupPaintToolActions ActionType)
 		break;
 
 	case EMeshGroupPaintToolActions::FreezeCurrent:
-		ToggleFrozenGroup(FilterProperties->SetGroup);
+		ToggleFrozenGroup(GetGroupToSet());
 		break;
 
 	case EMeshGroupPaintToolActions::FreezeOthers:
-		FreezeOtherGroups(FilterProperties->SetGroup);
+		FreezeOtherGroups(GetGroupToSet());
 		break;
 
 	case EMeshGroupPaintToolActions::GrowCurrent:
diff --git a/Engine/Plugins/Runtime/MeshModelingToolset/Source/MeshModelingTools/Public/ConvertToPolygonsTool.h b/Engine/Plugins/Runtime/MeshModelingToolset/Source/MeshModelingTools/Public/ConvertToPolygonsTool.h
--- a/Engine/Plugins/Runtime/MeshModelingToolset/Source/MeshModelingTools/Public/ConvertToPolygonsTool.h
+++ b/Engine/Plugins/Runtime/MeshModelingToolset/Source/MeshModelingTools/Public/ConvertToPolygonsTool.h
@@ -21,6 +21,7 @@ namespace UE::Geometry
 {
 	class FDynamicMesh3;
 	struct FDynamicSubmesh3;
+	class ITriangleGroupSetAdapter;
 }
 
 /**
@@ -60,6 +61,15 @@ enum class EConvertToPolygonsMode
 	CopyFromLayer = 6 UMETA(DisplayName = "Copy From Layer")
 };
 
+UENUM()
+enum class EPerTriangleAttributeType
+{
+	/** Create per-triangle integer attribute. */
+	Integer = 0 UMETA(DisplayName = "Int32"),
+	/** Create per-triangle label attribute. */
+	Label = 1 UMETA(DisplayName = "FName")
+};
+
 
 
 UCLASS(MinimalAPI)
@@ -149,6 +159,11 @@ public:
 	UPROPERTY(EditAnywhere, Category = Output, meta = (TransientToolProperty, DisplayName = "New Layer Name",
 		EditCondition = "bShowNewLayerName", HideEditConditionToggle, NoResetToDefault))
 	FString NewLayerName = TEXT("polygroups");
+
+	/** Type of the new Group Layer */
+	UPROPERTY(EditAnywhere, Category = Output, meta = (TransientToolProperty,
+		EditCondition = "bShowNewLayerName", HideEditConditionToggle, NoResetToDefault))
+	EPerTriangleAttributeType LayerType = EPerTriangleAttributeType::Integer;
 };
 
 
@@ -231,10 +246,12 @@ protected:
 	UE_API void UpdateVisualization();
 
 	// current input group set used in CopyFromLayer mode, relative to ComputeOperatorSharedMesh
-	TSharedPtr<UE::Geometry::FPolygroupSet, ESPMode::ThreadSafe> ActiveFromGroupSet;
+	TSharedPtr<UE::Geometry::ITriangleGroupSetAdapter, ESPMode::ThreadSafe> ActiveFromGroupSetAdapter;
 	UE_API void OnSelectedFromGroupLayerChanged();
 	UE_API void UpdateFromGroupLayer();
-
+	
+	// updates the group layer data to the mesh 
+	void DoUpdateMeshAttribute(const FDynamicMesh3& InMesh) const;
 };
 
 #undef UE_API
diff --git a/Engine/Plugins/Runtime/MeshModelingToolset/Source/MeshModelingTools/Public/MeshGroupPaintTool.h b/Engine/Plugins/Runtime/MeshModelingToolset/Source/MeshModelingTools/Public/MeshGroupPaintTool.h
--- a/Engine/Plugins/Runtime/MeshModelingToolset/Source/MeshModelingTools/Public/MeshGroupPaintTool.h
+++ b/Engine/Plugins/Runtime/MeshModelingToolset/Source/MeshModelingTools/Public/MeshGroupPaintTool.h
@@ -26,9 +26,14 @@
 
 #define UE_API MESHMODELINGTOOLS_API
 
+class FDynamicMeshTriangleLabelEditBuilder;
 class UMeshElementsVisualizer;
 class UGroupEraseBrushOpProps;
 class UGroupPaintBrushOpProps;
+namespace UE::Geometry
+{
+	class FTriangleLabelAdapter;
+}
 
 DECLARE_STATS_GROUP(TEXT("GroupPaintTool"), STATGROUP_GroupPaintTool, STATCAT_Advanced);
 DECLARE_CYCLE_STAT(TEXT("GroupPaintTool_UpdateROI"), GroupPaintTool_UpdateROI, STATGROUP_GroupPaintTool);
@@ -144,17 +149,27 @@ public:
 		HideEditConditionToggle, EditConditionHides, EditCondition = "SubToolType != EMeshGroupPaintInteractionType::PolyLasso"))
 	bool bHitBackFaces = true;
 
-	/** The group that will be assigned to triangles */
-	UPROPERTY(EditAnywhere, Category = ActionType, meta = (UIMin = 0, ClampMin = 0, Delta = 1, LinearDeltaSensitivity = 50))
+	/** The group index that will be assigned to triangles */
+	UPROPERTY(EditAnywhere, Category = ActionType, meta = (UIMin = 0, ClampMin = 0, Delta = 1, LinearDeltaSensitivity = 50,
+		HideEditConditionToggle, EditConditionHides, EditCondition = "bPaintingLabels == false"))
 	int32 SetGroup = 1;
+	
+	/** The group label that will be assigned to triangles */
+	UPROPERTY(EditAnywhere, Category = ActionType, meta = (HideEditConditionToggle, EditConditionHides, EditCondition = bUseLabel))
+	FName SetLabel = NAME_None;
 
 	/** If true, only triangles with no group assigned will be painted */
 	UPROPERTY(EditAnywhere, Category = ActionType)
 	bool bOnlySetUngrouped = false;
 
-	/** Group to set as Erased value */
-	UPROPERTY(EditAnywhere, Category = ActionType, meta = (UIMin = 0, ClampMin = 0, Delta = 1, LinearDeltaSensitivity = 50))
+	/** Group index to set as Erased value */
+	UPROPERTY(EditAnywhere, Category = ActionType, meta = (UIMin = 0, ClampMin = 0, Delta = 1, LinearDeltaSensitivity = 50,
+		HideEditConditionToggle, EditConditionHides, EditCondition = "bPaintingLabels == false"))
 	int32 EraseGroup = 0;
+	
+	/** Group label to set as Erased value */
+	UPROPERTY(EditAnywhere, Category = ActionType, meta = (HideEditConditionToggle, EditConditionHides, EditCondition = bUseLabel))
+	FName EraseLabel = NAME_None;
 
 	/** When enabled, only the current group configured in the Paint brush is erased */
 	UPROPERTY(EditAnywhere, Category = ActionType)
@@ -189,6 +204,10 @@ public:
 	/** Display the Group ID for all visible groups in the mesh */
 	UPROPERTY(EditAnywhere, Category = Visualization)
 	bool bShowAllGroups = false;
+	
+	/** whether the user is painting labels */
+	UPROPERTY()
+	bool bPaintingLabels = false;
 };
 
 
@@ -421,7 +440,20 @@ protected:
 	UE_API void OnDynamicMeshComponentChanged();
 	FDelegateHandle OnDynamicMeshComponentChangedHandle;
 
+	int32 GetGroup(const int32 TriangleID) const;
+	void SetGroup(const int32 TriangleID, const int32 GroupID, FDynamicMesh3& Mesh) const;
+	
+	int32 GetGroupToSet() const;
+	int32 GetGroupToErase() const;
+	
 	TUniquePtr<UE::Geometry::FPolygroupSet> ActiveGroupSet;
+	TSharedPtr<UE::Geometry::FTriangleLabelAdapter> LabelAdapter;
+	
+	bool IsPaintingLabel() const
+	{
+		return LabelAdapter.IsValid();
+	}
+	
 	UE_API void OnSelectedGroupLayerChanged();
 	UE_API void UpdateActiveGroupLayer();
 
@@ -466,8 +498,10 @@ protected:
 	UE_API bool SyncMeshWithGroupBuffer(FDynamicMesh3* Mesh);
 
 	TUniquePtr<FDynamicMeshGroupEditBuilder> ActiveGroupEditBuilder;
+	TUniquePtr<FDynamicMeshTriangleLabelEditBuilder> ActiveLabelEditBuilder;
 	UE_API void BeginChange();
 	UE_API void EndChange();
+	UE_API void SaveTriangleEdits() const;
 
 	TArray<int32> FrozenGroups;
 	UE_API void ToggleFrozenGroup(int32 GroupID);
diff --git a/Engine/Plugins/Runtime/MeshModelingToolset/Source/ModelingComponents/Private/Changes/MeshTriangleLabelChange.cpp b/Engine/Plugins/Runtime/MeshModelingToolset/Source/ModelingComponents/Private/Changes/MeshTriangleLabelChange.cpp
new file mode 100644
--- /dev/null
+++ b/Engine/Plugins/Runtime/MeshModelingToolset/Source/ModelingComponents/Private/Changes/MeshTriangleLabelChange.cpp
@@ -0,0 +1,93 @@
+// Copyright Epic Games, Inc. All Rights Reserved.
+
+#include "Changes/MeshTriangleLabelChange.h"
+#include "DynamicMesh/DynamicMesh3.h"
+
+using namespace UE::Geometry;
+
+void FDynamicMeshTriangleLabelEdit::ApplyToMesh(FDynamicMesh3* Mesh, bool bRevert)
+{
+	if (!ensure(Mesh && LayerName != NAME_None))
+	{
+		return;
+	}
+	
+	FDynamicMeshAttributeSet* Attributes = Mesh->Attributes();
+	if (!ensure(Attributes))
+	{
+		return;
+	}
+	
+	FDynamicMeshTriangleLabelAttribute* LabelAttribute = Attributes->FindTriangleLabelAttribute(LayerName);
+	if (ensure(LabelAttribute))
+	{
+		const int32 NumTris = Triangles.Num();
+		for (int32 k = 0; k < NumTris; ++k)
+		{
+			const int32 tid = Triangles[k];
+			if (Mesh->IsTriangle(tid) && IsTriangleIndexValid(k, bRevert))
+			{
+				const FName Label = (bRevert) ? OldLabels[k] : NewLabels[k];
+				LabelAttribute->SetValue(tid, Label);
+			}
+		}
+	}
+}
+
+bool FDynamicMeshTriangleLabelEdit::IsTriangleIndexValid(const int32 TriangleIndex, const bool bRevert) const
+{
+	return bRevert ? OldLabels.IsValidIndex(TriangleIndex) : NewLabels.IsValidIndex(TriangleIndex);
+}
+
+FDynamicMeshTriangleLabelEditBuilder::FDynamicMeshTriangleLabelEditBuilder(FTriangleLabelAdapter* InAdapter)
+	: Adapter(MakePimpl<FTriangleLabelAdapter>(*InAdapter))
+	, Edit(MakeUnique<FDynamicMeshTriangleLabelEdit>())
+{
+	Edit->LayerName = InAdapter->GetLayerName();
+}
+
+void FDynamicMeshTriangleLabelEditBuilder::SaveTriangle(int32 TriangleID)
+{
+	FName NewGroupLabel = Adapter->GetValue(TriangleID);
+	
+	const int32* NewIndex = SavedIndexMap.Find(TriangleID);
+	if (NewIndex == nullptr)
+	{
+		int32 Index = Edit->Triangles.Num();
+		SavedIndexMap.Add(TriangleID, Index);
+		Edit->Triangles.Add(TriangleID);
+		Edit->OldLabels.Add(NewGroupLabel);
+		Edit->NewLabels.Add(NewGroupLabel);
+	}
+	else
+	{
+		Edit->NewLabels[*NewIndex] = NewGroupLabel;
+	}
+}
+
+
+void FDynamicMeshTriangleLabelEditBuilder::SaveTriangle(int32 TriangleID, FName OldLabel, FName NewLabel)
+{
+	const int32* NewIndex = SavedIndexMap.Find(TriangleID);
+	if (NewIndex == nullptr)
+	{
+		int32 Index = Edit->Triangles.Num();
+		SavedIndexMap.Add(TriangleID, Index);
+		Edit->Triangles.Add(TriangleID);
+		Edit->OldLabels.Add(OldLabel);
+		Edit->NewLabels.Add(NewLabel);
+	}
+	else
+	{
+		Edit->NewLabels[*NewIndex] = NewLabel;
+	}
+}
+
+FMeshTriangleLabelChange::FMeshTriangleLabelChange(TUniquePtr<FDynamicMeshTriangleLabelEdit>&& LabelEditIn)
+	: LabelEdit(MoveTemp(LabelEditIn))
+{}
+
+void FMeshTriangleLabelChange::ApplyChangeToMesh(FDynamicMesh3* Mesh, bool bRevert) const
+{
+	LabelEdit->ApplyToMesh(Mesh, bRevert);
+}
diff --git a/Engine/Plugins/Runtime/MeshModelingToolset/Source/ModelingComponents/Private/PropertySets/PolygroupLayersProperties.cpp b/Engine/Plugins/Runtime/MeshModelingToolset/Source/ModelingComponents/Private/PropertySets/PolygroupLayersProperties.cpp
--- a/Engine/Plugins/Runtime/MeshModelingToolset/Source/ModelingComponents/Private/PropertySets/PolygroupLayersProperties.cpp
+++ b/Engine/Plugins/Runtime/MeshModelingToolset/Source/ModelingComponents/Private/PropertySets/PolygroupLayersProperties.cpp
@@ -13,13 +13,22 @@ void UPolygroupLayersProperties::InitializeGroupLayers(const FDynamicMesh3* Mesh
 {
 	GroupLayersList.Reset();
 	GroupLayersList.Add(TEXT("Default"));		// always have standard group
-	if (Mesh->Attributes())
+	
+	if (const FDynamicMeshAttributeSet* Attributes = Mesh->Attributes())
 	{
-		for (int32 k = 0; k < Mesh->Attributes()->NumPolygroupLayers(); k++)
+		for (int32 k = 0; k < Attributes->NumPolygroupLayers(); k++)
 		{
-			FName Name = Mesh->Attributes()->GetPolygroupLayer(k)->GetName();
+			FName Name = Attributes->GetPolygroupLayer(k)->GetName();
 			GroupLayersList.Add(Name.ToString());
 		}
+		
+		if (bSupportLabels)
+		{
+			for (const auto& [AttrName, TriangleLabelAttr]: Attributes->GetTriangleLabelAttributes())
+			{
+				GroupLayersList.Add(AttrName.ToString());
+			}
+		}
 	}
 
 	if (GroupLayersList.Contains(ActiveGroupLayer.ToString()) == false)		// discard restored value if it doesn't apply
diff --git a/Engine/Plugins/Runtime/MeshModelingToolset/Source/ModelingComponents/Public/Changes/MeshTriangleLabelChange.h b/Engine/Plugins/Runtime/MeshModelingToolset/Source/ModelingComponents/Public/Changes/MeshTriangleLabelChange.h
new file mode 100644
--- /dev/null
+++ b/Engine/Plugins/Runtime/MeshModelingToolset/Source/ModelingComponents/Public/Changes/MeshTriangleLabelChange.h
@@ -0,0 +1,93 @@
+// Copyright Epic Games, Inc. All Rights Reserved.
+
+#pragma once
+
+#include "Templates/PimplPtr.h"
+#include "Changes/MeshChange.h"
+#include "GeometryBase.h"
+#include "Polygroups/GroupSetAdapter.h"
+
+#define UE_API MODELINGCOMPONENTS_API
+
+PREDECLARE_USE_GEOMETRY_CLASS(FDynamicMesh3);
+
+/**
+ * FDynamicMeshTriangleLabelEdit stores a modification of labels on a set of triangles
+ */
+class FDynamicMeshTriangleLabelEdit
+{
+public:
+	/** Layer name as registered in FDynamicMeshAttributeSet::TriangleLabelAttributes */
+	FName LayerName;
+	/** IDs of triangles that are modified */
+	TArray<int32> Triangles;
+	/** Old labels for each triangle */
+	TArray<FName> OldLabels;
+	/** New labels for each triangle */
+	TArray<FName> NewLabels;
+	/** Applies the changes to Mesh looking for a FDynamicMeshTriangleLabelAttribute named LayerName */
+	UE_API void ApplyToMesh(FDynamicMesh3* Mesh, bool bRevert);
+	
+protected:
+	/** Ensures that TriangleIndex is a valid index for OldLabels or NewLabels */
+	bool IsTriangleIndexValid(const int32 TriangleIndex, const bool bRevert) const;
+};
+
+
+
+/**
+ * FDynamicMeshTriangleLabelEditBuilder builds up a FDynamicMeshTriangleLabelEdit incrementally storing old and new triangle labels
+ */
+class FDynamicMeshTriangleLabelEditBuilder
+{
+public:
+	
+	/** Stores a copy of the incoming adapter to get data from, as well as it's name, and starts a new edit */
+	UE_API FDynamicMeshTriangleLabelEditBuilder(UE::Geometry::FTriangleLabelAdapter* InAdapter);
+
+	/** Extracts the label edits */
+	TUniquePtr<FDynamicMeshTriangleLabelEdit> ExtractResult()
+	{
+		return MoveTemp(Edit);
+	}
+
+	/** Stores old and new label for this triangle using the adapter */
+	UE_API void SaveTriangle(int32 TriangleID);
+	
+	/** Stores old and new labels for the triangles using the adapter */
+	template<typename EnumerableType>
+	void SaveTriangles(EnumerableType Enumerable)
+	{
+		for (int32 tid : Enumerable)
+		{
+			SaveTriangle(tid);
+		}
+	}
+	
+	/** Stores old and new label for this triangle */
+	UE_API void SaveTriangle(int32 TriangleID, FName OldLabel, FName NewLabel);
+
+protected:
+	/** Adapter to get labels from */
+	TPimplPtr<UE::Geometry::FTriangleLabelAdapter> Adapter;
+	/** Edit storing old and new triangle labels */
+	TUniquePtr<FDynamicMeshTriangleLabelEdit> Edit;
+	/** Data table storing a triangle id - triangle edit mapping */
+	TMap<int32, int32> SavedIndexMap;
+};
+
+/**
+ * FMeshTriangleLabelChange stores some label changes a set of triangles, as a FDynamicMeshTriangleLabelEdit.
+ */
+class FMeshTriangleLabelChange : public FMeshChange
+{
+public:
+	UE_API FMeshTriangleLabelChange(TUniquePtr<FDynamicMeshTriangleLabelEdit>&& LabelEditIn);
+
+	/** Applies the underlying triangle label edit on the mesh */
+	UE_API virtual void ApplyChangeToMesh(FDynamicMesh3* Mesh, bool bRevert) const override;
+	
+	TUniquePtr<FDynamicMeshTriangleLabelEdit> LabelEdit;
+};
+
+#undef UE_API
diff --git a/Engine/Plugins/Runtime/MeshModelingToolset/Source/ModelingComponents/Public/PropertySets/PolygroupLayersProperties.h b/Engine/Plugins/Runtime/MeshModelingToolset/Source/ModelingComponents/Public/PropertySets/PolygroupLayersProperties.h
--- a/Engine/Plugins/Runtime/MeshModelingToolset/Source/ModelingComponents/Public/PropertySets/PolygroupLayersProperties.h
+++ b/Engine/Plugins/Runtime/MeshModelingToolset/Source/ModelingComponents/Public/PropertySets/PolygroupLayersProperties.h
@@ -43,6 +43,10 @@ public:
 	UE_API void SetSelectedFromPolygroupIndex(int32 Index);
 
 	UE_API UE::Geometry::FPolygroupLayer GetSelectedLayer(const FDynamicMesh3& FromMesh);
+	
+	// whether triangle label attributes should be added to the group layers list
+	UPROPERTY(meta = (TransientToolProperty))
+	bool bSupportLabels = false;
 };
 
 #undef UE_API
