fill colors from hole boundaries in most hole fill operations (except smooth hole fill)

#Synced-CL 51695643

diff --git a/Engine/Plugins/Runtime/GeometryProcessing/Source/DynamicMesh/Private/Operations/MinimalHoleFiller.cpp b/Engine/Plugins/Runtime/GeometryProcessing/Source/DynamicMesh/Private/Operations/MinimalHoleFiller.cpp
--- a/Engine/Plugins/Runtime/GeometryProcessing/Source/DynamicMesh/Private/Operations/MinimalHoleFiller.cpp
+++ b/Engine/Plugins/Runtime/GeometryProcessing/Source/DynamicMesh/Private/Operations/MinimalHoleFiller.cpp
@@ -2,6 +2,7 @@
 
 #include "Operations/MinimalHoleFiller.h"
 #include "Operations/SimpleHoleFiller.h"
+#include "Operations/HoleFillUtil.h"
 #include "MeshQueries.h"
 
 using namespace UE::Geometry;
@@ -601,6 +602,7 @@ bool FMinimalHoleFiller::Fill(int32 GroupID)
 
 	// do a simple fill
 	FSimpleHoleFiller Simplefill(Mesh, FillLoop);
+	Simplefill.bAutoFillPrimaryColors = false;
 
 	bool bOK = Simplefill.Fill(GroupID);
 	if (bOK == false)
@@ -608,9 +610,18 @@ bool FMinimalHoleFiller::Fill(int32 GroupID)
 		return false;
 	}
 
+	auto FillColors = [this]()
+	{
+		if (bAutoFillPrimaryColors)
+		{
+			HoleFillUtil::FillColorOverlay(*Mesh, FillLoop.Edges, NewTriangles);
+		}
+	};
+
 	if (FillLoop.Vertices.Num() <= 3) 
 	{
-		NewTriangles = Simplefill.NewTriangles;
+		NewTriangles = MoveTemp(Simplefill.NewTriangles);
+		FillColors();
 		return true;
 	}
 
@@ -618,7 +629,8 @@ bool FMinimalHoleFiller::Fill(int32 GroupID)
 	double LoopTotalLength = TMeshQueries<FDynamicMesh3>::TotalEdgeLength(*Mesh, FillLoop.Edges);
 	if (LoopTotalLength < 100.0f*FMathf::ZeroTolerance)
 	{
-		NewTriangles = Simplefill.NewTriangles;
+		NewTriangles = MoveTemp(Simplefill.NewTriangles);
+		FillColors();
 		return true;
 	}
 
@@ -751,5 +763,6 @@ bool FMinimalHoleFiller::Fill(int32 GroupID)
 	RegionOp->BackPropropagate(true);
 	NewTriangles = RegionOp->CurrentBaseTriangles();
 
+	FillColors();
 	return true;
 }
diff --git a/Engine/Plugins/Runtime/GeometryProcessing/Source/DynamicMesh/Private/Operations/PlanarHoleFiller.cpp b/Engine/Plugins/Runtime/GeometryProcessing/Source/DynamicMesh/Private/Operations/PlanarHoleFiller.cpp
--- a/Engine/Plugins/Runtime/GeometryProcessing/Source/DynamicMesh/Private/Operations/PlanarHoleFiller.cpp
+++ b/Engine/Plugins/Runtime/GeometryProcessing/Source/DynamicMesh/Private/Operations/PlanarHoleFiller.cpp
@@ -2,6 +2,7 @@
 
 #include "Operations/PlanarHoleFiller.h"
 #include "Curve/PlanarComplex.h"
+#include "Operations/HoleFillUtil.h"
 
 using namespace UE::Geometry;
 
@@ -57,6 +58,11 @@ bool FPlanarHoleFiller::Fill(int GroupID)
 			}
 		}
 	}
+
+	if (bAutoFillPrimaryColors)
+	{
+		HoleFillUtil::FillColorOverlay(*Mesh, *VertexLoops, NewTriangles);
+	}
 	
 	return bAddedAll;
 }
diff --git a/Engine/Plugins/Runtime/GeometryProcessing/Source/DynamicMesh/Private/Operations/SimpleHoleFiller.cpp b/Engine/Plugins/Runtime/GeometryProcessing/Source/DynamicMesh/Private/Operations/SimpleHoleFiller.cpp
--- a/Engine/Plugins/Runtime/GeometryProcessing/Source/DynamicMesh/Private/Operations/SimpleHoleFiller.cpp
+++ b/Engine/Plugins/Runtime/GeometryProcessing/Source/DynamicMesh/Private/Operations/SimpleHoleFiller.cpp
@@ -3,6 +3,7 @@
 #include "Operations/SimpleHoleFiller.h"
 #include "DynamicMeshEditor.h"
 #include "CompGeom/PolygonTriangulation.h"
+#include "Operations/HoleFillUtil.h"
 
 using namespace UE::Geometry;
 
@@ -30,6 +31,10 @@ bool FSimpleHoleFiller::Fill(int GroupID)
 		}
 		NewTriangles = { NewTID };
 		NewVertex = FDynamicMesh3::InvalidID;
+		if (bAutoFillPrimaryColors)
+		{
+			FillColors();
+		}
 		return true;
 	}
 
@@ -44,6 +49,10 @@ bool FSimpleHoleFiller::Fill(int GroupID)
 	{
 		bOK = Fill_Fan(GroupID);
 	}
+	if (bOK && bAutoFillPrimaryColors)
+	{
+		FillColors();
+	}
 
 	return bOK;
 }
@@ -108,6 +117,13 @@ bool FSimpleHoleFiller::Fill_EarClip(int GroupID)
 	return NewTriangles.Num() == Triangles.Num();
 }
 
+
+void FSimpleHoleFiller::FillColors()
+{
+	HoleFillUtil::FillColorOverlay(*Mesh, Loop.Edges, NewTriangles, NewVertex);
+}
+
+
 bool FSimpleHoleFiller::UpdateAttributes(TArray<FMeshRegionBoundaryLoops::VidOverlayMap<FVector2f>>& VidUVMaps)
 {
 	if (!Mesh->HasAttributes() || NewTriangles.Num() == 0)
diff --git a/Engine/Plugins/Runtime/GeometryProcessing/Source/DynamicMesh/Public/Operations/HoleFillUtil.h b/Engine/Plugins/Runtime/GeometryProcessing/Source/DynamicMesh/Public/Operations/HoleFillUtil.h
new file mode 100644
--- /dev/null
+++ b/Engine/Plugins/Runtime/GeometryProcessing/Source/DynamicMesh/Public/Operations/HoleFillUtil.h
@@ -0,0 +1,155 @@
+// Copyright Epic Games, Inc. All Rights Reserved.
+
+#pragma once
+
+#include "Containers/Array.h"
+#include "Containers/ArrayView.h"
+#include "Containers/Map.h"
+#include "DynamicMesh/DynamicMesh3.h"
+#include "DynamicMesh/DynamicMeshAttributeSet.h"
+#include "EdgeLoop.h"
+
+namespace UE
+{
+namespace Geometry
+{
+
+
+namespace HoleFillUtil
+{
+	/**
+	 * Fill in overlay elements on new triangles, by copying values from provided edges.
+	 * Used to fill in overlay values after a hole fill.
+	 * Only supports fills where at most one new vertex is added inside the hole (e.g., does not support smooth hole fill)
+	 * 
+	 * @param Mesh Mesh holding the Overlay
+	 * @param Overlay The dynamic mesh overlay to update
+	 * @param PerEdgeIDFn Enumerator function, which will call a passed-in method per Edge ID on the hole border
+	 * @param NewTriangles Triangles that fill the hole
+	 * @param NewVertex If not invalid, the at-most one vertex created when filling the hole.
+	 */
+	template<typename OverlayType, typename EdgeEnumerateFn>
+	void FillOverlayElements(const FDynamicMesh3& Mesh, OverlayType* Overlay, EdgeEnumerateFn PerEdgeIDFn, TConstArrayView<int32> NewTriangles, int32 NewVertex = IndexConstants::InvalidID)
+	{
+		TMap<int32, int32> VIDToElementID;
+		bool bHasNewVertex = NewVertex != IndexConstants::InvalidID;
+
+
+		PerEdgeIDFn([&Mesh, &Overlay, bHasNewVertex, &VIDToElementID](int32 EdgeID)
+		{
+			FDynamicMesh3::FEdge Edge = Mesh.GetEdge(EdgeID);
+			int32 SourceTID = Overlay->IsSetTriangle(Edge.Tri.A) ? Edge.Tri.A : ((Mesh.IsTriangle(Edge.Tri.B) && Overlay->IsSetTriangle(Edge.Tri.B)) ? Edge.Tri.B : IndexConstants::InvalidID);
+			if (SourceTID != IndexConstants::InvalidID)
+			{
+				FIndex3i SetTri = Mesh.GetTriangle(SourceTID);
+				FIndex3i SetElTri = Overlay->GetTriangle(SourceTID);
+				for (int32 SubIdx = 0; SubIdx < 2; ++SubIdx)
+				{
+					int32 SourceElementID = SetElTri[SetTri.IndexOf(Edge.Vert[SubIdx])];
+					VIDToElementID.Add(Edge.Vert[SubIdx], SourceElementID);
+				}
+			}
+		});
+		typename OverlayType::VectorType AvgValue(0);
+		if (bHasNewVertex && !VIDToElementID.IsEmpty())
+		{
+			for (TPair<int32, int32> VIDElemIDPair : VIDToElementID)
+			{
+				AvgValue += Overlay->GetElement(VIDElemIDPair.Value);
+			}
+			AvgValue /= (float)VIDToElementID.Num();
+		}
+		for (int32 NewTID : NewTriangles)
+		{
+			FIndex3i NewTri = Mesh.GetTriangle(NewTID);
+			FIndex3i NewElTri;
+			bool bAllFound = true;
+			int32 AddElementSubIdx = IndexConstants::InvalidID;
+			for (int32 SubIdx = 0; SubIdx < 3; ++SubIdx)
+			{
+				int32 VID = NewTri[SubIdx];
+				if (int32* FoundElementID = VIDToElementID.Find(VID))
+				{
+					NewElTri[SubIdx] = *FoundElementID;
+				}
+				else if (VID == NewVertex)
+				{
+					// we'll add this element below, if the triangle as a whole will be set
+					AddElementSubIdx = SubIdx;
+				}
+				else
+				{
+					bAllFound = false;
+					break;
+				}
+			}
+			if (bAllFound)
+			{
+				if (AddElementSubIdx != IndexConstants::InvalidID)
+				{
+					int32 NewElementID = Overlay->AppendElement(AvgValue);
+					VIDToElementID.Add(NewTri[AddElementSubIdx], NewElementID);
+					NewElTri[AddElementSubIdx] = NewElementID;
+				}
+				Overlay->SetTriangle(NewTID, NewElTri);
+			}
+		}
+	}
+
+	/**
+	 * Fill in color overlay elements on new triangles, by copying values from provided edges. If mesh doesn't have a color overlay, does nothing.
+	 * Used to fill in overlay values after a hole fill.
+	 * Only supports fills where at most one new vertex is added inside the hole (e.g., does not support smooth hole fill)
+	 *
+	 * @param Mesh The mesh to update
+	 * @param HoleEdges The edges at the boundary of the hole
+	 * @param NewTriangles Triangles that fill the hole
+	 * @param NewVertex If not invalid, the at-most one vertex created when filling the hole.
+	 */
+	inline void FillColorOverlay(FDynamicMesh3& Mesh, TConstArrayView<int32> HoleEdges, TConstArrayView<int32> NewTriangles, int32 NewVertex = IndexConstants::InvalidID)
+	{
+		if (Mesh.HasAttributes() && Mesh.Attributes()->HasPrimaryColors())
+		{
+			FillOverlayElements<FDynamicMeshColorOverlay>(Mesh, Mesh.Attributes()->PrimaryColors(), 
+				[HoleEdges](TFunctionRef<void(int32)> ProcessEID) {for (int32 EID : HoleEdges) { ProcessEID(EID); }},
+				NewTriangles, NewVertex);
+		}
+	}
+
+	/**
+	 * Fill in color overlay elements on new triangles, by copying values from edges along the provided vertex loops. If mesh doesn't have a color overlay, does nothing.
+	 * Used to fill in overlay values after a hole fill.
+	 * Only supports fills where at most one new vertex is added inside the hole (e.g., does not support smooth hole fill)
+	 *
+	 * @param Mesh The mesh to update
+	 * @param HoleVertexLoops Vertex loops at the boundary of the hole
+	 * @param NewTriangles Triangles that fill the hole
+	 * @param NewVertex If not invalid, the at-most one vertex created when filling the hole.
+	 */
+	inline void FillColorOverlay(FDynamicMesh3& Mesh, const TArray<TArray<int32>>& HoleVertexLoops, TConstArrayView<int32> NewTriangles, int32 NewVertex = IndexConstants::InvalidID)
+	{
+		if (Mesh.HasAttributes() && Mesh.Attributes()->HasPrimaryColors())
+		{
+			FillOverlayElements<FDynamicMeshColorOverlay>(Mesh, Mesh.Attributes()->PrimaryColors(), 
+				[&Mesh, &HoleVertexLoops](TFunctionRef<void(int32)> ProcessEID)
+				{
+					for (const TArray<int32>& Loop : HoleVertexLoops)
+					{
+						for (int32 Idx = 0, PrevIdx = Loop.Num() - 1; Idx < Loop.Num(); PrevIdx = Idx++)
+						{
+							int32 EID = Mesh.FindEdge(Loop[Idx], Loop[PrevIdx]);
+							if (EID != IndexConstants::InvalidID)
+							{
+								ProcessEID(EID);
+							}
+						}
+					}
+				}, NewTriangles, NewVertex
+			);
+		}
+	}
+}
+
+
+} // end namespace UE::Geometry
+} // end namespace UE
\ No newline at end of file
diff --git a/Engine/Plugins/Runtime/GeometryProcessing/Source/DynamicMesh/Public/Operations/MinimalHoleFiller.h b/Engine/Plugins/Runtime/GeometryProcessing/Source/DynamicMesh/Public/Operations/MinimalHoleFiller.h
--- a/Engine/Plugins/Runtime/GeometryProcessing/Source/DynamicMesh/Public/Operations/MinimalHoleFiller.h
+++ b/Engine/Plugins/Runtime/GeometryProcessing/Source/DynamicMesh/Public/Operations/MinimalHoleFiller.h
@@ -43,6 +43,8 @@ public:
 	bool bOptimizeDevelopability = true;
 	bool bOptimizeTriangles = true;
 	double DevelopabilityTolerance = 0.0001;
+	// whether to automatically fill in PrimaryColors on new hole fill triangles
+	bool bAutoFillPrimaryColors = true;
 
 private:
 
diff --git a/Engine/Plugins/Runtime/GeometryProcessing/Source/DynamicMesh/Public/Operations/PlanarHoleFiller.h b/Engine/Plugins/Runtime/GeometryProcessing/Source/DynamicMesh/Public/Operations/PlanarHoleFiller.h
--- a/Engine/Plugins/Runtime/GeometryProcessing/Source/DynamicMesh/Public/Operations/PlanarHoleFiller.h
+++ b/Engine/Plugins/Runtime/GeometryProcessing/Source/DynamicMesh/Public/Operations/PlanarHoleFiller.h
@@ -40,6 +40,9 @@ public:
 	FVector3d PlaneOrigin;
 	FVector3d PlaneNormal;
 
+	// whether to automatically fill in PrimaryColors on new hole fill triangles
+	bool bAutoFillPrimaryColors = true;
+
 public:
 	/**
 	 *  Construct planar hole filler (projects the hole edges to a plane and fills the hole w/ a general 2D polygon triangulation)
diff --git a/Engine/Plugins/Runtime/GeometryProcessing/Source/DynamicMesh/Public/Operations/SimpleHoleFiller.h b/Engine/Plugins/Runtime/GeometryProcessing/Source/DynamicMesh/Public/Operations/SimpleHoleFiller.h
--- a/Engine/Plugins/Runtime/GeometryProcessing/Source/DynamicMesh/Public/Operations/SimpleHoleFiller.h
+++ b/Engine/Plugins/Runtime/GeometryProcessing/Source/DynamicMesh/Public/Operations/SimpleHoleFiller.h
@@ -39,6 +39,9 @@ public:
 	FEdgeLoop Loop;
 	EFillType FillType = EFillType::TriangleFan;
 
+	// whether to automatically fill in PrimaryColors on new hole fill triangles
+	bool bAutoFillPrimaryColors = true;
+
 	//
 	// Outputs
 	//
@@ -103,6 +106,10 @@ public:
 protected:
 	UE_API bool Fill_Fan(int32 NewGroupID);
 	UE_API bool Fill_EarClip(int32 NewGroupID);
+
+private:
+	// Fill colors on triangles added by the fill, by copying boundary colors
+	UE_API void FillColors();
 };
 
 
diff --git a/Engine/Plugins/Runtime/MeshModelingToolset/Source/ModelingOperators/Private/CleaningOps/HoleFillOp.cpp b/Engine/Plugins/Runtime/MeshModelingToolset/Source/ModelingOperators/Private/CleaningOps/HoleFillOp.cpp
--- a/Engine/Plugins/Runtime/MeshModelingToolset/Source/ModelingOperators/Private/CleaningOps/HoleFillOp.cpp
+++ b/Engine/Plugins/Runtime/MeshModelingToolset/Source/ModelingOperators/Private/CleaningOps/HoleFillOp.cpp
@@ -5,6 +5,7 @@
 #include "Operations/PlanarHoleFiller.h"
 #include "Operations/MinimalHoleFiller.h"
 #include "Operations/SmoothHoleFiller.h"
+#include "Operations/HoleFillUtil.h"
 #include "ConstrainedDelaunay2.h"
 #include "CompGeom/PolygonTriangulation.h"
 #include "Selections/MeshConnectedComponents.h"
@@ -207,64 +208,74 @@ bool FHoleFillOp::FillSingleTriangleHole(const FEdgeLoop& Loop, int32& NewGroupI
 		return false;
 	}
 
-	if (ResultMesh->HasAttributes() && ResultMesh->Attributes()->PrimaryUV())
+	// Pick the overlay elements to use based on neighbors, to minimize added seams. Returns false when neighbors don't share elements along the hole.
+	auto GetOverlayElementsFromNeighbors = [this, &Vertices, &ExistingTriangles, &ExistingEdges]<typename OverlayType>(OverlayType* Overlay, FIndex3i& NewTriangleElements) -> bool
 	{
-		FDynamicMeshUVOverlay* UVOverlay = ResultMesh->Attributes()->PrimaryUV();
-
-		// check we have a hole in the UV mesh as well (note that boundary edges of unset UV triangles aren't seen 
-		// as seams by IsSeamEdge)
-		checkSlow(UVOverlay->IsSeamEdge(ExistingEdges[0]) || !UVOverlay->IsSetTriangle(ExistingTriangles[0]));
-		checkSlow(UVOverlay->IsSeamEdge(ExistingEdges[1]) || !UVOverlay->IsSetTriangle(ExistingTriangles[1]));
-		checkSlow(UVOverlay->IsSeamEdge(ExistingEdges[2]) || !UVOverlay->IsSetTriangle(ExistingTriangles[2]));
-
-		// We're going to try to assign the UVs such that the number of UV seam edges is minimized. If at least 
-		// two of our neighboring triangles have a shared UV element at one of our vertices, then using the elements
+		// We're going to try to assign the overlay elements such that the number of overlay seam edges is minimized. If at least 
+		// two of our neighboring triangles have a shared overlay element at one of our vertices, then using the elements
 		// of these two triangles will decrease the number of seam edges because we end up removing two seams and
 		// adding at most one (or none if this was an enclosed hole in a single island).
-	
+
 		// For each of our vertices, get the elements associated with that vertex by the two neighbor triangles. The 
 		// pair we store is (element from previous triangle, element from next triangle).
-		FIndex2i VertUVElements[3];
+		FIndex2i VertElements[3];
 		for (int NeighborIndex = 0; NeighborIndex < 3; ++NeighborIndex)
 		{
 			FIndex3i NeighborTriVids = ResultMesh->GetTriangle(ExistingTriangles[NeighborIndex]);
-			FIndex3i NeighborTriElements = UVOverlay->GetTriangle(ExistingTriangles[NeighborIndex]);
+			FIndex3i NeighborTriElements = Overlay->GetTriangle(ExistingTriangles[NeighborIndex]);
 
 			// Find our two vids in TriVids, get the corresponding elements in TriElements
 			int VertIndex1 = NeighborIndex;
 			int VertIndex2 = (VertIndex1 + 1) % 3;
 
-			VertUVElements[VertIndex1].B = NeighborTriElements[IndexUtil::FindTriIndex(Vertices[VertIndex1], NeighborTriVids)];
-			VertUVElements[VertIndex2].A = NeighborTriElements[IndexUtil::FindTriIndex(Vertices[VertIndex2], NeighborTriVids)];
+			VertElements[VertIndex1].B = NeighborTriElements[IndexUtil::FindTriIndex(Vertices[VertIndex1], NeighborTriVids)];
+			VertElements[VertIndex2].A = NeighborTriElements[IndexUtil::FindTriIndex(Vertices[VertIndex2], NeighborTriVids)];
 		}
 
-		FIndex3i NewTriangleElements;
-
 		/*
 		 * Imagine going counterclockwise from the bottom left in the triangle below. A and B correspond
 		 * to the element values in VertUVElements. When the two are equal the vertex has a shared element,
 		 * and the adjacent elements need to be on the adjoining sides.
-		 * 
+		 *
 		 *     B.A
 		 *     / \
 		 *   A.---.B
 		 *    B   A
 		 */
-		if (VertUVElements[0].A == VertUVElements[0].B)
+		if (VertElements[0].A == VertElements[0].B)
 		{
-			NewTriangleElements = FIndex3i(VertUVElements[0].A, VertUVElements[1].A, VertUVElements[2].B);
+			NewTriangleElements = FIndex3i(VertElements[0].A, VertElements[1].A, VertElements[2].B);
 		}
-		else if (VertUVElements[1].A == VertUVElements[1].B)
+		else if (VertElements[1].A == VertElements[1].B)
 		{
-			NewTriangleElements = FIndex3i(VertUVElements[0].B, VertUVElements[1].A, VertUVElements[2].A);
+			NewTriangleElements = FIndex3i(VertElements[0].B, VertElements[1].A, VertElements[2].A);
 		}
-		else if (VertUVElements[2].A == VertUVElements[2].B)
+		else if (VertElements[2].A == VertElements[2].B)
 		{
-			NewTriangleElements = FIndex3i(VertUVElements[0].A, VertUVElements[1].B, VertUVElements[2].A);
+			NewTriangleElements = FIndex3i(VertElements[0].A, VertElements[1].B, VertElements[2].A);
 		}
 		else
+		{
+			return false; // failed to find a clean match, fall back to overlay-type-specific logic
+		}
+		return true;
+	};
+
+	if (ResultMesh->HasAttributes() && ResultMesh->Attributes()->PrimaryUV())
+	{
+		FDynamicMeshUVOverlay* UVOverlay = ResultMesh->Attributes()->PrimaryUV();
+
+		// check we have a hole in the UV mesh as well (note that boundary edges of unset UV triangles aren't seen 
+		// as seams by IsSeamEdge)
+		checkSlow(UVOverlay->IsSeamEdge(ExistingEdges[0]) || !UVOverlay->IsSetTriangle(ExistingTriangles[0]));
+		checkSlow(UVOverlay->IsSeamEdge(ExistingEdges[1]) || !UVOverlay->IsSetTriangle(ExistingTriangles[1]));
+		checkSlow(UVOverlay->IsSeamEdge(ExistingEdges[2]) || !UVOverlay->IsSetTriangle(ExistingTriangles[2]));
+		FIndex3i NewTriangleElements = FIndex3i::Invalid();
+
+		if (!GetOverlayElementsFromNeighbors(UVOverlay, NewTriangleElements))
 		{
 			// Since we didn't have a matching element ID above, at least one of ExistingTriangles[0] or [1] must have set UVs --
+			//  i.e., if they were both unset, then their element IDs would all be -1 and one of the tests for matching element IDs would pass
 			// Pick an arbitrary set UV edge and add a new third UV element
 			int32 NeighborTidToUse = UVOverlay->IsSetTriangle(ExistingTriangles[0]) ? ExistingTriangles[0] : ExistingTriangles[1];
 			SetUVTriangleFromAdjacentTriangle(*UVOverlay, NeighborTidToUse, NewTriangleID, NewTriangleElements);
@@ -284,6 +295,29 @@ bool FHoleFillOp::FillSingleTriangleHole(const FEdgeLoop& Loop, int32& NewGroupI
 		}
 	}
 
+	if (ResultMesh->HasAttributes() && ResultMesh->Attributes()->HasPrimaryColors())
+	{
+		FDynamicMeshColorOverlay* ColorOverlay = ResultMesh->Attributes()->PrimaryColors();
+		FIndex3i NewTriangleElements = FIndex3i::Invalid();
+
+		// try to fill the color overlay seamlessly from local info
+		if (GetOverlayElementsFromNeighbors(ColorOverlay, NewTriangleElements))
+		{
+			if (NewTriangleElements.A != IndexConstants::InvalidID)
+			{
+				EMeshResult Result = ColorOverlay->SetTriangle(NewTriangleID, NewTriangleElements);
+				ensure(Result == EMeshResult::Ok); // should always succeed, but if it doesn't we will just have some unset colors
+			}
+		}
+		else
+		{
+			// fall back to a standard color overlay fill method, which should at least try to avoid a seam along one edge
+			// (since it stomps IDs to match for each it finds with valid overlay elements, so should avoid seams at least on the last edge it processes)
+			TStaticArray<int32, 3> ExistingEdgesArr{ ExistingEdges.A, ExistingEdges.B, ExistingEdges.C };
+			HoleFillUtil::FillColorOverlay(*ResultMesh, ExistingEdgesArr, TArrayView<int32>(&NewTriangleID, 1));
+		}
+	}
+
 
 	// And normals
 	if (ResultMesh->HasAttributes())
diff --git a/Engine/Source/Runtime/GeometryCore/Private/DynamicMesh/DynamicMeshOverlay.cpp b/Engine/Source/Runtime/GeometryCore/Private/DynamicMesh/DynamicMeshOverlay.cpp
--- a/Engine/Source/Runtime/GeometryCore/Private/DynamicMesh/DynamicMeshOverlay.cpp
+++ b/Engine/Source/Runtime/GeometryCore/Private/DynamicMesh/DynamicMeshOverlay.cpp
@@ -1048,6 +1048,30 @@ void TDynamicMeshOverlay<RealType, ElementSize>::GetVertexElements(int vid, TArr
 	}
 }
 
+template<typename RealType, int ElementSize>
+bool TDynamicMeshOverlay<RealType, ElementSize>::FindAnyElementIDAtVertex(int32 VertexID, int32& OutElementID) const
+{
+	OutElementID = IndexConstants::InvalidID;
+	for (int TID : ParentMesh->VtxTrianglesItr(VertexID))
+	{
+		if (!IsSetTriangle(TID))
+		{
+			continue;
+		}
+		FIndex3i Triangle = GetTriangle(TID);
+		for (int SubIdx = 0; SubIdx < 3; ++SubIdx)
+		{
+			if (ParentVertices[Triangle[SubIdx]] == VertexID)
+			{
+				OutElementID = Triangle[SubIdx];
+				return true;
+			}
+		}
+	}
+	return false;
+}
+
+
 
 
 template<typename RealType, int ElementSize>
diff --git a/Engine/Source/Runtime/GeometryCore/Public/DynamicMesh/DynamicMeshOverlay.h b/Engine/Source/Runtime/GeometryCore/Public/DynamicMesh/DynamicMeshOverlay.h
--- a/Engine/Source/Runtime/GeometryCore/Public/DynamicMesh/DynamicMeshOverlay.h
+++ b/Engine/Source/Runtime/GeometryCore/Public/DynamicMesh/DynamicMeshOverlay.h
@@ -606,6 +606,15 @@ public:
 	 */
 	int GetElementIDAtVertex(int TriangleID, int VertexID) const;
 
+	/**
+	 * Find an element ID associated with the given VertexID
+	 * 
+	 * @param VertexID The vertex ID to search
+	 * @param OutElementID Will be set to the found element ID, or FDynamicMesh3::InvalidID if none found
+	 * @return true if an element ID was found, false otherwise
+	 */
+	bool FindAnyElementIDAtVertex(int32 VertexID, int32& OutElementID) const;
+
 	/** @return true if overlay has any interior seam edges. This requires an O(N) search unless it early-outs. */
 	bool HasInteriorSeamEdges() const;
 
@@ -709,11 +718,12 @@ protected:
  * a specific N-element Vector type to the template, and adds accessor functions
  * that convert between that N-element vector type and the N-element arrays used by TDynamicMeshOverlay.
  */
-template<typename RealType, int ElementSize, typename VectorType>
+template<typename RealType, int ElementSize, typename InVectorType>
 class TDynamicMeshVectorOverlay : public TDynamicMeshOverlay<RealType, ElementSize>
 {
 public:
 	using BaseType = TDynamicMeshOverlay<RealType, ElementSize>;
+	using VectorType = InVectorType;
 
 	TDynamicMeshVectorOverlay()
 		: TDynamicMeshOverlay<RealType, ElementSize>()
