add cvar to swap shift/ctrl in modeling sculpt tools

#Synced-CL 51714673

diff --git a/Engine/Plugins/Runtime/MeshModelingToolset/Source/MeshModelingTools/Private/MeshVertexSculptTool.cpp b/Engine/Plugins/Runtime/MeshModelingToolset/Source/MeshModelingTools/Private/MeshVertexSculptTool.cpp
--- a/Engine/Plugins/Runtime/MeshModelingToolset/Source/MeshModelingTools/Private/MeshVertexSculptTool.cpp
+++ b/Engine/Plugins/Runtime/MeshModelingToolset/Source/MeshModelingTools/Private/MeshVertexSculptTool.cpp
@@ -74,6 +74,11 @@ namespace MeshVertexSculptToolLocals
 	static FAutoConsoleVariable DisableOctreeUpdates(TEXT("modeling.Sculpting.DisableOctreeUpdates"), false, TEXT("Disable updating the octree during sculpting"));
 	static FAutoConsoleVariable OctreeRootCellSizeOverride(TEXT("modeling.Sculpting.OctreeRootCellSizeOverride"), 0.0f, TEXT("If greater than 0.0, set octree root cell size to this value instead of auto-computing."));
 	static FAutoConsoleVariable OctreeTreeDepthOverride(TEXT("modeling.Sculpting.OctreeTreeDepthOverride"), 0, TEXT("If greater than 0, set octree tree depth to this value instead of auto-computing."));
+
+	static TAutoConsoleVariable<bool> CVarSwapSmoothInvertModifiers(
+		TEXT("modeling.Sculpting.SwapSmoothInvertModifiers"),
+		false,
+		TEXT("When true, swaps the Shift/Ctrl modifier key roles in the vertex sculpt tool"));
 }
 
 /*
@@ -115,6 +120,15 @@ public:
  * Tool
  */
 
+void UMeshVertexSculptTool::SaveActiveStrokeModifiers()
+{
+	Super::SaveActiveStrokeModifiers();
+	if (MeshVertexSculptToolLocals::CVarSwapSmoothInvertModifiers.GetValueOnGameThread())
+	{
+		Swap(bSmoothing, bInvert);
+	}
+}
+
 void UMeshVertexSculptTool::Setup()
 {
 	UMeshSculptToolBase::Setup();
diff --git a/Engine/Plugins/Runtime/MeshModelingToolset/Source/MeshModelingTools/Public/MeshVertexSculptTool.h b/Engine/Plugins/Runtime/MeshModelingToolset/Source/MeshModelingTools/Public/MeshVertexSculptTool.h
--- a/Engine/Plugins/Runtime/MeshModelingToolset/Source/MeshModelingTools/Public/MeshVertexSculptTool.h
+++ b/Engine/Plugins/Runtime/MeshModelingToolset/Source/MeshModelingTools/Public/MeshVertexSculptTool.h
@@ -324,6 +324,7 @@ protected:
 	UE_API virtual FString GetPropertyCacheIdentifier() const;
 
 	// UMeshSculptToolBase API
+	UE_API virtual void SaveActiveStrokeModifiers() override;
 	UE_API virtual void InitializeIndicator() override;
 	UE_API virtual UPreviewMesh* MakeBrushIndicatorMesh(UObject* Parent, UWorld* World) override;
 
