fixed crash in simplify tool on attempting to simplify a dynamic mesh with the UE standard simplifier
by converting dynamic mesh to mesh description when there is no original mesh description source + validating the mesh description has static mesh attributes before attempting to run the simplifier on it
#Synced-CL 51735880

diff --git a/Engine/Plugins/Runtime/MeshModelingToolset/Source/MeshModelingToolsEditorOnly/Private/SimplifyMeshTool.cpp b/Engine/Plugins/Runtime/MeshModelingToolset/Source/MeshModelingToolsEditorOnly/Private/SimplifyMeshTool.cpp
--- a/Engine/Plugins/Runtime/MeshModelingToolset/Source/MeshModelingToolsEditorOnly/Private/SimplifyMeshTool.cpp
+++ b/Engine/Plugins/Runtime/MeshModelingToolset/Source/MeshModelingToolsEditorOnly/Private/SimplifyMeshTool.cpp
@@ -7,6 +7,9 @@
 #include "Drawing/MeshElementsVisualizer.h"
 #include "ToolBuilderUtil.h"
 #include "MeshDescriptionToDynamicMesh.h"
+#include "DynamicMeshToMeshDescription.h"
+#include "StaticMeshAttributes.h"
+#include "TargetInterfaces/MeshDescriptionProvider.h"
 #include "ModelingToolTargetUtil.h"
 #include "ToolSetupUtil.h"
 #include "DynamicMesh/DynamicMesh3.h"
@@ -220,7 +223,22 @@ TUniquePtr<FDynamicMeshOperator> USimplifyMeshTool::MakeNewOperator()
 
 		FGetMeshParameters GetMeshParams;
 		GetMeshParams.bWantMeshTangents = true;
-		OriginalMeshDescription = MakeShared<FMeshDescription, ESPMode::ThreadSafe>(UE::ToolTarget::GetMeshDescriptionCopy(Target, GetMeshParams));
+
+		IMeshDescriptionProvider* MeshDescriptionProvider = Cast<IMeshDescriptionProvider>(Target);
+		if (MeshDescriptionProvider)
+		{
+			OriginalMeshDescription = MakeShared<FMeshDescription, ESPMode::ThreadSafe>(
+				MeshDescriptionProvider->GetMeshDescriptionCopy(GetMeshParams));
+		}
+		else
+		{
+			// Target doesn't have a mesh description, need to make one from the dynamic mesh
+			OriginalMeshDescription = MakeShared<FMeshDescription, ESPMode::ThreadSafe>();
+			FStaticMeshAttributes Attributes(*OriginalMeshDescription);
+			Attributes.Register();
+			FDynamicMeshToMeshDescription Converter;
+			Converter.Convert(OriginalMesh.Get(), *OriginalMeshDescription, GetMeshParams.bWantMeshTangents);
+		}
 	}
 
 	Op->bDiscardAttributes = SimplifyProperties->bDiscardAttributes;
diff --git a/Engine/Plugins/Runtime/MeshModelingToolset/Source/MeshModelingToolsEditorOnly/Public/SimplifyMeshTool.h b/Engine/Plugins/Runtime/MeshModelingToolset/Source/MeshModelingToolsEditorOnly/Public/SimplifyMeshTool.h
--- a/Engine/Plugins/Runtime/MeshModelingToolset/Source/MeshModelingToolsEditorOnly/Public/SimplifyMeshTool.h
+++ b/Engine/Plugins/Runtime/MeshModelingToolset/Source/MeshModelingToolsEditorOnly/Public/SimplifyMeshTool.h
@@ -190,6 +190,7 @@ private:
 	UPROPERTY()
 	TObjectPtr<UMeshElementsVisualizer> MeshElementsDisplay;
 
+	// cached mesh description version of the original mesh -- created (by conversion or copy) as needed
 	TSharedPtr<FMeshDescription, ESPMode::ThreadSafe> OriginalMeshDescription;
 	// Dynamic Mesh versions precomputed in Setup (rather than recomputed for every simplify op)
 	TSharedPtr<UE::Geometry::FDynamicMesh3, ESPMode::ThreadSafe> OriginalMesh;
diff --git a/Engine/Plugins/Runtime/MeshModelingToolset/Source/ModelingOperatorsEditorOnly/Private/CleaningOps/SimplifyMeshOp.cpp b/Engine/Plugins/Runtime/MeshModelingToolset/Source/ModelingOperatorsEditorOnly/Private/CleaningOps/SimplifyMeshOp.cpp
--- a/Engine/Plugins/Runtime/MeshModelingToolset/Source/ModelingOperatorsEditorOnly/Private/CleaningOps/SimplifyMeshOp.cpp
+++ b/Engine/Plugins/Runtime/MeshModelingToolset/Source/ModelingOperatorsEditorOnly/Private/CleaningOps/SimplifyMeshOp.cpp
@@ -22,6 +22,7 @@
 #include "Operations/MeshClusterSimplifier.h"
 #include "ConstrainedDelaunay2.h"
 #include "OverlappingCorners.h"
+#include "StaticMeshAttributes.h"
 #include "StaticMeshOperations.h"
 
 #include UE_INLINE_GENERATED_CPP_BY_NAME(SimplifyMeshOp)
@@ -152,6 +153,17 @@ bool FSimplifyMeshOp::ComputeStandardSimplifier(IMeshReduction* MeshReduction, c
 		return false;
 	}
 
+	// Spot-check the mesh description attributes, to catch mesh descriptions w/out expected static mesh attributes
+	{
+		const FStaticMeshConstAttributes SrcAttrs(SrcMeshDescription);
+		if (!SrcAttrs.GetVertexInstanceUVs().IsValid())
+		{
+			UE_LOG(LogGeometry, Error,
+				TEXT("FSimplifyMeshOp::ComputeStandardSimplifier -- mesh description is missing standard static mesh attributes expected by the simplifier"));
+			return false;
+		}
+	}
+
 	// Note: Simplifier cannot run in place, so we always need to make this copy (even if the SrcMeshDescription was temporary ...)
 	FMeshDescription DstMeshDescription(SrcMeshDescription);
 
