Implements a new View Mode for Mesh Attribute Paint Tool which allows for displaying the source material rather than just the grey shaded options.

#Synced-CL 51758550

diff --git a/Engine/Plugins/Runtime/MeshModelingToolset/Source/MeshModelingTools/Private/MeshAttributePaintTool.cpp b/Engine/Plugins/Runtime/MeshModelingToolset/Source/MeshModelingTools/Private/MeshAttributePaintTool.cpp
--- a/Engine/Plugins/Runtime/MeshModelingToolset/Source/MeshModelingTools/Private/MeshAttributePaintTool.cpp
+++ b/Engine/Plugins/Runtime/MeshModelingToolset/Source/MeshModelingTools/Private/MeshAttributePaintTool.cpp
@@ -381,6 +381,17 @@ void UMeshAttributePaintTool::Setup()
 	// build octree
 	VerticesOctree.Initialize(PreviewMesh->GetMesh(), true);
 
+	// assign materials
+	if (IMaterialProvider* MaterialProvider = Cast<IMaterialProvider>(Target))
+	{
+		FComponentMaterialSet MaterialSet;
+		MaterialProvider->GetMaterialSet(MaterialSet);
+		for (int k = 0; k < MaterialSet.Materials.Num(); ++k)
+		{
+			PreviewMesh->SetMaterial(k, MaterialSet.Materials[k]);
+		}
+	}
+
 	UpdateMaterialMode(ViewProperties->MaterialMode);
 
 	RecalculateBrushRadius();
@@ -792,10 +803,17 @@ void UMeshAttributePaintTool::ApplyStamp(const FBrushStampData& Stamp)
 			{
 				ColorOverlay->SetElement(elid, NewColor);
 			}
+			AttribData.Attribute->SetValue(vid, AttribData.CurrentValues[srcvid]);
 			ElIDs.Reset();
 		}
 	}, false);
-	PreviewMesh->NotifyDeferredEditCompleted(UPreviewMesh::ERenderUpdateMode::FastUpdate, EMeshRenderAttributeFlags::VertexColors, false);
+
+	// Updating the vertex colors is only required when the paint data is being displayed via the vertex color data
+	if (ViewProperties->MaterialMode == EMeshAttributePaintMaterialMode::Shaded
+		|| ViewProperties->MaterialMode == EMeshAttributePaintMaterialMode::ColorOnly)
+	{
+		PreviewMesh->NotifyDeferredEditCompleted(UPreviewMesh::ERenderUpdateMode::FastUpdate, EMeshRenderAttributeFlags::VertexColors, false);
+	}
 }
 
 void UMeshAttributePaintTool::UpdateMaterialMode(EMeshAttributePaintMaterialMode MaterialMode)
@@ -815,6 +833,11 @@ void UMeshAttributePaintTool::UpdateMaterialMode(EMeshAttributePaintMaterialMode
 		ActiveOverrideMaterial = nullptr;
 		PreviewMesh->SetOverrideRenderMaterial(GEngine->VertexColorViewModeMaterial_ColorOnly);
 	}
+	else if (MaterialMode == EMeshAttributePaintMaterialMode::ExistingMaterial)
+	{
+		ActiveOverrideMaterial = nullptr;
+		PreviewMesh->ClearOverrideRenderMaterial();
+	}
 
 	if (ActiveOverrideMaterial != nullptr)
 	{
@@ -1094,6 +1117,7 @@ void UMeshAttributePaintTool::ExternalUpdateValues(int32 AttribIndex, const TArr
 				{
 					ColorOverlay->SetElement(elid, NewColor);
 				}
+				AttribData.Attribute->SetValue(vid, AttribData.CurrentValues[vid]);
 				ElIDs.Reset();
 			}
 		});
diff --git a/Engine/Plugins/Runtime/MeshModelingToolset/Source/MeshModelingTools/Public/MeshAttributePaintTool.h b/Engine/Plugins/Runtime/MeshModelingToolset/Source/MeshModelingTools/Public/MeshAttributePaintTool.h
--- a/Engine/Plugins/Runtime/MeshModelingToolset/Source/MeshModelingTools/Public/MeshAttributePaintTool.h
+++ b/Engine/Plugins/Runtime/MeshModelingToolset/Source/MeshModelingTools/Public/MeshAttributePaintTool.h
@@ -108,7 +108,8 @@ UENUM()
 enum class EMeshAttributePaintMaterialMode
 {
 	Shaded,
-	ColorOnly
+	ColorOnly,
+	ExistingMaterial
 };
 
 
@@ -374,7 +375,7 @@ protected:
 	TUniquePtr<TIndexedValuesChangeBuilder<float, FMeshAttributePaintChange>> ActiveChangeBuilder;
 	UE_API void BeginChange();
 	UE_API TUniquePtr<FMeshAttributePaintChange> EndChange();
-	UE_API void ExternalUpdateValues(int32 AttribIndex, const TArray<int32>& VertexIndices, const TArray<float>& NewValues);
+	UE_API virtual void ExternalUpdateValues(int32 AttribIndex, const TArray<int32>& VertexIndices, const TArray<float>& NewValues);
 
 private:
 	struct FBrushActionData
