add boundary locking support to the vertex sculpt tool, and auto-enable it for height sculpt tools
+ fix minor bug in how vertex sculpt's connected component propagation logic

#Synced-CL 51921100

diff --git a/Engine/Plugins/Experimental/MeshPartition/Source/MeshPartitionModelingToolset/Private/MeshPartitionHeightSculptTool.cpp b/Engine/Plugins/Experimental/MeshPartition/Source/MeshPartitionModelingToolset/Private/MeshPartitionHeightSculptTool.cpp
--- a/Engine/Plugins/Experimental/MeshPartition/Source/MeshPartitionModelingToolset/Private/MeshPartitionHeightSculptTool.cpp
+++ b/Engine/Plugins/Experimental/MeshPartition/Source/MeshPartitionModelingToolset/Private/MeshPartitionHeightSculptTool.cpp
@@ -62,7 +62,11 @@ void UHeightSculptTool::Setup()
 	{
 		ModifierTarget->ConfigurePreviewForRendering(DynamicMeshComponent);
 	}
+
+	// TODO: as with SymmetryProperties->bEnableSymmetry above, revisit how to set these properties
+	// in a way that doesn't fight with the save/restore properties features of the base tool
 	ViewProperties->MaterialMode = EMeshEditingMaterialModes::ExistingMaterial;
+	SculptProperties->bLockBoundaries = true;
 }
 
 void UHeightSculptTool::Shutdown(EToolShutdownType ShutdownType)
diff --git a/Engine/Plugins/Runtime/GeometryProcessing/Source/DynamicMesh/Private/Parameterization/MeshPlanarSymmetry.cpp b/Engine/Plugins/Runtime/GeometryProcessing/Source/DynamicMesh/Private/Parameterization/MeshPlanarSymmetry.cpp
--- a/Engine/Plugins/Runtime/GeometryProcessing/Source/DynamicMesh/Private/Parameterization/MeshPlanarSymmetry.cpp
+++ b/Engine/Plugins/Runtime/GeometryProcessing/Source/DynamicMesh/Private/Parameterization/MeshPlanarSymmetry.cpp
@@ -694,7 +694,7 @@ FQuaterniond FMeshPlanarSymmetry::GetMirroredOrientation(const FQuaterniond& Ori
 }
 
 
-void FMeshPlanarSymmetry::GetMirrorVertexROI(const TArray<int>& VertexROI, TArray<int>& MirrorVertexROI, bool bForceSameSizeWithGaps) const
+void FMeshPlanarSymmetry::GetMirrorVertexROI(const TArray<int>& VertexROI, TArray<int>& MirrorVertexROI, bool bForceSameSizeWithGaps, TFunctionRef<bool(int32 VID)> ShouldExcludeVertex) const
 {
 	MirrorVertexROI.Reset();
 	if (bForceSameSizeWithGaps)
@@ -705,19 +705,26 @@ void FMeshPlanarSymmetry::GetMirrorVertexROI(const TArray<int>& VertexROI, TArra
 	for (int32 VertexID : VertexROI)
 	{
 		const FSymmetryVertex& Vertex = Vertices[VertexID];
+		int32 MirrorVID = INDEX_NONE;
 		if (Vertex.bIsSourceVertex)
 		{
-			int32 MirrorVID = Vertex.PairedVertex;
-			MirrorVertexROI.Add(MirrorVID);
+			MirrorVID = Vertex.PairedVertex;
 		}
 		else if (Vertex.bIsSourceVertex == false && Vertex.bOnPlane == false && Vertex.PairedVertex >= 0)
 		{
 			// The ROI already contains the mirror-vertex, we will keep it so that it can be updated later
-			MirrorVertexROI.Add(VertexID);
+			MirrorVID = VertexID;
 		}
-		else if (bForceSameSizeWithGaps)
+		if (MirrorVID == INDEX_NONE || ShouldExcludeVertex(MirrorVID))
 		{
-			MirrorVertexROI.Add(-1);
+			if (bForceSameSizeWithGaps)
+			{
+				MirrorVertexROI.Add(INDEX_NONE);
+			}
+		}
+		else
+		{
+			MirrorVertexROI.Add(MirrorVID);
 		}
 	}
 }
diff --git a/Engine/Plugins/Runtime/GeometryProcessing/Source/DynamicMesh/Public/Parameterization/MeshPlanarSymmetry.h b/Engine/Plugins/Runtime/GeometryProcessing/Source/DynamicMesh/Public/Parameterization/MeshPlanarSymmetry.h
--- a/Engine/Plugins/Runtime/GeometryProcessing/Source/DynamicMesh/Public/Parameterization/MeshPlanarSymmetry.h
+++ b/Engine/Plugins/Runtime/GeometryProcessing/Source/DynamicMesh/Public/Parameterization/MeshPlanarSymmetry.h
@@ -105,9 +105,12 @@ public:
 	/**
 	 * Computes list of vertices that are mirror-constrained to vertices in VertexROI.
 	 * If bForceSameSizeWithGaps is true, MirrorVertexROIOut will be the same size as VertexROI, and -1 is stored 
-	 * for vertices in VertexROI that are source vertices or on the symmetry plane. Otherwise those vertices are skipped.
+	 * for vertices in VertexROI that are source vertices, or on the symmetry plane, or excluded by ShouldExcludeVertex().
+	 * Otherwise those vertices are skipped.
+	 * 
+	 * @param ShouldExcludeVertex Optional filter to exclude vertices: takes vertex ID, returns true if the vertex should not be added to the ROI, false otherwise
 	 */
-	UE_API void GetMirrorVertexROI(const TArray<int>& VertexROI, TArray<int>& MirrorVertexROIOut, bool bForceSameSizeWithGaps) const;
+	UE_API void GetMirrorVertexROI(const TArray<int>& VertexROI, TArray<int>& MirrorVertexROIOut, bool bForceSameSizeWithGaps, TFunctionRef<bool(int32 VID)> ShouldExcludeVertex = [](int32) {return false;}) const;
 
 	/**
 	 * For any vertices in VertexIndices that are on-plane, take the position in VertexPositionsInOut, apply the
diff --git a/Engine/Plugins/Runtime/MeshModelingToolset/Source/MeshModelingTools/Private/MeshVertexSculptTool.cpp b/Engine/Plugins/Runtime/MeshModelingToolset/Source/MeshModelingTools/Private/MeshVertexSculptTool.cpp
--- a/Engine/Plugins/Runtime/MeshModelingToolset/Source/MeshModelingTools/Private/MeshVertexSculptTool.cpp
+++ b/Engine/Plugins/Runtime/MeshModelingToolset/Source/MeshModelingTools/Private/MeshVertexSculptTool.cpp
@@ -861,6 +861,8 @@ void UMeshVertexSculptTool::UpdateRangeQueryTriBuffer(const FFrame3d& LocalFrame
 		});
 	}
 
+	auto ShouldExcludeVertex = [this, &Mesh](int32 VID) { return SculptProperties->bLockBoundaries && Mesh->IsBoundaryVertex(VID); };
+
 	// Build up vertex and triangle ROIs from the remaining range-query triangles.
 	{
 		TRACE_CPUPROFILER_EVENT_SCOPE(DynamicMeshSculptTool_UpdateROI_3Collect);
@@ -876,7 +878,7 @@ void UMeshVertexSculptTool::UpdateRangeQueryTriBuffer(const FFrame3d& LocalFrame
 			int InsideCount = 0;
 			for (int j = 0; j < 3; ++j)
 			{
-				if (Inside[j])
+				if (Inside[j] && !ShouldExcludeVertex(TriV[j]))
 				{
 					VertexROIBuilder.Add(TriV[j]);
 					InsideCount++;
@@ -906,7 +908,7 @@ void UMeshVertexSculptTool::UpdateRangeQueryTriBuffer(const FFrame3d& LocalFrame
 			TSet<int32> ConnectedROIVids;
 				FMeshConnectedComponents Components(Mesh);
 				Components.GrowToConnectedVertices(*Mesh, SeedVids, ConnectedROIVids, nullptr,
-					[this](int32 Vid, int32 Tid) { return VertexROIBuilder.Contains(Vid); });
+					[this](int32 FromVID, int32 ToVID) { return VertexROIBuilder.Contains(ToVID); });
 
 			// We'll need to update TriangleROIBuilder based on the vertices too
 			TArray<int32> TidsToFilter = TriangleROIBuilder.TakeValues();
@@ -934,7 +936,7 @@ void UMeshVertexSculptTool::UpdateRangeQueryTriBuffer(const FFrame3d& LocalFrame
 		{
 			// Find symmetric Vertex ROI. This will overlap with VertexROI in many cases.
 			SymmetricVertexROI.Reset();
-			Symmetry->GetMirrorVertexROI(VertexROI, SymmetricVertexROI, true);
+			Symmetry->GetMirrorVertexROI(VertexROI, SymmetricVertexROI, true, ShouldExcludeVertex);
 			// expand the Triangle ROI to include the symmetric vertex one-rings
 			for (int32 VertexID : SymmetricVertexROI)
 			{
@@ -960,6 +962,7 @@ void UMeshVertexSculptTool::UpdateRangeQueryTriBuffer(const FFrame3d& LocalFrame
 	ensure(SculptProperties->BrushFilter == EMeshVertexSculptBrushFilterType::None);
 	ensure(!bApplySymmetry);
 	ensure(!RequireConnectivityToHitPointInStamp());
+	ensure(!SculptProperties->bLockBoundaries);
 
 	// collect set of vertices and triangles inside brush sphere, from range query result
 	{
diff --git a/Engine/Plugins/Runtime/MeshModelingToolset/Source/MeshModelingTools/Public/MeshVertexSculptTool.h b/Engine/Plugins/Runtime/MeshModelingToolset/Source/MeshModelingTools/Public/MeshVertexSculptTool.h
--- a/Engine/Plugins/Runtime/MeshModelingToolset/Source/MeshModelingTools/Public/MeshVertexSculptTool.h
+++ b/Engine/Plugins/Runtime/MeshModelingToolset/Source/MeshModelingTools/Public/MeshVertexSculptTool.h
@@ -160,6 +160,10 @@ public:
 	UPROPERTY(EditAnywhere, Category = Sculpting, meta = (EditCondition = "bCanFreezeTarget", HideEditConditionToggle))
 	bool bFreezeTarget = false;
 
+	/** Whether to prevent sculpt brushes from moving mesh boundary vertices */
+	UPROPERTY(EditAnywhere, Category = Sculpting)
+	bool bLockBoundaries = false;
+
 	UPROPERTY()
 	bool bCanFreezeTarget = false;
 
