Expose more editor functions for  polygon selection
- undeprecate the properties
- add a method to choose which type of properties the derived class can create
- add grow / shrink / flood and convert to brorder methods

#Synced-CL 52046631

diff --git a/Engine/Plugins/Runtime/MeshModelingToolset/Source/ModelingComponents/Private/Selection/MeshTopologySelectionMechanic.cpp b/Engine/Plugins/Runtime/MeshModelingToolset/Source/ModelingComponents/Private/Selection/MeshTopologySelectionMechanic.cpp
--- a/Engine/Plugins/Runtime/MeshModelingToolset/Source/ModelingComponents/Private/Selection/MeshTopologySelectionMechanic.cpp
+++ b/Engine/Plugins/Runtime/MeshModelingToolset/Source/ModelingComponents/Private/Selection/MeshTopologySelectionMechanic.cpp
@@ -40,6 +40,12 @@ UMeshTopologySelectionMechanic::~UMeshTopologySelectionMechanic()
 	checkf(PreviewGeometryActor == nullptr, TEXT("Shutdown() should be called before UMeshTopologySelectionMechanic is destroyed."));
 }
 
+
+UMeshTopologySelectionMechanicProperties* UMeshTopologySelectionMechanic::CreateProperties()
+{
+	return NewObject<UMeshTopologySelectionMechanicProperties>(this);
+}
+
 void UMeshTopologySelectionMechanic::Setup(UInteractiveTool* ParentToolIn)
 {
 	UInteractionMechanic::Setup(ParentToolIn);
@@ -77,7 +83,7 @@ void UMeshTopologySelectionMechanic::Setup(UInteractiveTool* ParentToolIn)
 	ClickOrDragBehavior->SetDefaultPriority(BasePriority);
 	ParentTool->AddInputBehavior(ClickOrDragBehavior, this);
 
-	Properties = NewObject<UMeshTopologySelectionMechanicProperties>(this);
+	Properties = CreateProperties();
 	Properties->Initialize(this);
 	if (bAddSelectionFilterPropertiesToParentTool)
 	{
diff --git a/Engine/Plugins/Runtime/MeshModelingToolset/Source/ModelingComponents/Private/Selection/PolygonSelectionMechanic.cpp b/Engine/Plugins/Runtime/MeshModelingToolset/Source/ModelingComponents/Private/Selection/PolygonSelectionMechanic.cpp
--- a/Engine/Plugins/Runtime/MeshModelingToolset/Source/ModelingComponents/Private/Selection/PolygonSelectionMechanic.cpp
+++ b/Engine/Plugins/Runtime/MeshModelingToolset/Source/ModelingComponents/Private/Selection/PolygonSelectionMechanic.cpp
@@ -32,6 +32,42 @@ namespace PolygonSelectionMechanicLocals
 	}
 }
 
+///////////////////////////////////////////////////////////////////////////////////////////////////////////////////////////////////////
+
+void UPolygonSelectionMechanicProperties::GrowSelection()
+{
+	if (UPolygonSelectionMechanic* PolygonMechanic = Cast<UPolygonSelectionMechanic>(Mechanic.Get()))
+	{
+		PolygonMechanic->GrowSelection(/*bAsTriangleTopology*/ true);
+	}
+}
+
+void UPolygonSelectionMechanicProperties::ShrinkSelection()
+{
+	if (UPolygonSelectionMechanic* PolygonMechanic = Cast<UPolygonSelectionMechanic>(Mechanic.Get()))
+	{
+		PolygonMechanic->ShrinkSelection(/*bAsTriangleTopology*/ true);
+	}
+}
+
+void UPolygonSelectionMechanicProperties::FloodSelection()
+{
+	if (UPolygonSelectionMechanic* PolygonMechanic = Cast<UPolygonSelectionMechanic>(Mechanic.Get()))
+	{
+		PolygonMechanic->FloodSelection();
+	}
+}
+
+void UPolygonSelectionMechanicProperties::ConvertToBorder()
+{
+	if (UPolygonSelectionMechanic* PolygonMechanic = Cast<UPolygonSelectionMechanic>(Mechanic.Get()))
+	{
+		PolygonMechanic->ConvertSelectionToBorderVertices(/*bAsTriangleTopology*/ true);
+	}
+}
+
+///////////////////////////////////////////////////////////////////////////////////////////////////////////////////////////////////////
+
 void UPolygonSelectionMechanic::Initialize(
 	const FDynamicMesh3* MeshIn,
 	FTransform3d TargetTransformIn,
@@ -57,6 +93,10 @@ void UPolygonSelectionMechanic::Initialize(
 		GetSpatialSourceFuncIn);
 }
 
+UMeshTopologySelectionMechanicProperties* UPolygonSelectionMechanic::CreateProperties()
+{
+	return NewObject<UPolygonSelectionMechanicProperties>(this);
+}
 
 void UPolygonSelectionMechanic::GetSelection_AsGroupTopology(UE::Geometry::FGeometrySelection& SelectionOut, const FCompactMaps* CompactMapsToApply) const
 {
diff --git a/Engine/Plugins/Runtime/MeshModelingToolset/Source/ModelingComponents/Public/Selection/MeshTopologySelectionMechanic.h b/Engine/Plugins/Runtime/MeshModelingToolset/Source/ModelingComponents/Public/Selection/MeshTopologySelectionMechanic.h
--- a/Engine/Plugins/Runtime/MeshModelingToolset/Source/ModelingComponents/Public/Selection/MeshTopologySelectionMechanic.h
+++ b/Engine/Plugins/Runtime/MeshModelingToolset/Source/ModelingComponents/Public/Selection/MeshTopologySelectionMechanic.h
@@ -176,6 +176,8 @@ public:
 		UWorld* WorldIn,
 		TFunction<FDynamicMeshAABBTree3* ()> GetSpatialSourceFuncIn);
 
+	UE_API virtual UMeshTopologySelectionMechanicProperties* CreateProperties();
+
 	UE_API virtual void Setup(UInteractiveTool* ParentTool) override;
 	UE_API virtual void Shutdown() override;
 
diff --git a/Engine/Plugins/Runtime/MeshModelingToolset/Source/ModelingComponents/Public/Selection/PolygonSelectionMechanic.h b/Engine/Plugins/Runtime/MeshModelingToolset/Source/ModelingComponents/Public/Selection/PolygonSelectionMechanic.h
--- a/Engine/Plugins/Runtime/MeshModelingToolset/Source/ModelingComponents/Public/Selection/PolygonSelectionMechanic.h
+++ b/Engine/Plugins/Runtime/MeshModelingToolset/Source/ModelingComponents/Public/Selection/PolygonSelectionMechanic.h
@@ -9,14 +9,29 @@
 
 namespace UE::Geometry { struct FGeometrySelection; }
 
-// DEPRECATED: Use UMeshTopologySelectionMechanicProperties
-UCLASS(MinimalAPI, Deprecated)
-class UDEPRECATED_PolygonSelectionMechanicProperties : public UMeshTopologySelectionMechanicProperties
+UCLASS(MinimalAPI)
+class UPolygonSelectionMechanicProperties : public UMeshTopologySelectionMechanicProperties
 {
 	GENERATED_BODY()
 
 public:
 
+	/** Grow current selection. If selection is empty, nothing happens */
+	UFUNCTION(CallInEditor, Category = SelectionActions)
+	UE_API void GrowSelection();
+
+	/** Shrink current selection. If selection is empty, nothing happens */
+	UFUNCTION(CallInEditor, Category = SelectionActions)
+	UE_API void ShrinkSelection();
+
+	/** Expands selection to encompass connected components. */
+	UFUNCTION(CallInEditor, Category = SelectionActions)
+	UE_API void FloodSelection();
+
+	/** Converts selection to a vertex/corner selection of just the boundary vertices/corners. */
+	UFUNCTION(CallInEditor, Category = SelectionActions)
+	UE_API void ConvertToBorder();
+
 	void Initialize(UPolygonSelectionMechanic* MechanicIn)
 	{
 		Mechanic = MechanicIn;
@@ -46,6 +61,8 @@ public:
 		TFunction<FDynamicMeshAABBTree3* ()> GetSpatialSourceFunc
 	);
 
+	UE_API virtual UMeshTopologySelectionMechanicProperties* CreateProperties() override;
+
 	/*
 	 * Expands selection at the borders.
 	 * 
