Fixed normals being randomly corrupted by undo/redo in vertex sculpt tool
by making vertex sculpt tool normals recompute on the game thread on undo/redo -- since the Async normal update was racing the DynamicMeshComponent render update otherwise. Also remove the render update from OnTick() since it's now fully redundant to the DynamicMeshComponent's automatic one.

(We could alternatively try to make the OnTick() render update more-reliably run to fix the normals after they're computed, but it likely wouldn't be a perf win unless we also suppress the full render update that the undo/redo is triggering on dynamicmeshcomponent, which would be a larger change. Typically the local normal recompute alone is relatively not that expensive compared to the render buffer updates.)

#Synced-CL 52079281

diff --git a/Engine/Plugins/Runtime/MeshModelingToolset/Source/MeshModelingTools/Private/MeshVertexSculptTool.cpp b/Engine/Plugins/Runtime/MeshModelingToolset/Source/MeshModelingTools/Private/MeshVertexSculptTool.cpp
--- a/Engine/Plugins/Runtime/MeshModelingToolset/Source/MeshModelingTools/Private/MeshVertexSculptTool.cpp
+++ b/Engine/Plugins/Runtime/MeshModelingToolset/Source/MeshModelingTools/Private/MeshVertexSculptTool.cpp
@@ -1480,20 +1480,8 @@ void UMeshVertexSculptTool::OnTick(float DeltaTime)
 
 	TRACE_CPUPROFILER_EVENT_SCOPE(VtxSculptTool_Tick);
 
-	// process the undo update
-	if (bUndoUpdatePending)
-	{
-		// wait for updates
-		WaitForPendingUndoRedoUpdate();
-
-		// post rendering update
-		DynamicMeshComponent->FastNotifyTriangleVerticesUpdated(AccumulatedTriangleROI,
-			EMeshRenderAttributeFlags::Positions | EMeshRenderAttributeFlags::VertexNormals);
-		GetToolManager()->PostInvalidation();
-
-		// ignore stamp and wait for next tick to do anything else
-		return;
-	}
+	// Ensure undo/redo async work is complete before processing stamps
+	WaitForPendingUndoRedoUpdate();
 
 	// if user changed to not-frozen, we need to reinitialize the target
 	if (bCachedFreezeTarget != SculptProperties->bFreezeTarget)
@@ -1926,14 +1914,9 @@ void UMeshVertexSculptTool::OnDynamicMeshComponentChanged(UDynamicMeshComponent*
 	}, bRevert);
 
 	UndoUpdateFuture.Reset();
-	bUndoUpdatePending = true;
 
-	// start the normal recomputation
-	UndoNormalsFuture = Async(VertexSculptToolAsyncExecTarget, [this, Mesh]()
-	{
-		UE::SculptUtil::RecalculateROINormals(Mesh, AccumulatedTriangleROI, NormalsROIBuilder);
-		return true;
-	});
+	// Recompute normals synchronously, so we don't end up racing UDynamicMeshComponent's render buffer updates (which happen soon after this method)
+	UE::SculptUtil::RecalculateROINormals(Mesh, AccumulatedTriangleROI, NormalsROIBuilder);
 
 	// start the octree update
 	UndoUpdateOctreeFuture = Async(VertexSculptToolAsyncExecTarget, [this, Mesh]()
@@ -1952,10 +1935,8 @@ void UMeshVertexSculptTool::OnDynamicMeshComponentChanged(UDynamicMeshComponent*
 
 	UndoUpdateFuture = Async(VertexSculptToolAsyncExecTarget, [this]()
 	{
-		UndoNormalsFuture.Wait();
 		UndoUpdateOctreeFuture.Wait();
 		UndoUpdateBaseMeshFuture.Wait();
-		bUndoUpdatePending = false;
 	});
 
 	RefreshSculptLayerProperties();
diff --git a/Engine/Plugins/Runtime/MeshModelingToolset/Source/MeshModelingTools/Public/MeshVertexSculptTool.h b/Engine/Plugins/Runtime/MeshModelingToolset/Source/MeshModelingTools/Public/MeshVertexSculptTool.h
--- a/Engine/Plugins/Runtime/MeshModelingToolset/Source/MeshModelingTools/Public/MeshVertexSculptTool.h
+++ b/Engine/Plugins/Runtime/MeshModelingToolset/Source/MeshModelingTools/Public/MeshVertexSculptTool.h
@@ -371,8 +371,9 @@ protected:
 	int32 InitialStrokeTriangleID = -1;
 
 	TSet<int32> AccumulatedTriangleROI;
-	
+	UE_DEPRECATED(5.8, "bUndoUpdatePending is no longer used")
 	bool bUndoUpdatePending = false;
+	UE_DEPRECATED(5.8, "UndoNormalsFuture is no longer used")
 	TFuture<bool> UndoNormalsFuture;
 	TFuture<bool> UndoUpdateOctreeFuture;
 	TFuture<bool> UndoUpdateBaseMeshFuture;
