default sculpt smooth brushes bPreserveUVFlow to true (for primary smooth and shift-smooth)

#Synced-CL 52243346

diff --git a/Engine/Plugins/Runtime/MeshModelingToolset/Source/MeshModelingTools/Public/Sculpting/MeshSmoothingBrushOps.h b/Engine/Plugins/Runtime/MeshModelingToolset/Source/MeshModelingTools/Public/Sculpting/MeshSmoothingBrushOps.h
--- a/Engine/Plugins/Runtime/MeshModelingToolset/Source/MeshModelingTools/Public/Sculpting/MeshSmoothingBrushOps.h
+++ b/Engine/Plugins/Runtime/MeshModelingToolset/Source/MeshModelingTools/Public/Sculpting/MeshSmoothingBrushOps.h
@@ -39,7 +39,7 @@ public:
 
 	/** If true, try to preserve the shape of the UV/3D mapping. This will limit Smoothing and Remeshing in some cases. */
 	UPROPERTY(EditAnywhere, Category = SmoothBrush)
-	bool bPreserveUVFlow = false;
+	bool bPreserveUVFlow = true;
 
 	virtual float GetStrength() override { return Strength; }
 	virtual void SetStrength(float NewStrength) override { Strength = FMathf::Clamp(NewStrength, 0.0f, 1.0f); }
@@ -68,7 +68,7 @@ public:
 
 	/** If true, try to preserve the shape of the UV/3D mapping. This will limit Smoothing and Remeshing in some cases. */
 	UPROPERTY(EditAnywhere, Category = ShiftSmoothBrush)
-	bool bPreserveUVFlow = false;
+	bool bPreserveUVFlow = true;
 
 	virtual float GetStrength() override { return Strength; }
 	virtual void SetStrength(float NewStrength) override { Strength = FMathf::Clamp(NewStrength, 0.0f, 1.0f); }
