Vertex Sculpt tool optimizations:
- Change ROI queries to use a vertex-octree; just use the triangle-octree for raycasts. (Vertex ROI queries are much faster)
- Cache the vertex->vertex and vertex->tri connectivity and use the cached connectivity for vtx->tri ROI conversion and for region growing

#Synced-CL 52243346

diff --git a/Engine/Plugins/Runtime/MeshModelingToolset/Source/MeshModelingTools/Private/MeshVertexSculptTool.cpp b/Engine/Plugins/Runtime/MeshModelingToolset/Source/MeshModelingTools/Private/MeshVertexSculptTool.cpp
--- a/Engine/Plugins/Runtime/MeshModelingToolset/Source/MeshModelingTools/Private/MeshVertexSculptTool.cpp
+++ b/Engine/Plugins/Runtime/MeshModelingToolset/Source/MeshModelingTools/Private/MeshVertexSculptTool.cpp
@@ -163,31 +163,37 @@ void UMeshVertexSculptTool::Setup()
 	int TreeDepthOverride = FMath::Abs(MeshVertexSculptToolLocals::OctreeTreeDepthOverride->GetInt());
 	TFuture<void> InitializeOctree = Async(VertexSculptToolAsyncExecTarget, [SculptMesh, Bounds, RootCellSizeOverride, TreeDepthOverride, this]()
 	{
+		double RootDim;
+		int32 Depth;
 		if (SculptMesh->TriangleCount() > 100000)
 		{
-			Octree.RootDimension = InitialBoundsMaxDim / 10.0;
-			Octree.SetMaxTreeDepth(4);
+			RootDim = InitialBoundsMaxDim / 10.0;
+			Depth = 4;
 		}
 		else
 		{
-			Octree.RootDimension = InitialBoundsMaxDim / 2.0;
-			Octree.SetMaxTreeDepth(8);
+			RootDim = InitialBoundsMaxDim / 2.0;
+			Depth = 8;
 		}
 
 		if (!FMath::IsNearlyZero(RootCellSizeOverride))
 		{
-			Octree.RootDimension = RootCellSizeOverride;
+			RootDim = RootCellSizeOverride;
 		}
 		if (TreeDepthOverride > 0)
 		{
-			Octree.SetMaxTreeDepth(TreeDepthOverride);
+			Depth = TreeDepthOverride;
 		}
 
+		// Vertex octree for region-of-interest range queries
+		VertexOctree.RootDimension = RootDim;
+		VertexOctree.MaxTreeDepth = Depth;
+		VertexOctree.Initialize(SculptMesh, false);
+
+		// Triangle octree for hit testing
+		Octree.RootDimension = RootDim;
+		Octree.SetMaxTreeDepth(Depth);
 		Octree.Initialize(SculptMesh);
-		//Octree.CheckValidity(EValidityCheckFailMode::Check, true, true);
-		//FDynamicMeshOctree3::FStatistics Stats;
-		//Octree.ComputeStatistics(Stats);
-		//UE_LOG(LogTemp, Warning, TEXT("Octree Stats: %s"), *Stats.ToString());
 	});
 
 	// find mesh connected-component index for each triangle
@@ -225,6 +231,19 @@ void UMeshVertexSculptTool::Setup()
 		bTargetDirty = false;
 	});
 
+	// Cache mesh connectivity once up front. This assumes vertex sculpting never changes topology
+	// (only vertex positions), so the connectivity remains valid for the lifetime of the tool
+	TFuture<void> InitializeCachedConnectivity = Async(VertexSculptToolAsyncExecTarget, [SculptMesh, this]()
+	{
+		CachedVtxTriConnectivity.SetNum(SculptMesh->MaxVertexID());
+		CachedVtxVtxConnectivity.SetNum(SculptMesh->MaxVertexID());
+		for (int32 VID : SculptMesh->VertexIndicesItr())
+		{
+			SculptMesh->EnumerateVertexTriangles(VID, [VID, this](int32 TID) { CachedVtxTriConnectivity[VID].Add(TID); });
+			SculptMesh->EnumerateVertexVertices(VID, [VID, this](int32 OtherVID) { CachedVtxVtxConnectivity[VID].Add(OtherVID); });
+		}
+	});
+
 	// initialize render decomposition
 	TFuture<void> InitializeRenderDecomp = Async(VertexSculptToolAsyncExecTarget, [SculptMesh, &MaterialSet, this]()
 	{
@@ -246,6 +265,7 @@ void UMeshVertexSculptTool::Setup()
 	InitializeBaseMesh.Wait();
 	InitializeRenderDecomp.Wait();
 	InitializeSymmetry.Wait();
+	InitializeCachedConnectivity.Wait();
 
 	// initialize brush radius range interval, brush properties
 	UMeshSculptToolBase::InitializeBrushSizeRange(Bounds);
@@ -699,7 +719,7 @@ void UMeshVertexSculptTool::OnCancelStroke()
 // The first part of UpdateROI, which updates TriangleROIArray to be triangles
 //  in our region of interest, and VertexROI to be vertices in our region of
 //  interest.
-void UMeshVertexSculptTool::UpdateRangeQueryTriBuffer(const FFrame3d& LocalFrame)
+void UMeshVertexSculptTool::ComputeVertexAndTriangleROI(const FFrame3d& LocalFrame)
 {
 	if (RequireConnectivityToHitPointInStamp()
 		// It's possible for LastBrushTriangleID to be null if we started a stroke and brushed off the
@@ -707,7 +727,7 @@ void UMeshVertexSculptTool::UpdateRangeQueryTriBuffer(const FFrame3d& LocalFrame
 		&& LastBrushTriangleID == IndexConstants::InvalidID)
 	{
 		// If we're requiring connectivity, and we didn't hit a triangle to start with, then we shouldn't
-		//  move any triangles.
+		//  move any vertices.
 		// Make sure TriangleROIArray is not in use, as we're about to clear it
 		WaitForPendingStampUpdateConst();
 		
@@ -722,21 +742,20 @@ void UMeshVertexSculptTool::UpdateRangeQueryTriBuffer(const FFrame3d& LocalFrame
 
 	// By default, our brush is a sphere, and we affect vertices inside it
 	float RadiusSqr = GetCurrentBrushRadius() * GetCurrentBrushRadius();
-	// This function gets called when first gathering the triangles that might intersect our brush
-	//  from the octree's cells
-	TFunction<void()> GatherOverlappingCells = [this, &BrushPos]()
-	{
-		FAxisAlignedBox3d BrushBox(
-			BrushPos - GetCurrentBrushRadius() * FVector3d::One(),
-			BrushPos + GetCurrentBrushRadius() * FVector3d::One());
-		Octree.ParallelRangeQuery(BrushBox, RangeQueryTriBuffer);
-	};
 	// This is used to filter the gathered verts for ones that are actually in the brush.
 	TFunction<bool(int32 Vid)> IsVertInBrush = [Mesh, &BrushPos, RadiusSqr](int32 Vid)
 	{
 		return DistanceSquared(BrushPos, Mesh->GetVertexRef(Vid)) < RadiusSqr;
 	};
 
+	// Conservative bounding box for the octree range query (for default sphere brush)
+	FAxisAlignedBox3d BrushBox(
+		BrushPos - GetCurrentBrushRadius() * FVector3d::One(),
+		BrushPos + GetCurrentBrushRadius() * FVector3d::One());
+
+	// Optional bounds predicate for tighter octree cell culling (used by cylinder brushes)
+	TFunction<bool(const FAxisAlignedBox3d&)> BoundsPredicateFunc;
+
 	// Some brush types want their brush to be a cylinder, so we need to change how we evaluate
 	//  cells/vertices that are within reach
 	TUniquePtr<FMeshSculptBrushOp>& CurrentBrush = GetActiveBrushOp();
@@ -745,7 +764,7 @@ void UMeshVertexSculptTool::UpdateRangeQueryTriBuffer(const FFrame3d& LocalFrame
 		|| CurrentBrush->GetBrushRegionType() == FMeshSculptBrushOp::EBrushRegionType::CylinderOnSphere))
 	{
 		double CylinderRadius = GetCurrentBrushRadius();
-		double CylinderHeight = TNumericLimits<double>::Max();
+		double CylinderHeight = UE_LARGE_WORLD_MAX;
 
 		FVector3d CylinderCenter, CylinderAxis;
 		if (CurrentBrush->GetBrushRegionType() == FMeshSculptBrushOp::EBrushRegionType::InfiniteCylinder)
@@ -782,37 +801,33 @@ void UMeshVertexSculptTool::UpdateRangeQueryTriBuffer(const FFrame3d& LocalFrame
 			};
 		}
 
-		auto DoesCellIntersectBrush = [CylinderCenter, CylinderAxis, CylinderRadius, CylinderHeight](const FAxisAlignedBox3d& CellBounds)
-		{
-			return UE::Geometry::DoesCylinderIntersectBox(CellBounds,
-				CylinderCenter, CylinderAxis, CylinderRadius, CylinderHeight);
-		};
-		
-		FAxisAlignedBox3d ConservativeCylinderBounds;
-		ConservativeCylinderBounds.Contain(CylinderCenter + CylinderAxis * (CylinderHeight / 2));
-		ConservativeCylinderBounds.Contain(CylinderCenter - CylinderAxis * (CylinderHeight / 2));
-		ConservativeCylinderBounds.Expand(CylinderRadius);
+		// Change the conservative brush bounding box to bound the query cylinder
+		BrushBox = FAxisAlignedBox3d::Empty();
+		BrushBox.Contain(CylinderCenter + CylinderAxis * (CylinderHeight / 2));
+		BrushBox.Contain(CylinderCenter - CylinderAxis * (CylinderHeight / 2));
+		BrushBox.Expand(CylinderRadius);
 
-		GatherOverlappingCells = [this, ConservativeCylinderBounds, BoundsOverlapFn = MoveTemp(DoesCellIntersectBrush)]()
+		// For non-axis-aligned cylinders the UE_LARGE_WORLD_MAX-height cylinder's AABB is likely to cover the entire mesh bounds,
+		// so also use a line distance test to conservatively filter
+		if (CylinderAxis.GetAbsMax() < 1.-UE_DOUBLE_KINDA_SMALL_NUMBER)
 		{
-			Octree.ParallelRangeQuery(ConservativeCylinderBounds, BoundsOverlapFn, RangeQueryTriBuffer);
-		};
-	}
-
-	// Do a parallel range query to find those triangles that may intersect with
-	//  our brush bounds. This grabs all triangles of intersecting cells, so we
-	//  will need to do additional filtering afterward.
-	RangeQueryTriBuffer.Reset();
-	{
-		TRACE_CPUPROFILER_EVENT_SCOPE(VtxSculptTool_UpdateROI_RangeQuery);
-		GatherOverlappingCells();
+			FLine3d CylinderLine(CylinderCenter, CylinderAxis);
+			BoundsPredicateFunc = [CylinderLine, CylinderRadius](const FAxisAlignedBox3d& CellBounds)
+			{
+				FVector3d BoundsCenter = CellBounds.Center();
+				double LineDistSq = CylinderLine.DistanceSquared(BoundsCenter);
+				double OverlapRadius = (CylinderRadius + CellBounds.Diagonal().Length() * .5);
+				return LineDistSq <= OverlapRadius * OverlapRadius;
+			};
+		}
 	}
 
+	// Set up component/polygroup filter IDs
 	int32 ActiveComponentID = -1;
 	int32 ActiveGroupID = -1;
 	if (SculptProperties->BrushFilter == EMeshVertexSculptBrushFilterType::Component)
 	{
-		ActiveComponentID = (InitialStrokeTriangleID >= 0 && InitialStrokeTriangleID <= TriangleComponentIDs.Num()) ?
+		ActiveComponentID = (InitialStrokeTriangleID >= 0 && InitialStrokeTriangleID < TriangleComponentIDs.Num()) ?
 			TriangleComponentIDs[InitialStrokeTriangleID] : -1;
 	}
 	else if (SculptProperties->BrushFilter == EMeshVertexSculptBrushFilterType::PolyGroup)
@@ -820,182 +835,148 @@ void UMeshVertexSculptTool::UpdateRangeQueryTriBuffer(const FFrame3d& LocalFrame
 		ActiveGroupID = Mesh->IsTriangle(InitialStrokeTriangleID) ? ActiveGroupSet->GetGroup(InitialStrokeTriangleID) : -1;
 	}
 
-#if 1
-	// in this path we use more memory but this lets us do more in parallel
+	bool bFilterByBoundary = SculptProperties->bLockBoundaries;
 
-	// Construct array of inside/outside flags for each triangle's vertices. If no
-	// vertices are inside, clear the triangle ID from the range query buffer.
-	// This can be done in parallel and it's cheaper to do repeated distance computations
-	// than to try to do it inside the ROI building below (todo: profile this some more?)
-	TriangleROIInBuf.SetNum(RangeQueryTriBuffer.Num(), EAllowShrinking::No);
+	// Query the vertex octree, filtering by brush shape, boundary vertex, and active component/group
+	TArray<int32> VertexQueryBuffer;
 	{
-		TRACE_CPUPROFILER_EVENT_SCOPE(DynamicMeshSculptTool_UpdateROI_TriVerts);
-		ParallelFor(RangeQueryTriBuffer.Num(), [&](int k)
+		TRACE_CPUPROFILER_EVENT_SCOPE(VtxSculptTool_UpdateROI_RangeQuery);
+		auto VertexPredicate = [this, Mesh, &IsVertInBrush, bFilterByBoundary, ActiveComponentID, ActiveGroupID](int32 VID) -> bool
 		{
-			// check various triangle ROI filters
-			int32 tid = RangeQueryTriBuffer[k];
-			bool bDiscardTriangle = false;
-			if (ActiveComponentID >= 0 && TriangleComponentIDs[tid] != ActiveComponentID)
-			{
-				bDiscardTriangle = true;
-			}
-			if (ActiveGroupID >= 0 && ActiveGroupSet->GetGroup(tid) != ActiveGroupID)
+			if (!IsVertInBrush(VID))
 			{
-				bDiscardTriangle = true;
+				return false;
 			}
-			if (bDiscardTriangle)
+			if (bFilterByBoundary && Mesh->IsBoundaryVertex(VID))
 			{
-				TriangleROIInBuf[k].A = TriangleROIInBuf[k].B = TriangleROIInBuf[k].C = 0;
-				RangeQueryTriBuffer[k] = -1;
-				return;
+				return false;
 			}
-
-			const FIndex3i& TriV = Mesh->GetTriangleRef(tid);
-			TriangleROIInBuf[k].A = IsVertInBrush(TriV.A) ? 1 : 0;
-			TriangleROIInBuf[k].B = IsVertInBrush(TriV.B) ? 1 : 0;
-			TriangleROIInBuf[k].C = IsVertInBrush(TriV.C) ? 1 : 0;
-			if (TriangleROIInBuf[k].A + TriangleROIInBuf[k].B + TriangleROIInBuf[k].C == 0)
+			if (ActiveComponentID >= 0 || ActiveGroupID >= 0)
 			{
-				RangeQueryTriBuffer[k] = -1;
+				for (int32 tid : CachedVtxTriConnectivity[VID])
+				{
+					if (ActiveComponentID >= 0 && TriangleComponentIDs[tid] != ActiveComponentID)
+					{
+						continue;
+					}
+					if (ActiveGroupID >= 0 && ActiveGroupSet->GetGroup(tid) != ActiveGroupID)
+					{
+						continue;
+					}
+					return true; // found a tri matching the group and component filters
+				}
+				return false; // filter(s) were active and didn't find a match
 			}
-		});
+			return true;
+		};
+		if (BoundsPredicateFunc)
+		{
+			VertexOctree.ParallelRangeQuery(BrushBox, BoundsPredicateFunc, VertexPredicate, VertexQueryBuffer);
+		}
+		else
+		{
+			VertexOctree.ParallelRangeQuery(BrushBox, VertexPredicate, VertexQueryBuffer);
+		}
 	}
 
-	auto ShouldExcludeVertex = [this, &Mesh](int32 VID) { return SculptProperties->bLockBoundaries && Mesh->IsBoundaryVertex(VID); };
+	auto ShouldExcludeVertex = [Mesh, bFilterByBoundary](int32 VID) { return bFilterByBoundary && Mesh->IsBoundaryVertex(VID); };
+
+	// Make sure TriangleROIArray, VertexROI and SymmetricVertexROI are not in use, as we're about to change them
+	WaitForPendingStampUpdateConst();
 
-	// Build up vertex and triangle ROIs from the remaining range-query triangles.
+	// Build VertexROI from query results. Optionally filter by connectivity to hit point.
 	{
-		TRACE_CPUPROFILER_EVENT_SCOPE(DynamicMeshSculptTool_UpdateROI_3Collect);
-		VertexROIBuilder.Initialize(Mesh->MaxVertexID());
-		TriangleROIBuilder.Initialize(Mesh->MaxTriangleID());
-		int32 N = RangeQueryTriBuffer.Num();
-		for ( int32 k = 0; k < N; ++k )
+		TRACE_CPUPROFILER_EVENT_SCOPE(VtxSculptTool_UpdateROI_BuildROI);
+
+		if (RequireConnectivityToHitPointInStamp()
+			&& ensure(LastBrushTriangleID != IndexConstants::InvalidID))
 		{
-			int32 tid = RangeQueryTriBuffer[k];
-			if (tid == -1) continue;		// triangle was deleted in previous step
-			const FIndex3i& TriV = Mesh->GetTriangleRef(RangeQueryTriBuffer[k]);
-			const FIndex3i& Inside = TriangleROIInBuf[k];
-			int InsideCount = 0;
-			for (int j = 0; j < 3; ++j)
+			TRACE_CPUPROFILER_EVENT_SCOPE(VtxSculptTool_UpdateROI_BuildROI_FindConnectedComponents);
+			// Build a 'lookup' index set of all vertices in the original brush shape
 			{
-				if (Inside[j] && !ShouldExcludeVertex(TriV[j]))
+				VertexROILookup.Initialize(Mesh->MaxVertexID());
+				for (int32 VID : VertexQueryBuffer)
 				{
-					VertexROIBuilder.Add(TriV[j]);
-					InsideCount++;
+					VertexROILookup.Add(VID);
 				}
 			}
-			if (InsideCount > 0)
-			{
-				TriangleROIBuilder.Add(tid);
-			}
-		}
 
-		// See if we need to filter our vertices based on connectivity to hit location (used to avoid affecting
-		//  hidden regions of a mesh that might be in the volume of the brush)
-		if (RequireConnectivityToHitPointInStamp()
-			&& ensure(LastBrushTriangleID != IndexConstants::InvalidID))
-		{
+			// Start growing actual ROI from vertices on the hit triangle that are in that 'lookup' set
 			FIndex3i HitTriVids = Mesh->GetTriangle(LastBrushTriangleID);
 			TArray<int32> SeedVids;
 			for (int i = 0; i < 3; ++i)
 			{
-				if (VertexROIBuilder.Contains(HitTriVids[i]))
+				if (VertexROILookup.Contains(HitTriVids[i]))
 				{
 					SeedVids.Add(HitTriVids[i]);
 				}
 			}
-			
-			TSet<int32> ConnectedROIVids;
-				FMeshConnectedComponents Components(Mesh);
-				Components.GrowToConnectedVertices(*Mesh, SeedVids, ConnectedROIVids, nullptr,
-					[this](int32 FromVID, int32 ToVID) { return VertexROIBuilder.Contains(ToVID); });
-
-			// We'll need to update TriangleROIBuilder based on the vertices too
-			TArray<int32> TidsToFilter = TriangleROIBuilder.TakeValues();
-			TriangleROIBuilder.Initialize(Mesh->MaxTriangleID());
-			for (int32 Tid : TidsToFilter)
+
+			// grow to connected vertices from the seed
+			TArray<int32> LocalQueue;
+			LocalQueue.Append(SeedVids);
+
+			VertexROIBuilder.Initialize(Mesh->MaxVertexID());
+			for (int32 InputValue : SeedVids)
+			{
+				VertexROIBuilder.Add(InputValue);
+			}
+
+			while (LocalQueue.Num() > 0)
 			{
-				FIndex3i TriVids = Mesh->GetTriangle(Tid);
-				for (int i = 0; i < 3; ++i)
+				int32 CurVertexID = LocalQueue.Pop(EAllowShrinking::No);
+
+				for (int32 NbrVertexID : CachedVtxVtxConnectivity[CurVertexID])
 				{
-					if (ConnectedROIVids.Contains(TriVids[i]))
+					if (!VertexROIBuilder.Contains(NbrVertexID) && VertexROILookup.Contains(NbrVertexID))
 					{
-						TriangleROIBuilder.Add(Tid);
-						break; // continue to next Tid
+						LocalQueue.Add(NbrVertexID);
+						VertexROIBuilder.Add(NbrVertexID);
 					}
 				}
 			}
-			VertexROI = ConnectedROIVids.Array();
+
+			VertexROIBuilder.SwapValuesWith(VertexROI);
 		}
 		else
 		{
-			VertexROIBuilder.SwapValuesWith(VertexROI);
+			VertexROI = MoveTemp(VertexQueryBuffer);
 		}
 
 		if (bApplySymmetry)
 		{
-			// Find symmetric Vertex ROI. This will overlap with VertexROI in many cases.
 			SymmetricVertexROI.Reset();
 			Symmetry->GetMirrorVertexROI(VertexROI, SymmetricVertexROI, true, ShouldExcludeVertex);
-			// expand the Triangle ROI to include the symmetric vertex one-rings
-			for (int32 VertexID : SymmetricVertexROI)
-			{
-				if (Mesh->IsVertex(VertexID))
-				{
-					Mesh->EnumerateVertexTriangles(VertexID, [&](int32 tid)
-					{
-						TriangleROIBuilder.Add(tid);
-					});
-				}
-			}
 		}
-		// Make sure TriangleROIArray is not in use, as we're about to change it
-		WaitForPendingStampUpdateConst();
-		TriangleROIBuilder.SwapValuesWith(TriangleROIArray);
-	}
-
-#else
-	// In this path we combine everything into one loop. Does fewer distance checks
-	// but nothing can be done in parallel (would change if ROIBuilders had atomic-try-add)
 
-	// TODO would need to support these, this branch is likely dead though
-	ensure(SculptProperties->BrushFilter == EMeshVertexSculptBrushFilterType::None);
-	ensure(!bApplySymmetry);
-	ensure(!RequireConnectivityToHitPointInStamp());
-	ensure(!SculptProperties->bLockBoundaries);
-
-	// collect set of vertices and triangles inside brush sphere, from range query result
-	{
-		TRACE_CPUPROFILER_EVENT_SCOPE(DynamicMeshSculptTool_UpdateROI_2Collect);
-		VertexROIBuilder.Initialize(Mesh->MaxVertexID());
+		// Build TriangleROIArray by expanding VertexROI (+ symmetric) to one-ring triangles
 		TriangleROIBuilder.Initialize(Mesh->MaxTriangleID());
-		for (int32 TriIdx : RangeQueryTriBuffer)
 		{
-			FIndex3i TriV = Mesh->GetTriangle(TriIdx);
-			int InsideCount = 0;
-			for (int j = 0; j < 3; ++j)
+			TRACE_CPUPROFILER_EVENT_SCOPE(VtxSculptTool_UpdateROI_BuildROI_ConvertToTriangles);
+			for (int32 VID : VertexROI)
 			{
-				if (VertexROIBuilder.Contains(TriV[j]))
-				{
-					InsideCount++;
-				} 
-				else if (IsVertInBrush(TriV[j]))
+				for (int32 TID : CachedVtxTriConnectivity[VID])
 				{
-					VertexROIBuilder.Add(TriV[j]);
-					InsideCount++;
+					TriangleROIBuilder.Add(TID);
 				}
 			}
-			if (InsideCount > 0)
+		}
+		if (bApplySymmetry)
+		{
+			TRACE_CPUPROFILER_EVENT_SCOPE(VtxSculptTool_UpdateROI_BuildROI_ConvertSymmetricROIToTriangles);
+			for (int32 VID : SymmetricVertexROI)
 			{
-				TriangleROIBuilder.Add(tid);
+				if (Mesh->IsVertex(VID))
+				{
+					for (int32 TID : CachedVtxTriConnectivity[VID])
+					{
+						TriangleROIBuilder.Add(TID);
+					}
+				}
 			}
 		}
-		VertexROIBuilder.SwapValuesWith(VertexROI);
-		// Make sure TriangleROIArray is not in use, as we're about to change it
-		WaitForPendingStampUpdateConst();
 		TriangleROIBuilder.SwapValuesWith(TriangleROIArray);
 	}
-#endif
 }
 
 // Second part of UpdateROI, which fills out ROIPrevPositionBuffer, prepares ROIPositionBuffer,
@@ -1036,7 +1017,7 @@ void UMeshVertexSculptTool::UpdateROI(const FFrame3d& LocalFrame)
 {
 	TRACE_CPUPROFILER_EVENT_SCOPE(VtxSculptTool_UpdateROI);
 
-	UpdateRangeQueryTriBuffer(LocalFrame);
+	ComputeVertexAndTriangleROI(LocalFrame);
 	PrepROIVertPositionBuffers();
 }
 
@@ -1569,7 +1550,20 @@ void UMeshVertexSculptTool::OnTick(float DeltaTime)
 			StampUpdateOctreeFuture = Async(VertexSculptToolAsyncExecTarget, [this]()
 			{
 				TRACE_CPUPROFILER_EVENT_SCOPE(VtxSculptTool_Tick_OctreeReinsert);
+				// Vertex octree reinsertion (for next stamp's range query)
+				TFuture<void> VertexOctreeUpdateFuture = Async(VertexSculptToolAsyncExecTarget, [this]()
+				{
+					VertexOctree.ReinsertVertices(VertexROI);
+					if (bApplySymmetry)
+					{
+						const FDynamicMesh3* SculptMesh = GetSculptMesh();
+						VertexOctree.ReinsertFilteredVertices(SymmetricVertexROI, [SculptMesh](int32 VID) {return !SculptMesh->IsVertex(VID);});
+					}
+				});
+			
+				// Triangle octree reinsertion (for hit testing)
 				Octree.ReinsertTrianglesParallel(TriangleROIArray, OctreeUpdateTempBuffer, OctreeUpdateTempFlagBuffer);
+				VertexOctreeUpdateFuture.Wait();
 				bOctreeUpdated = true;
 			});
 		}
@@ -1589,7 +1583,6 @@ void UMeshVertexSculptTool::OnTick(float DeltaTime)
 		}
 
 		// recalculate normals. This has to complete before we can update component
-		// (in fact we could do it per-chunk...)
 		{
 			TRACE_CPUPROFILER_EVENT_SCOPE(VtxSculptTool_Tick_RecalcNormals);
 			NormalsROI.Wait();
@@ -1906,11 +1899,21 @@ void UMeshVertexSculptTool::OnDynamicMeshComponentChanged(UDynamicMeshComponent*
 
 	FDynamicMesh3* Mesh = GetSculptMesh();
 
-	// figure out the set of modified triangles
+	// Note we rely on the fact that Vertex Sculpt tool does not change topology, so we only compute
+	// CachedVtxTriConnectivity/CachedVtxVtxConnectivity once in Setup()
+	checkSlow(Mesh->MaxVertexID() == CachedVtxTriConnectivity.Num());
+
+	// figure out the set of modified vertices and triangles
 	AccumulatedTriangleROI.Reset();
-	Change->ProcessChangeVertices(Mesh, [this, &Mesh](TConstArrayView<int32> Vertices)
+	TArray<int32> ChangedVertices;
+	Change->ProcessChangeVertices(Mesh, [this, &Mesh, &ChangedVertices](TConstArrayView<int32> Vertices)
 	{
-		UE::Geometry::VertexToTriangleOneRing(Mesh, Vertices, AccumulatedTriangleROI);
+		AccumulatedTriangleROI.Reserve(Vertices.Num() * 4);
+		for (int32 VID : Vertices)
+		{
+			AccumulatedTriangleROI.Append(CachedVtxTriConnectivity[VID]);
+		}
+		ChangedVertices.Append(Vertices.GetData(), Vertices.Num());
 	}, bRevert);
 
 	UndoUpdateFuture.Reset();
@@ -1918,9 +1921,13 @@ void UMeshVertexSculptTool::OnDynamicMeshComponentChanged(UDynamicMeshComponent*
 	// Recompute normals synchronously, so we don't end up racing UDynamicMeshComponent's render buffer updates (which happen soon after this method)
 	UE::SculptUtil::RecalculateROINormals(Mesh, AccumulatedTriangleROI, NormalsROIBuilder);
 
-	// start the octree update
-	UndoUpdateOctreeFuture = Async(VertexSculptToolAsyncExecTarget, [this, Mesh]()
+	// start the octree updates
+	// Note: both UndoUpdateOctreeFuture and UndoUpdateBaseMeshFuture read AccumulatedTriangleROI concurrently (read-only after this point)
+	UndoUpdateOctreeFuture = Async(VertexSculptToolAsyncExecTarget, [this, Mesh, ChangedVerts = MoveTemp(ChangedVertices)]()
 	{
+		// Vertex octree reinsert changed vertices
+		VertexOctree.ReinsertVertices(ChangedVerts);
+		// Triangle octree reinsert affected triangles (for hit testing)
 		Octree.ReinsertTriangles(AccumulatedTriangleROI);
 		bOctreeUpdated = true;
 		return true;
diff --git a/Engine/Plugins/Runtime/MeshModelingToolset/Source/MeshModelingTools/Public/MeshVertexSculptTool.h b/Engine/Plugins/Runtime/MeshModelingToolset/Source/MeshModelingTools/Public/MeshVertexSculptTool.h
--- a/Engine/Plugins/Runtime/MeshModelingToolset/Source/MeshModelingTools/Public/MeshVertexSculptTool.h
+++ b/Engine/Plugins/Runtime/MeshModelingToolset/Source/MeshModelingTools/Public/MeshVertexSculptTool.h
@@ -8,6 +8,7 @@
 #include "Components/DynamicMeshComponent.h"
 #include "DynamicMesh/DynamicMeshAABBTree3.h"
 #include "DynamicMesh/DynamicMeshOctree3.h"
+#include "DynamicMesh/DynamicVerticesOctree3.h"
 #include "DynamicMesh/MeshNormals.h"
 #include "BaseTools/BaseBrushTool.h"
 #include "BaseTools/MeshSurfacePointMeshEditingTool.h"
@@ -394,9 +395,11 @@ protected:
 	void WaitForPendingStampUpdate() {};
 	UE_API void WaitForPendingStampUpdateConst() const;
 
+	UE_DEPRECATED(5.8, "RangeQueryTriBuffer is not used anymore")
 	TArray<int> RangeQueryTriBuffer;
 	UE::Geometry::FUniqueIndexSet VertexROIBuilder;
 	UE::Geometry::FUniqueIndexSet TriangleROIBuilder;
+	UE_DEPRECATED(5.8, "TriangleROIInBuf is not used anymore")
 	TArray<UE::Geometry::FIndex3i> TriangleROIInBuf;
 	TArray<int> VertexROI;
 	TArray<int> TriangleROIArray;
@@ -404,7 +407,7 @@ protected:
 	UE_DEPRECATED(5.6, "Use the FFrame3d overload instead.")
 	UE_API void UpdateROI(const FVector3d& BrushPos);
 private:
-	UE_API void UpdateRangeQueryTriBuffer(const FFrame3d& LocalFrame);
+	UE_API void ComputeVertexAndTriangleROI(const FFrame3d& LocalFrame);
 	UE_API void PrepROIVertPositionBuffers();
 protected:
 	virtual bool RequireConnectivityToHitPointInStamp() const { return false; }
@@ -445,6 +448,7 @@ protected:
 	UE_API bool GetBaseMeshNearest(int32 VertexID, const FVector3d& Position, double SearchRadius, FVector3d& TargetPosOut, FVector3d& TargetNormalOut);
 	TFunction<bool(int32, const FVector3d&, double MaxDist, FVector3d&, FVector3d&)> BaseMeshQueryFunc;
 
+	// triangle octree used for raycast queries; for ROI queries we use the VertexOctree (in private scope, below)
 	UE::Geometry::FDynamicMeshOctree3 Octree;
 	bool bOctreeUpdated = false;
 	TSharedPtr<UE::Geometry::FDynamicMeshOctree3::FTreeCutSet> CutTree;
@@ -489,6 +493,12 @@ private:
 	UE_API void RealignBrush(FMeshSculptBrushOp::EStampAlignmentType AlignmentType);
 
 	double InitialBoundsMaxDim = 0;
+	// Cached mesh connectivity for faster iteration
+	TArray<TArray<int32, TInlineAllocator<8>>> CachedVtxTriConnectivity, CachedVtxVtxConnectivity;
+	// Re-used index set for looking up the initial ROI during connected component growth
+	UE::Geometry::FUniqueIndexSet VertexROILookup;
+	// Octree used for vertex queries
+	UE::Geometry::TDynamicVerticesOctree3<FDynamicMesh3> VertexOctree;
 };
 
 
diff --git a/Engine/Source/Runtime/GeometryCore/Public/DynamicMesh/DynamicVerticesOctree3.h b/Engine/Source/Runtime/GeometryCore/Public/DynamicMesh/DynamicVerticesOctree3.h
--- a/Engine/Source/Runtime/GeometryCore/Public/DynamicMesh/DynamicVerticesOctree3.h
+++ b/Engine/Source/Runtime/GeometryCore/Public/DynamicMesh/DynamicVerticesOctree3.h
@@ -158,9 +158,10 @@ public:
 
 
 	/**
-	 * Reinsert a set of vertices into the tree
+	 * Reinsert Vertices into the tree
 	 */
-	void ReinsertVertices(const TSet<int>& Vertices)
+	template<typename EnumerableType>
+	void ReinsertVertices(const EnumerableType& Vertices)
 	{
 		for (int VertexID : Vertices)
 		{
@@ -170,6 +171,23 @@ public:
 		}
 	}
 
+	/**
+	 * Reinsert Vertices into the tree, unless ShouldSkip(VertexID) returns true
+	 */
+	template<typename EnumerableType>
+	void ReinsertFilteredVertices(const EnumerableType& Vertices, TFunctionRef<bool(int32)> ShouldSkip)
+	{
+		for (int32 VertexID : Vertices)
+		{
+			if (!ShouldSkip(VertexID))
+			{
+				FVector3d Point = VertexSource->GetVertex(VertexID);
+				ModifiedBounds.Contain(Point);
+				FSparseDynamicPointOctree3::ReinsertPoint(VertexID, Point);
+			}
+		}
+	}
+
 
 	/**
 	 * Include the current bounds of a triangle in the ModifiedBounds box
