Low-hanging optimizations for MeshAttributePaintTool

#Synced-CL 52243346

diff --git a/Engine/Plugins/Runtime/MeshModelingToolset/Source/MeshModelingTools/Private/MeshAttributePaintTool.cpp b/Engine/Plugins/Runtime/MeshModelingToolset/Source/MeshModelingTools/Private/MeshAttributePaintTool.cpp
--- a/Engine/Plugins/Runtime/MeshModelingToolset/Source/MeshModelingTools/Private/MeshAttributePaintTool.cpp
+++ b/Engine/Plugins/Runtime/MeshModelingToolset/Source/MeshModelingTools/Private/MeshAttributePaintTool.cpp
@@ -29,6 +29,16 @@ using namespace UE::Geometry;
 
 #define LOCTEXT_NAMESPACE "UMeshAttributePaintTool"
 
+namespace
+{
+	// probably should be something defined for the whole tool framework...
+#if WITH_EDITOR
+	static EAsyncExecution AttributePaintToolAsyncExecTarget = EAsyncExecution::LargeThreadPool;
+#else
+	static EAsyncExecution AttributePaintToolAsyncExecTarget = EAsyncExecution::ThreadPool;
+#endif
+}
+
 class FMeshDescriptionVertexAttributeAdapter : public IMeshVertexAttributeAdapter
 {
 public:
@@ -353,7 +363,7 @@ void UMeshAttributePaintTool::Setup()
 		Mesh.Attributes()->DisablePrimaryColors();
 		Mesh.Attributes()->EnablePrimaryColors();
 		// Create an overlay that has no split elements, init with zero value.
-		Mesh.Attributes()->PrimaryColors()->CreatePerVertex(0.f);
+		Mesh.Attributes()->PrimaryColors()->CreatePerVertex(0.f, true);
 
 		if (GeometrySelection.IsSet())
 		{
@@ -392,6 +402,9 @@ void UMeshAttributePaintTool::Setup()
 		}
 	}
 
+	// Enable render mesh decomposition for efficient partial updates during painting
+	PreviewMesh->SetEnableRenderMeshDecomposition(true);
+
 	UpdateMaterialMode(ViewProperties->MaterialMode);
 
 	RecalculateBrushRadius();
@@ -587,7 +600,7 @@ void UMeshAttributePaintTool::CalculateVertexROI(const FBrushStampData& Stamp, T
 			LocalEyePosition = (CurrentTargetTransform.InverseTransformPosition((FVector3d)StateOut.Position));
 		}
 		
-		VerticesOctree.RangeQuery(QueryBox,
+		VerticesOctree.ParallelRangeQuery(QueryBox,
 			[this, Mesh, &StampPosLocal, RadiusSqr, &LocalEyePosition](int32 VertexID)
 		{
 			const bool bIsSelectionValid = !IsFilteringTrianglesBySelection() || SelectionVids.Contains(VertexID);
@@ -609,7 +622,7 @@ void UMeshAttributePaintTool::CalculateVertexROI(const FBrushStampData& Stamp, T
 	}
 	else
 	{
-		VerticesOctree.RangeQuery(QueryBox,
+		VerticesOctree.ParallelRangeQuery(QueryBox,
 			[Mesh, &StampPosLocal, RadiusSqr](int32 VertexID) 
 		{ 
 			return DistanceSquared(Mesh->GetVertex(VertexID), StampPosLocal) < RadiusSqr; 
@@ -695,13 +708,14 @@ void UMeshAttributePaintTool::UpdateVisibleAttribute()
 		{
 			FDynamicMeshColorOverlay* ColorOverlay = Mesh.Attributes()->PrimaryColors();
 			FNonManifoldMappingSupport NonManifoldMappingSupport(Mesh);
-			for (int32 elid : ColorOverlay->ElementIndicesItr())
+			for (int32 vid : Mesh.VertexIndicesItr())
 			{
-				const int32 vid = ColorOverlay->GetParentVertex(elid);
 				const int32 srcvid = NonManifoldMappingSupport.GetOriginalNonManifoldVertexID(vid);
 				const float Value = AttribData.CurrentValues[srcvid];
 				const FVector4f Color4f = ToVector4<float>(ColorMapper->ToColor(Value));
-				ColorOverlay->SetElement(elid,  Color4f);
+				// Note: This tool specifically constructs the color overlay element IDs to exactly match vertex IDs
+				checkSlow(ColorOverlay->GetParentVertex(vid) == vid);
+				ColorOverlay->SetElement(vid,  Color4f);
 			}
 		});
 
@@ -780,9 +794,13 @@ void UMeshAttributePaintTool::ApplyStamp(const FBrushStampData& Stamp)
 
 
 	// track changes
+	TFuture<void> ChangeBuilderCompute;
 	if (ActiveChangeBuilder)
 	{
-		ActiveChangeBuilder->UpdateValues(ActionData.ROIVertices, ActionData.ROIBefore, ActionData.ROIAfter);
+		ChangeBuilderCompute = Async(AttributePaintToolAsyncExecTarget, [this, &ActionData]()
+		{
+			ActiveChangeBuilder->UpdateValues(ActionData.ROIVertices, ActionData.ROIBefore, ActionData.ROIAfter);
+		});
 	}
 
 	// update values and colors
@@ -797,12 +815,10 @@ void UMeshAttributePaintTool::ApplyStamp(const FBrushStampData& Stamp)
 			int32 vid = ActionData.ROIVertices[k];
 			int32 srcvid = NonManifoldMappingSupport.GetOriginalNonManifoldVertexID(vid);
 			AttribData.CurrentValues[srcvid] = ActionData.ROIAfter[k];
-			FVector4f NewColor( ToVector4<float>(ColorMapper->ToColor(ActionData.ROIAfter[k])) );
-			ColorOverlay->GetVertexElements(vid, ElIDs);
-			for (int elid : ElIDs)
-			{
-				ColorOverlay->SetElement(elid, NewColor);
-			}
+			FVector4f NewColor(ToVector4<float>(ColorMapper->ToColor(ActionData.ROIAfter[k])));
+			// Note: This tool specifically constructs the color overlay element IDs to exactly match vertex IDs
+			checkSlow(ColorOverlay->GetParentVertex(vid) == vid);
+			ColorOverlay->SetElement(vid, NewColor);
 			AttribData.Attribute->SetValue(vid, AttribData.CurrentValues[srcvid]);
 			ElIDs.Reset();
 		}
@@ -812,7 +828,12 @@ void UMeshAttributePaintTool::ApplyStamp(const FBrushStampData& Stamp)
 	if (ViewProperties->MaterialMode == EMeshAttributePaintMaterialMode::Shaded
 		|| ViewProperties->MaterialMode == EMeshAttributePaintMaterialMode::ColorOnly)
 	{
-		PreviewMesh->NotifyDeferredEditCompleted(UPreviewMesh::ERenderUpdateMode::FastUpdate, EMeshRenderAttributeFlags::VertexColors, false);
+		PreviewMesh->NotifyVertexRegionDeferredEditCompleted(ActionData.ROIVertices, EMeshRenderAttributeFlags::VertexColors);
+	}
+
+	if (ActiveChangeBuilder)
+	{
+		ChangeBuilderCompute.Wait();
 	}
 }
 
@@ -1107,18 +1128,14 @@ void UMeshAttributePaintTool::ExternalUpdateValues(int32 AttribIndex, const TArr
 	{
 		PreviewMesh->EditMesh([&](FDynamicMesh3& Mesh)
 		{
-			TArray<int> ElIDs;
 			FDynamicMeshColorOverlay* ColorOverlay = Mesh.Attributes()->PrimaryColors();
 			for (int32 vid : VertexIndices)
 			{
 				FVector4f NewColor( ToVector4<float>(ColorMapper->ToColor(AttribData.CurrentValues[vid])) );
-				ColorOverlay->GetVertexElements(vid, ElIDs);
-				for (int elid : ElIDs)
-				{
-					ColorOverlay->SetElement(elid, NewColor);
-				}
+				// Note: This tool specifically constructs the color overlay element IDs to exactly match vertex IDs
+				checkSlow(ColorOverlay->GetParentVertex(vid) == vid);
+				ColorOverlay->SetElement(vid, NewColor);
 				AttribData.Attribute->SetValue(vid, AttribData.CurrentValues[vid]);
-				ElIDs.Reset();
 			}
 		});
 	}
