fix open/close transaction mismatch bug in attribute paint tool after attempting to paint on a mesh with no attribute channels

#Synced-CL 52243346

diff --git a/Engine/Plugins/Runtime/MeshModelingToolset/Source/MeshModelingTools/Private/MeshAttributePaintTool.cpp b/Engine/Plugins/Runtime/MeshModelingToolset/Source/MeshModelingTools/Private/MeshAttributePaintTool.cpp
--- a/Engine/Plugins/Runtime/MeshModelingToolset/Source/MeshModelingTools/Private/MeshAttributePaintTool.cpp
+++ b/Engine/Plugins/Runtime/MeshModelingToolset/Source/MeshModelingTools/Private/MeshAttributePaintTool.cpp
@@ -546,10 +546,12 @@ void UMeshAttributePaintTool::OnEndDrag(const FRay& Ray)
 	bStampPending = false;
 
 	// close change record
-	TUniquePtr<FMeshAttributePaintChange> Change = EndChange();
-	if (Change)
+	if (ActiveChangeBuilder)
 	{
-		GetToolManager()->EmitObjectChange(this, MoveTemp(Change), LOCTEXT("AttributeValuesChange", "Paint"));
+		if (TUniquePtr<FMeshAttributePaintChange> Change = EndChange())
+		{
+			GetToolManager()->EmitObjectChange(this, MoveTemp(Change), LOCTEXT("AttributeValuesChange", "Paint"));
+		}
 		LongTransactions.Close(GetToolManager());
 	}
 }
@@ -1069,6 +1071,7 @@ void UMeshAttributePaintTool::BeginChange()
 {
 	if (CurrentAttributeIndex < 0)
 	{
+		ActiveChangeBuilder.Reset();
 		return;
 	}
 	if (! ActiveChangeBuilder)
@@ -1084,7 +1087,7 @@ void UMeshAttributePaintTool::BeginChange()
 
 TUniquePtr<FMeshAttributePaintChange> UMeshAttributePaintTool::EndChange()
 {
-	if (!ActiveChangeBuilder)
+	if (!ensure(ActiveChangeBuilder)) // should be validated before calling EndChange
 	{
 		return nullptr;
 	}
