make vertex sculpt tool update rendered mesh positions+normals after external changes (that create change transactions)

#Synced-CL 52311299

diff --git a/Engine/Plugins/Runtime/MeshModelingToolset/Source/MeshModelingTools/Private/MeshVertexSculptTool.cpp b/Engine/Plugins/Runtime/MeshModelingToolset/Source/MeshModelingTools/Private/MeshVertexSculptTool.cpp
--- a/Engine/Plugins/Runtime/MeshModelingToolset/Source/MeshModelingTools/Private/MeshVertexSculptTool.cpp
+++ b/Engine/Plugins/Runtime/MeshModelingToolset/Source/MeshModelingTools/Private/MeshVertexSculptTool.cpp
@@ -516,7 +516,8 @@ void UMeshVertexSculptTool::UpdateToolMeshes(TFunctionRef<TUniquePtr<FMeshRegion
 		if (Change)
 		{
 			// pass through the change to trigger standard mesh updates / octree recomputation
-			OnDynamicMeshComponentChanged(DynamicMeshComponent, Change.Get(), false);
+			constexpr bool bUpdateRender = true;
+			OnDynamicMeshComponentChanged_Internal(DynamicMeshComponent, Change.Get(), false, bUpdateRender);
 
 			TUniquePtr<TWrappedToolCommandChange<FMeshRegionChangeBase>> NewChange = MakeUnique<TWrappedToolCommandChange<FMeshRegionChangeBase>>();
 			NewChange->WrappedChange = MoveTemp(Change);
@@ -1891,6 +1892,10 @@ void UMeshVertexSculptTool::WaitForPendingUndoRedoUpdate() const
 }
 
 void UMeshVertexSculptTool::OnDynamicMeshComponentChanged(UDynamicMeshComponent* Component, const FMeshRegionChangeBase* Change, bool bRevert)
+{
+	OnDynamicMeshComponentChanged_Internal(Component, Change, bRevert, /*bUpdateRender*/false);
+}
+void UMeshVertexSculptTool::OnDynamicMeshComponentChanged_Internal(UDynamicMeshComponent* Component, const FMeshRegionChangeBase* Change, bool bRevert, bool bUpdateRender)
 {
 	// have to wait for any outstanding stamp update to finish...
 	WaitForPendingStampUpdateConst();
@@ -1921,6 +1926,11 @@ void UMeshVertexSculptTool::OnDynamicMeshComponentChanged(UDynamicMeshComponent*
 	// Recompute normals synchronously, so we don't end up racing UDynamicMeshComponent's render buffer updates (which happen soon after this method)
 	UE::SculptUtil::RecalculateROINormals(Mesh, AccumulatedTriangleROI, NormalsROIBuilder);
 
+	if (bUpdateRender)
+	{
+		Component->FastNotifyTriangleVerticesUpdated(AccumulatedTriangleROI, EMeshRenderAttributeFlags::Positions | EMeshRenderAttributeFlags::VertexNormals);
+	}
+
 	// start the octree updates
 	// Note: both UndoUpdateOctreeFuture and UndoUpdateBaseMeshFuture read AccumulatedTriangleROI concurrently (read-only after this point)
 	UndoUpdateOctreeFuture = Async(VertexSculptToolAsyncExecTarget, [this, Mesh, ChangedVerts = MoveTemp(ChangedVertices)]()
diff --git a/Engine/Plugins/Runtime/MeshModelingToolset/Source/MeshModelingTools/Public/MeshVertexSculptTool.h b/Engine/Plugins/Runtime/MeshModelingToolset/Source/MeshModelingTools/Public/MeshVertexSculptTool.h
--- a/Engine/Plugins/Runtime/MeshModelingToolset/Source/MeshModelingTools/Public/MeshVertexSculptTool.h
+++ b/Engine/Plugins/Runtime/MeshModelingToolset/Source/MeshModelingTools/Public/MeshVertexSculptTool.h
@@ -492,6 +492,9 @@ private:
 	// Assumes that brush is currently aligned to hit normal
 	UE_API void RealignBrush(FMeshSculptBrushOp::EStampAlignmentType AlignmentType);
 
+	// internal implementation of OnDynamicMeshComponentChanged, exposing option to update component rendering as well
+	void OnDynamicMeshComponentChanged_Internal(UDynamicMeshComponent* Component, const FMeshRegionChangeBase* Change, bool bRevert, bool bUpdateRender = false);
+
 	double InitialBoundsMaxDim = 0;
 	// Cached mesh connectivity for faster iteration
 	TArray<TArray<int32, TInlineAllocator<8>>> CachedVtxTriConnectivity, CachedVtxVtxConnectivity;
