[SkeletalMeshEditingTools] Improve mirroring bones so it replaces strings from right to left. For example bone_lower_l_01 is now replaced with bone_lower_r_01, instead of faulty bone_rower_l_01

#Synced-CL 52398802

diff --git a/Engine/Plugins/Runtime/MeshModelingToolset/Source/SkeletalMeshModifiers/Private/SkeletonModifier.cpp b/Engine/Plugins/Runtime/MeshModelingToolset/Source/SkeletalMeshModifiers/Private/SkeletonModifier.cpp
--- a/Engine/Plugins/Runtime/MeshModelingToolset/Source/SkeletalMeshModifiers/Private/SkeletonModifier.cpp
+++ b/Engine/Plugins/Runtime/MeshModelingToolset/Source/SkeletalMeshModifiers/Private/SkeletonModifier.cpp
@@ -1406,16 +1406,15 @@ void USkeletonModifier::GetBonesToMirror(
 
 void USkeletonModifier::GetMirroredNames(
 	const TArray<int32>& InBonesToMirror, const FMirrorOptions& InOptions, TArray<FName>& OutBonesName) const
-{
+{	
+	// Similar to UMirrorDataTable::GetMirrorName but doing a search from the end if the find string is contained
+
 	OutBonesName.Reset();
 	if (InBonesToMirror.IsEmpty())
 	{
 		return;
 	}
 
-	const FName LeftName(InOptions.LeftString);
-	const FName RightName(InOptions.RightString);
-
 	const FName LeftFindRegex = *FString::Printf(TEXT("([^}]*)%s([^}]*)"), *InOptions.LeftString);
 	const FName RightReplaceRegex = *FString::Printf(TEXT("$1%s$2"), *InOptions.RightString);
 
@@ -1423,29 +1422,66 @@ void USkeletonModifier::GetMirroredNames(
 	const FName LeftReplaceRegex = *FString::Printf(TEXT("$1%s$2"), *InOptions.LeftString);
 
 	const TArray MirrorFindReplaceExpressions({
-		FMirrorFindReplaceExpression(LeftName, RightName, EMirrorFindReplaceMethod::Prefix),
-		FMirrorFindReplaceExpression(RightName, LeftName, EMirrorFindReplaceMethod::Prefix),
-		FMirrorFindReplaceExpression(LeftName, RightName, EMirrorFindReplaceMethod::Suffix),
-		FMirrorFindReplaceExpression(RightName, LeftName, EMirrorFindReplaceMethod::Suffix),
 		FMirrorFindReplaceExpression(LeftFindRegex, RightReplaceRegex, EMirrorFindReplaceMethod::RegularExpression),
 		FMirrorFindReplaceExpression(RightFindRegex, LeftReplaceRegex, EMirrorFindReplaceMethod::RegularExpression)
 	});
 		
 	const TArray<FMeshBoneInfo>& BoneInfos = ReferenceSkeleton->GetRawRefBoneInfo();
 	
+	const auto FindReplaceFromEnd = [](const FName& InBoneName, const FString& InFind, const FString& InReplace, FName& OutResult)
+		{
+			if (InBoneName.IsNone() ||
+				InFind.IsEmpty() ||
+				InReplace.IsEmpty())
+			{
+				// Since replacing can not occur, result to NAME_None implying a new name should be generated
+				OutResult = NAME_None;
+				return true;
+			}
+
+			const FString BoneNameString = InBoneName.ToString();
+
+			const int32 ReplaceIndex = BoneNameString.Find(InFind, ESearchCase::CaseSensitive, ESearchDir::FromEnd);
+			if (ReplaceIndex != INDEX_NONE)
+			{
+				FString ReplaceBoneName = BoneNameString;
+				ReplaceBoneName.RemoveAt(ReplaceIndex, InFind.Len());
+				ReplaceBoneName.InsertAt(ReplaceIndex, InReplace);
+
+				OutResult = *ReplaceBoneName;
+				return true;
+			}
+
+			return false;
+		};
+
 	OutBonesName.Reserve(InBonesToMirror.Num());
 	Algo::Transform(InBonesToMirror, OutBonesName, [&](const int32 BoneIndex)
-	{
-		const FName BoneName = BoneInfos[BoneIndex].Name;
-		const FName MirrorName = UMirrorDataTable::GetMirrorName(BoneName, MirrorFindReplaceExpressions);
-
-		if (MirrorName.IsNone() || MirrorName == BoneName)
 		{
-			return GetUniqueName(BoneName, OutBonesName);
-		}
+			const FName BoneName = BoneInfos[BoneIndex].Name;
 
-		return MirrorName;
-	});
+			const FName MirrorName = [&BoneName, &InOptions, &FindReplaceFromEnd, &MirrorFindReplaceExpressions]()
+			{
+				// Try to replace with an exact match from end, and if successful prefer it over 
+				// UMirrorDataTable::GetMirrorName that does look ups from the front.
+				// Still use UMirrorDataTable::GetMirrorName to support regular expressions.
+				FName MirrorName;
+				if (!FindReplaceFromEnd(BoneName, InOptions.LeftString, InOptions.RightString, MirrorName) &&
+					!FindReplaceFromEnd(BoneName, InOptions.RightString, InOptions.LeftString, MirrorName))
+				{
+					MirrorName = UMirrorDataTable::GetMirrorName(BoneName, MirrorFindReplaceExpressions);
+				}
+
+				return MirrorName;
+			}(); 
+
+			if (MirrorName.IsNone() || MirrorName == BoneName)
+			{
+				return GetUniqueName(BoneName, OutBonesName);
+			}
+
+			return MirrorName;
+		});
 }
 
 void USkeletonModifier::GetMirroredBones(
