[SkeletalMeshToolTarget] Making sure prechange is called for skeletal meshes being modified by tools

#Synced-CL 52406569

diff --git a/Engine/Plugins/Runtime/MeshModelingToolset/Source/ModelingComponentsEditorOnly/Private/ToolTargets/SkeletalMeshToolTarget.cpp b/Engine/Plugins/Runtime/MeshModelingToolset/Source/ModelingComponentsEditorOnly/Private/ToolTargets/SkeletalMeshToolTarget.cpp
--- a/Engine/Plugins/Runtime/MeshModelingToolset/Source/ModelingComponentsEditorOnly/Private/ToolTargets/SkeletalMeshToolTarget.cpp
+++ b/Engine/Plugins/Runtime/MeshModelingToolset/Source/ModelingComponentsEditorOnly/Private/ToolTargets/SkeletalMeshToolTarget.cpp
@@ -238,6 +238,8 @@ void USkeletalMeshToolTarget::CommitMeshDescription(USkeletalMesh* SkeletalMesh,
 		return;
 	}
 
+	// This scope will invoke the skeletal mesh render data rebuild machinery on scope exit
+	FScopedSkeletalMeshPostEditChange ScopedPostEditChange(SkeletalMesh);
 	// flush any pending rendering commands, which might touch a component while we are rebuilding it's mesh
 	FlushRenderingCommands();
 
@@ -259,8 +261,6 @@ void USkeletalMeshToolTarget::CommitMeshDescription(USkeletalMesh* SkeletalMesh,
 	// the actual copying from SourceModel to LODModel happens during build time
 	FLODUtilities::UpdateLODInfoVertexAttributes(SkeletalMesh, LODIndex, LODIndex, false);
 	
-	// This call will invoke the skeletal mesh render data rebuild machinery.
-	UE::ToolTarget::Internal::PostEditChangeWithConditionalUndo(SkeletalMesh);
 }
 
 void USkeletalMeshToolTarget::CommitDynamicMesh(const FDynamicMesh3& Mesh, const FDynamicMeshCommitInfo& CommitInfo)
