add option for smooth brush to favor polygroup edges, so that meshes with quads will better preserve the quad shape on smoothing

#Synced-CL 52426395

diff --git a/Engine/Plugins/Runtime/MeshModelingToolset/Source/MeshModelingTools/Public/Sculpting/MeshSmoothingBrushOps.h b/Engine/Plugins/Runtime/MeshModelingToolset/Source/MeshModelingTools/Public/Sculpting/MeshSmoothingBrushOps.h
--- a/Engine/Plugins/Runtime/MeshModelingToolset/Source/MeshModelingTools/Public/Sculpting/MeshSmoothingBrushOps.h
+++ b/Engine/Plugins/Runtime/MeshModelingToolset/Source/MeshModelingTools/Public/Sculpting/MeshSmoothingBrushOps.h
@@ -16,6 +16,8 @@ class UBaseSmoothBrushOpProps : public UMeshSculptBrushOpProps
 	GENERATED_BODY()
 public:
 	virtual bool GetPreserveUVFlow() { return false; }
+	// @return a float value in the [0,1] range controlling the desired relative weight of PolyGroup edges in smoothing (1 = maximally favor PolyGroup edges, 0 = ignore PolyGroup edges)
+	virtual float GetPolyGroupEdgesWeight() { return 0.f; }
 	virtual bool SupportsStrengthPressure() override { return true; }
 	virtual int32 GetIterations() { return 1; }
 };
@@ -41,10 +43,15 @@ public:
 	UPROPERTY(EditAnywhere, Category = SmoothBrush)
 	bool bPreserveUVFlow = true;
 
+	/** If true, favor smoothing along PolyGroup edges where they are present. Can help to preserve quads. No effect where PolyGroups are not defined. */
+	UPROPERTY(EditAnywhere, Category = SmoothBrush)
+	bool bFavorPolyGroupEdges = false;
+
 	virtual float GetStrength() override { return Strength; }
 	virtual void SetStrength(float NewStrength) override { Strength = FMathf::Clamp(NewStrength, 0.0f, 1.0f); }
 	virtual float GetFalloff() override { return Falloff; }
 	virtual bool GetPreserveUVFlow() override { return bPreserveUVFlow; }
+	virtual float GetPolyGroupEdgesWeight() override { return (float)bFavorPolyGroupEdges; }
 	virtual int32 GetIterations() { return Iterations; }
 };
 
@@ -70,10 +77,15 @@ public:
 	UPROPERTY(EditAnywhere, Category = ShiftSmoothBrush)
 	bool bPreserveUVFlow = true;
 
+	/** If true, favor smoothing along PolyGroup edges where they are present. Can help to preserve quads. No effect where PolyGroups are not defined. */
+	UPROPERTY(EditAnywhere, Category = ShiftSmoothBrush)
+	bool bFavorPolyGroupEdges = false;
+
 	virtual float GetStrength() override { return Strength; }
 	virtual void SetStrength(float NewStrength) override { Strength = FMathf::Clamp(NewStrength, 0.0f, 1.0f); }
 	virtual float GetFalloff() override { return Falloff; }
 	virtual bool GetPreserveUVFlow() override { return bPreserveUVFlow; }
+	virtual float GetPolyGroupEdgesWeight() override { return (float)bFavorPolyGroupEdges; }
 	virtual int32 GetIterations() { return Iterations; }
 };
 
@@ -144,6 +156,27 @@ public:
 			{
 				SmoothedPos = UE::Geometry::FMeshWeights::UniformCentroid(*Mesh, VID, GetVertexFn);
 			}
+			double PolyGroupEdgesWeight = (double)GetPropertySetAs<UBaseSmoothBrushOpProps>()->GetPolyGroupEdgesWeight();
+			if (PolyGroupEdgesWeight > 0 && Mesh->HasTriangleGroups())
+			{
+				FVector3d PolyGroupNbrSum = FVector3d::Zero();
+				int32 FoundGroupEdges = 0;
+				Mesh->EnumerateVertexEdges(VID, [&PolyGroupNbrSum, &FoundGroupEdges, &Mesh, &GetVertexFn, VID](int32 EID)
+				{
+					const FDynamicMesh3::FEdge& Edge = Mesh->GetEdgeRef(EID);
+					if (Mesh->GetTriangleGroup(Edge.Tri.A) != Mesh->GetTriangleGroup(Edge.Tri.B))
+					{
+						FoundGroupEdges++;
+						PolyGroupNbrSum += GetVertexFn(Edge.Vert.OtherElement(VID));
+					}
+				});
+				if (FoundGroupEdges > 0)
+				{
+					double UseGroupEdgeWeight = FMath::Min(1., PolyGroupEdgesWeight);
+					SmoothedPos = UE::Geometry::Lerp(SmoothedPos, PolyGroupNbrSum / (double)FoundGroupEdges, UseGroupEdgeWeight);
+				}
+			}
+			
 			return UE::Geometry::Lerp(OrigPos, SmoothedPos, Falloff * Stamp.Power);
 		});
 	}
