[SkeletalMeshModellingTools] The Skeletal Mesh Editor now shows a linear color ramp instead of an array of colors

#Synced-CL 52705616

diff --git a/Engine/Plugins/Animation/SkeletalMeshModelingTools/Source/SkeletalMeshModelingTools/Private/SkeletalMeshModelingToolsEditorModeToolkit.cpp b/Engine/Plugins/Animation/SkeletalMeshModelingTools/Source/SkeletalMeshModelingTools/Private/SkeletalMeshModelingToolsEditorModeToolkit.cpp
--- a/Engine/Plugins/Animation/SkeletalMeshModelingTools/Source/SkeletalMeshModelingTools/Private/SkeletalMeshModelingToolsEditorModeToolkit.cpp
+++ b/Engine/Plugins/Animation/SkeletalMeshModelingTools/Source/SkeletalMeshModelingTools/Private/SkeletalMeshModelingToolsEditorModeToolkit.cpp
@@ -4,6 +4,8 @@
 
 #include "SkeletalMeshModelingToolsCommands.h"
 
+#include "Curves/LinearColorRamp.h"
+#include "Customizations/LinearColorRampCustomization.h"
 #include "Tools/EdModeInteractiveToolsContext.h"
 #include "Features/IModularFeatures.h"
 #include "Framework/MultiBox/MultiBoxBuilder.h"
@@ -98,6 +100,15 @@ void FSkeletalMeshModelingToolsEditorModeToolkit::Init(
 	
 	// create the toolkit widget
 	{
+		if (ModeDetailsView.IsValid())
+		{
+			constexpr FLinearColorRampCustomizationArgs Args({ .bWithAlphaChannel = false });
+
+			ModeDetailsView->RegisterInstancedCustomPropertyTypeLayout(
+				FLinearColorRamp::StaticStruct()->GetFName(),
+				FOnGetPropertyTypeCustomizationInstance::CreateStatic(&FLinearColorRampCustomization::MakeInstance, Args));
+		}
+
 		ToolkitSections->ModeWarningArea = ModeWarningArea;
 		ToolkitSections->DetailsView = ModeDetailsView;
 		ToolkitSections->ToolWarningArea = ToolWarningArea;
diff --git a/Engine/Plugins/Animation/SkeletalMeshModelingTools/Source/SkeletalMeshModelingTools/SkeletalMeshModelingTools.Build.cs b/Engine/Plugins/Animation/SkeletalMeshModelingTools/Source/SkeletalMeshModelingTools/SkeletalMeshModelingTools.Build.cs
--- a/Engine/Plugins/Animation/SkeletalMeshModelingTools/Source/SkeletalMeshModelingTools/SkeletalMeshModelingTools.Build.cs
+++ b/Engine/Plugins/Animation/SkeletalMeshModelingTools/Source/SkeletalMeshModelingTools/SkeletalMeshModelingTools.Build.cs
@@ -12,6 +12,7 @@ public class SkeletalMeshModelingTools : ModuleRules
 				"ContentBrowser",
 				"Core",
 				"CoreUObject",
+				"DetailCustomizations",
 				"EditorFramework",
 				"EditorInteractiveToolsFramework",
 				"Engine",
diff --git a/Engine/Plugins/Runtime/MeshModelingToolset/Source/MeshModelingToolsEditorOnly/Private/SkeletalMesh/SkinWeightsPaintTool.cpp b/Engine/Plugins/Runtime/MeshModelingToolset/Source/MeshModelingToolsEditorOnly/Private/SkeletalMesh/SkinWeightsPaintTool.cpp
--- a/Engine/Plugins/Runtime/MeshModelingToolset/Source/MeshModelingToolsEditorOnly/Private/SkeletalMesh/SkinWeightsPaintTool.cpp
+++ b/Engine/Plugins/Runtime/MeshModelingToolset/Source/MeshModelingToolsEditorOnly/Private/SkeletalMesh/SkinWeightsPaintTool.cpp
@@ -67,6 +67,7 @@
 #include "Widgets/Images/SImage.h"
 #include "Misc/MessageDialog.h"
 #include "SkeletalMesh/SReferenceSkeletonTree.h"
+#include "SkeletalMesh/SkinWeightsPaintToolPropertiesCustomVersion.h"
 
 
 #include UE_INLINE_GENERATED_CPP_BY_NAME(SkinWeightsPaintTool)
@@ -570,6 +571,7 @@ float FDirectEditWeightState::GetModeMaxValue()
 }
 
 USkinWeightsPaintToolProperties::USkinWeightsPaintToolProperties()
+	: CustomVersion(ESkinWeightsPaintToolPropertiesCustomVersion::NoVersion)
 {
 	BrushConfigs.Add(EWeightEditOperation::Add, &BrushConfigAdd);
 	BrushConfigs.Add(EWeightEditOperation::Replace, &BrushConfigReplace);
@@ -578,19 +580,68 @@ USkinWeightsPaintToolProperties::USkinWeightsPaintToolProperties()
 
 	LoadConfig();
 
-	if (ColorRamp.IsEmpty())
+	if (LinearColorRamp.IsEmpty())
 	{
-		// default color ramp simulates a heat map
-		ColorRamp.Add(FLinearColor(0.8f, 0.4f, 0.8f)); // Purple
-		ColorRamp.Add(FLinearColor(0.0f, 0.0f, 0.5f)); // Dark Blue
-		ColorRamp.Add(FLinearColor(0.2f, 0.2f, 1.0f)); // Light Blue
-		ColorRamp.Add(FLinearColor(0.0f, 1.0f, 0.0f)); // Green
-		ColorRamp.Add(FLinearColor(1.0f, 1.0f, 0.0f)); // Yellow
-		ColorRamp.Add(FLinearColor(1.0f, 0.65f, 0.0f)); // Orange
-		ColorRamp.Add(FLinearColor(1.0f, 0.0f, 0.0f, 0.0f)); // Red
+		const TArray<FLinearColor> DefaultColorRamp
+		{
+			FLinearColor(0.8f, 0.4f, 0.8f),	// Purple
+			FLinearColor(0.f, 0.f, 0.5f),	// Dark Blue
+			FLinearColor(0.2f, 0.2f, 1.0f),	// Light Blue
+			FLinearColor(0.f, 1.f, 0.f),	// Green
+			FLinearColor(1.f, 1.f, 0.f), 	// Yellow
+			FLinearColor(1.f, 0.65f, 0.f),	// Orange
+			FLinearColor(1.f, 0.f, 0.f)		// Red
+		};
+
+		float Time = 0.f;
+		const float Delta = 1.f / (DefaultColorRamp.Num() - 1);
+		for (const FLinearColor& LinearColor : DefaultColorRamp)
+		{
+			constexpr bool bOnlyRGB = true;
+			LinearColorRamp.SetColorAtTime(Time, LinearColor, bOnlyRGB);
+			
+			Time += Delta;
+		}
 	}
 }
 
+void USkinWeightsPaintToolProperties::RestoreProperties(UInteractiveTool* RestoreToTool, const FString& CacheIdentifier)
+{
+	Super::RestoreProperties(RestoreToTool, CacheIdentifier);
+
+	// Upgrade SkinWeightsPaintToolProperties to use FLinearColorRamp instead of an array of linear colors
+	if (CustomVersion < ESkinWeightsPaintToolPropertiesCustomVersion::SkinWeightPaintToolUseLinearColorRamp)
+	{
+		PRAGMA_DISABLE_DEPRECATION_WARNINGS		
+		if (!ColorRamp.IsEmpty())
+		{
+			LinearColorRamp.Reset();
+		
+			constexpr bool bOnlyRGB = true;
+			LinearColorRamp.SetColorAtTime(0.f, ColorRamp[0], bOnlyRGB);
+			
+			const float DeltaTime = ColorRamp.Num() > 1 ? 1.f / (ColorRamp.Num() - 1) : 0.f;
+			for (int32 RampStepIndex = 1; RampStepIndex < ColorRamp.Num() - 1; RampStepIndex++)
+			{
+				if (!FMath::Lerp(ColorRamp[RampStepIndex - 1], ColorRamp[RampStepIndex + 1], .5f).Equals(ColorRamp[RampStepIndex]))
+				{
+					// SetColorAtTime doesn't adopt alpha 1, hence force it here
+					const FLinearColor Color = FLinearColor(ColorRamp[RampStepIndex]).CopyWithNewOpacity(1.f);
+					LinearColorRamp.SetColorAtTime(DeltaTime * RampStepIndex, Color, bOnlyRGB);
+				}
+			}
+
+			LinearColorRamp.SetColorAtTime(1.f, FLinearColor(ColorRamp.Last()).CopyWithNewOpacity(1.f), bOnlyRGB);
+		}
+		PRAGMA_ENABLE_DEPRECATION_WARNINGS
+
+		CustomVersion = ESkinWeightsPaintToolPropertiesCustomVersion::SkinWeightPaintToolUseLinearColorRamp;
+		SaveProperties(RestoreToTool, CacheIdentifier);
+		SaveConfig();
+	}
+
+}
+
 namespace SkinWeightLayer
 {
 
@@ -2247,6 +2298,9 @@ void USkinWeightsPaintTool::Setup()
 	WeightToolProperties->SourceSkeletalMesh = nullptr;
     WeightToolProperties->SourcePreviewOffset = FTransform::Identity;
 		
+	// listen to linear color ramp changes
+	WeightToolProperties->LinearColorRamp.OnColorCurveChangedDelegate.AddUObject(this, &USkinWeightsPaintTool::OnLinearColorRampChanged);
+
 	// replace the base brush properties
 	ReplaceToolPropertySource(BrushProperties, WeightToolProperties);
 	BrushProperties = WeightToolProperties;
@@ -2309,6 +2363,11 @@ void USkinWeightsPaintTool::Setup()
 
 void USkinWeightsPaintTool::Shutdown(EToolShutdownType ShutdownType) 
 {
+	if (ensure(WeightToolProperties))
+	{
+		WeightToolProperties->LinearColorRamp.OnColorCurveChangedDelegate.RemoveAll(this);
+	}
+
 	SkeletonTreeNotifierBindScope.Reset();
 	
 	if (ShutdownType == EToolShutdownType::Accept)
@@ -2571,6 +2630,15 @@ void USkinWeightsPaintTool::OnTick(float DeltaTime)
 		bVertexColorsNeedUpdated = false;
 	}
 
+	if (bConfigDirty)
+	{
+		if (WeightToolProperties)
+		{
+			WeightToolProperties->SaveConfig();
+		}
+		bConfigDirty = false;
+	}
+
 	if (!VerticesToUpdateColor.IsEmpty())
 	{
 		UpdateVertexColorForSubsetOfVertices();
@@ -2581,6 +2649,12 @@ void USkinWeightsPaintTool::OnTick(float DeltaTime)
 	Weights.Deformer.UpdateVertexDeformation(this, GetComponentSpaceBoneTransforms());
 }
 
+void USkinWeightsPaintTool::OnLinearColorRampChanged(TArray<FRichCurve*> ChangedCurves)
+{
+	bVertexColorsNeedUpdated = true;
+	bConfigDirty = true;
+}
+
 void USkinWeightsPaintTool::UpdateCurrentlyEditedMesh(const FDynamicMesh3& InDynamicMesh)
 {
 	// update the preview mesh in the viewport
@@ -2967,57 +3041,10 @@ FVector4f USkinWeightsPaintTool::GetColorOfVertex(VertexIndex InVertexIndex, Bon
 				return FLinearColor::Black; // with no bone selected, all vertices are drawn black
 			}
 
-			// get user-specified colors
-			const TArray<FLinearColor>& Colors = WeightToolProperties->ColorRamp;
 			// get weight value
 			float Value = Weights.GetWeightOfBoneOnVertex(InCurrentBoneIndex, InVertexIndex, Weights.CurrentWeights);
-			Value = FMath::Clamp(Value, 0.0f, 1.0f);
-
-			// ZERO user supplied colors, then revert to greyscale
-			if (Colors.IsEmpty())
-			{
-				return FMath::Lerp(FLinearColor::Black, FLinearColor::White, Value);
-			}
-
-			// ONE user defined color, blend it with black
-			if (Colors.Num() == 1)
-			{
-				return FMath::Lerp(FLinearColor::Black, Colors[0], Value);
-			}
-
-			// TWO user defined color, simple LERP
-			if (Colors.Num() == 2)
-			{
-				return FMath::Lerp(Colors[0], Colors[1], Value);
-			}
-			
-			// blend colors between min and max value
-			constexpr float MinValue = 0.1f;
-			constexpr float MaxValue = 0.9f;
-			
-			// early out zero weights to min color
-			if (Value <= MinValue)
-			{
-				return Colors[0];
-			}
-
-			// early out full weights to max color
-			if (Value >= MaxValue)
-			{
-				return Colors.Last();
-			}
 			
-			// remap from 0-1 to range of MinValue to MaxValue
-			const float ScaledValue = (Value - MinValue) * 1.0f / (MaxValue - MinValue);
-			// interpolate within two nearest ramp colors
-			const float PerColorRange = 1.0f / (Colors.Num() - 1);
-			const int ColorIndex = static_cast<int>(ScaledValue / PerColorRange);
-			const float RangeStart = ColorIndex * PerColorRange;
-			const float RangeEnd = (ColorIndex + 1) * PerColorRange;
-			const float Param = (ScaledValue - RangeStart) / (RangeEnd - RangeStart);
-			const FLinearColor& StartColor = Colors[ColorIndex];
-			const FLinearColor& EndColor = Colors[ColorIndex+1];
-			return UE::Geometry::ToVector4<float>(FMath::Lerp(StartColor, EndColor, Param));
+			return WeightToolProperties->LinearColorRamp.GetLinearColorValue(Value);
 		}
 	case EWeightColorMode::BoneColors:
 		{
@@ -5278,34 +5305,17 @@ void USkinWeightsPaintTool::OnPropertyModified(UObject* ModifiedObject, FPropert
 		{
 			WeightToolProperties->GetBrushConfig().Strength = WeightToolProperties->BrushStrength;
 		}
-		if (ModifiedProperty->GetName() == GET_MEMBER_NAME_STRING_CHECKED(USkinWeightsPaintToolProperties, BrushRadius))
+		else if (ModifiedProperty->GetName() == GET_MEMBER_NAME_STRING_CHECKED(USkinWeightsPaintToolProperties, BrushRadius))
 		{
 			WeightToolProperties->GetBrushConfig().Radius = WeightToolProperties->BrushRadius;
 		}
-		if (ModifiedProperty->GetName() == GET_MEMBER_NAME_STRING_CHECKED(USkinWeightsPaintToolProperties, BrushFalloffAmount))
+		else if (ModifiedProperty->GetName() == GET_MEMBER_NAME_STRING_CHECKED(USkinWeightsPaintToolProperties, BrushFalloffAmount))
 		{
 			WeightToolProperties->GetBrushConfig().Falloff = WeightToolProperties->BrushFalloffAmount;
 		}
-	
-		const FString NameOfModifiedProperty = ModifiedProperty->GetNameCPP();
-
-		// invalidate vertex color cache when any weight color properties are modified
-		const TArray<FString> ColorPropertyNames = {
-			GET_MEMBER_NAME_STRING_CHECKED(USkinWeightsPaintToolProperties, ColorMode),
-			GET_MEMBER_NAME_STRING_CHECKED(USkinWeightsPaintToolProperties, ColorRamp),
-			GET_MEMBER_NAME_STRING_CHECKED(FLinearColor, R),
-			GET_MEMBER_NAME_STRING_CHECKED(FLinearColor, G),
-			GET_MEMBER_NAME_STRING_CHECKED(FLinearColor, B),
-			GET_MEMBER_NAME_STRING_CHECKED(FLinearColor, A)};
-		if (ColorPropertyNames.Contains(NameOfModifiedProperty))
+		else if (ModifiedProperty->GetName() == GET_MEMBER_NAME_STRING_CHECKED(USkinWeightsPaintToolProperties, ColorMode))
 		{
 			bVertexColorsNeedUpdated = true;
-
-			// force all colors to have Alpha = 1
-			for (FLinearColor& Color : WeightToolProperties->ColorRamp)
-			{
-				Color.A = 1.f;
-			}
 		}
 
 		// let the mesh transfer system react to properties being set
diff --git a/Engine/Plugins/Runtime/MeshModelingToolset/Source/MeshModelingToolsEditorOnly/Private/SkeletalMesh/SkinWeightsPaintToolPropertiesCustomVersion.h b/Engine/Plugins/Runtime/MeshModelingToolset/Source/MeshModelingToolsEditorOnly/Private/SkeletalMesh/SkinWeightsPaintToolPropertiesCustomVersion.h
new file mode 100644
--- /dev/null
+++ b/Engine/Plugins/Runtime/MeshModelingToolset/Source/MeshModelingToolsEditorOnly/Private/SkeletalMesh/SkinWeightsPaintToolPropertiesCustomVersion.h
@@ -0,0 +1,28 @@
+// Copyright Epic Games, Inc. All Rights Reserved.
+
+#pragma once
+
+#include "GenericPlatform/GenericPlatform.h"
+
+#include "SkinWeightsPaintToolPropertiesCustomVersion.generated.h"
+
+/** 
+ * Versioning for the skin weights paint tool properties (USkinWeightsPaintToolProperties).
+ * Mind this differs fundamentally from the common custom object version logic,
+ * it is and has to use the UInteractiveToolPropertySet::RestoreProperties system.
+ * 
+ * See USkinWeightsPaintToolProperties::RestoreProperties for the related upgrade path
+ */
+UENUM()
+enum class ESkinWeightsPaintToolPropertiesCustomVersion : uint32
+{
+	// Before versioning was added
+	NoVersion = 0x00,
+
+	// Upgrade SkinWeightsPaintToolProperties to use FLinearColorRamp instead of an array of linear colors
+	SkinWeightPaintToolUseLinearColorRamp = 0x01 << 0,
+
+	// Add new versions above 
+	VersionPlusOne,
+	LatestVersion = VersionPlusOne - 1
+};
diff --git a/Engine/Plugins/Runtime/MeshModelingToolset/Source/MeshModelingToolsEditorOnly/Public/SkeletalMesh/SkinWeightsPaintTool.h b/Engine/Plugins/Runtime/MeshModelingToolset/Source/MeshModelingToolsEditorOnly/Public/SkeletalMesh/SkinWeightsPaintTool.h
--- a/Engine/Plugins/Runtime/MeshModelingToolset/Source/MeshModelingToolsEditorOnly/Public/SkeletalMesh/SkinWeightsPaintTool.h
+++ b/Engine/Plugins/Runtime/MeshModelingToolset/Source/MeshModelingToolsEditorOnly/Public/SkeletalMesh/SkinWeightsPaintTool.h
@@ -3,6 +3,7 @@
 #pragma once
 
 #include "DynamicMeshBrushTool.h"
+#include "Curves/LinearColorRamp.h"
 #include "BaseTools/MeshSurfacePointMeshEditingTool.h"
 #include "MeshDescription.h"
 #include "DynamicMesh/DynamicVerticesOctree3.h"
@@ -39,6 +40,7 @@ class USkeletonModifier;
 class SButton;
 class SReferenceSkeletonTree;
 class FSkeletalMeshNotifierBindScope;
+enum class ESkinWeightsPaintToolPropertiesCustomVersion : uint32;
 
 namespace UE::Geometry 
 {
@@ -398,6 +400,9 @@ class USkinWeightsPaintToolProperties : public UBrushBaseProperties
 	UE_API USkinWeightsPaintToolProperties();
 	
 public:
+	//~ Begin UInteractiveToolPropertySet interface
+	virtual void RestoreProperties(UInteractiveTool* RestoreToTool, const FString& CacheIdentifier = TEXT("")) override;
+	//~ End UInteractiveToolPropertySet interface
 
 	// brush vs selection modes
 	UPROPERTY(Config)
@@ -415,8 +420,13 @@ public:
 	// weight color properties
 	UPROPERTY(EditAnywhere, Config, Category = MeshDisplay)
 	EWeightColorMode ColorMode;
-	UPROPERTY(EditAnywhere, Config, Category = MeshDisplay)
+
+	UE_DEPRECATED(5.8, "Use LinearColorRamp instead")
+	UPROPERTY(Config, meta = (DeprecatedProperty, DeprecationMessage = "Use LinearColorRamp instead."))
 	TArray<FLinearColor> ColorRamp;
+	
+	UPROPERTY(EditAnywhere, Config, Category = MeshDisplay, Meta = (EditCondition="ColorMode==EWeightColorMode::Ramp", EditConditionHides))
+	FLinearColorRamp LinearColorRamp;
 
 	// weight editing arguments
 	UPROPERTY(Config)
@@ -501,6 +511,13 @@ private:
 	UE_API TArray<FName> GetSourceLODsFunc() const;
 	UFUNCTION()
 	UE_API TArray<FName> GetSourceSkinWeightProfilesFunc() const;
+
+	/** 
+	 * Defines a fully custom version stored in the UInteractiveToolPropertySet system.
+	 * This allows to bump property versions much like common UObject allow bumping post load.
+	 */
+	UPROPERTY(Config)
+	ESkinWeightsPaintToolPropertiesCustomVersion CustomVersion;
 };
 
 // this class wraps the all the components to enable selection on a single mesh in the skin weights tool
@@ -859,6 +876,9 @@ protected:
 	UE_API void OnShutdown(EToolShutdownType ShutdownType) override;
 	UE_API void OnTick(float DeltaTime) override;
 
+	/** Called when linear color ramp curves changed */
+	void OnLinearColorRampChanged(TArray<FRichCurve*> ChangedCurves);
+
 	// stamp
 	UE_API float CalculateBrushFalloff(float Distance) const;
 	UE_API void CalculateVertexROI(
@@ -985,7 +1005,10 @@ protected:
 	// manages isolating a selection of the mesh
 	UPROPERTY()
 	TObjectPtr<UWeightToolSelectionIsolator> SelectionIsolator = nullptr;
-	
+
+	/** True while the config changed but wasn't saved yet */
+	bool bConfigDirty = false;
+
 	/** The skeleton tree displayed when using this tool */
 	TSharedPtr<SReferenceSkeletonTree> ToolSkeletonTree;
 
