[SkeletalMeshModelingTools][SkeletonEditingTool] Add geometry isolation support to Skeleton Editing Tool

#Synced-CL 52740731

diff --git a/Engine/Plugins/Runtime/MeshModelingToolset/Source/MeshModelingToolsEditorOnly/Private/SkeletalMesh/SkeletonEditingTool.cpp b/Engine/Plugins/Runtime/MeshModelingToolset/Source/MeshModelingToolsEditorOnly/Private/SkeletalMesh/SkeletonEditingTool.cpp
--- a/Engine/Plugins/Runtime/MeshModelingToolset/Source/MeshModelingToolsEditorOnly/Private/SkeletalMesh/SkeletonEditingTool.cpp
+++ b/Engine/Plugins/Runtime/MeshModelingToolset/Source/MeshModelingToolsEditorOnly/Private/SkeletalMesh/SkeletonEditingTool.cpp
@@ -20,6 +20,7 @@
 #include "DynamicMesh/MeshNormals.h"
 #include "Framework/MultiBox/MultiBoxBuilder.h"
 #include "Preferences/PersonaOptions.h"
+#include "DynamicSubmesh3.h"
 #include "Selection/PolygonSelectionMechanic.h"
 
 #include "SkeletalMesh/SkeletonTransformProxy.h"
@@ -198,7 +199,28 @@ void USkeletonEditingTool::SetupPreviewMesh()
 
 	ToolSetupUtil::ApplyRenderingConfigurationToPreview(PreviewMesh, Target);
 	PreviewMesh->SetTangentsMode(EDynamicMeshComponentTangentsMode::AutoCalculated);
-	PreviewMesh->ReplaceMesh(UE::ToolTarget::GetDynamicMeshCopy(Target));
+
+	// Use isolated submesh if isolation was active, otherwise use the full mesh
+	{
+		FDynamicMesh3 FullMesh = UE::ToolTarget::GetDynamicMeshCopy(Target);
+		bool bShouldUseFullMesh = true;
+		if (EditorContext.IsValid())
+		{
+			const TArray<int32>& IsolatedTriangles = EditorContext->GetIsolatedTriangles();
+			if (!IsolatedTriangles.IsEmpty())
+			{
+				UE::Geometry::FDynamicSubmesh3 Submesh(&FullMesh, IsolatedTriangles);
+				PreviewMesh->ReplaceMesh(Submesh.GetSubmesh());
+				bShouldUseFullMesh = false;
+			}
+		}
+		
+		if (bShouldUseFullMesh)
+		{
+			PreviewMesh->ReplaceMesh(MoveTemp(FullMesh));
+		}	
+	}
+
 
 	const FComponentMaterialSet MaterialSet = UE::ToolTarget::GetMaterialSet(Target);
 	PreviewMesh->SetMaterials(MaterialSet.Materials);
